#summary Can search for nodes in the palette menu in Blueprints with both the native name and the localized name.

#add Added UEdGraphNode::GetNodeNativeTitle to return a native title for a node.
#add Added UEdGraphNode::GetNodeSearchTitle to return the native and localized title for a node, together, for searching.
#add Can hold "alt" over a node (in the graph panel, or the palette) to see the native name of the node.

#ttp 331252 - Blueprints: Editor: L10N: Blueprints need to consistently show localized node names and when searching need to search both the localized name and the native name

diff --git a/Engine/Plugins/Editor/IntroTutorials/Source/IntroTutorials/Classes/K2Node_TutorialExcerptComplete.h b/Engine/Plugins/Editor/IntroTutorials/Source/IntroTutorials/Classes/K2Node_TutorialExcerptComplete.h
--- a/Engine/Plugins/Editor/IntroTutorials/Source/IntroTutorials/Classes/K2Node_TutorialExcerptComplete.h
+++ b/Engine/Plugins/Editor/IntroTutorials/Source/IntroTutorials/Classes/K2Node_TutorialExcerptComplete.h
@@ -28,7 +28,8 @@ class UK2Node_TutorialExcerptComplete : public UK2Node_CallFunction
 	// End UK2Node_CallFunction interface
 
 	// Begin UEdGraphNode interface
-	virtual FString GetNodeTitle(ENodeTitleType::Type TitleType) const OVERRIDE;
+	virtual FText GetNodeTitle(ENodeTitleType::Type TitleType) const OVERRIDE;
+	virtual FString GetNodeNativeTitle(ENodeTitleType::Type TitleType) const OVERRIDE;
 	virtual void PostPlacedNewNode() OVERRIDE;
 	virtual void ValidateNodeDuringCompilation(class FCompilerResultsLog& MessageLog) const OVERRIDE;
 	// End UEdGraphNode interface
diff --git a/Engine/Plugins/Editor/IntroTutorials/Source/IntroTutorials/Private/K2Node_TutorialExcerptComplete.cpp b/Engine/Plugins/Editor/IntroTutorials/Source/IntroTutorials/Private/K2Node_TutorialExcerptComplete.cpp
--- a/Engine/Plugins/Editor/IntroTutorials/Source/IntroTutorials/Private/K2Node_TutorialExcerptComplete.cpp
+++ b/Engine/Plugins/Editor/IntroTutorials/Source/IntroTutorials/Private/K2Node_TutorialExcerptComplete.cpp
@@ -41,11 +41,12 @@ void UK2Node_TutorialExcerptComplete::GetMenuEntries(FGraphContextMenuBuilder& C
 
 	const FString Category = LOCTEXT("TutorialExcerptComplete_Category", "Call Function").ToString();
 	const FString SubCategory = LOCTEXT("TutorialExcerptComplete_SubCategory", "Tutorials").ToString();
-	const FString MenuDesc = LOCTEXT("TutorialExcerptComplete_Desc", "Tutorial Excerpt Complete").ToString();
+	const FText MenuDesc = LOCTEXT("TutorialExcerptComplete_Desc", "Tutorial Excerpt Complete");
 	const FString Tooltip = LOCTEXT("TutorialExcerptComplete_Tooltip", "Add a node to complete a tutorial excerpt").ToString();
 
 	TSharedPtr<FEdGraphSchemaAction_K2NewNode> NodeAction = FK2ActionMenuBuilder::AddNewNodeAction(ContextMenuBuilder, Category + TEXT("|") + SubCategory, MenuDesc, Tooltip);
 	NodeAction->NodeTemplate = TemplateNode;
+	NodeAction->SearchTitle = TemplateNode->GetNodeSearchTitle();
 }
 
 void UK2Node_TutorialExcerptComplete::PinDefaultValueChanged(UEdGraphPin* Pin)
@@ -91,17 +92,29 @@ void UK2Node_TutorialExcerptComplete::PostPlacedNewNode()
 	Super::PostPlacedNewNode();
 }
 
-FString UK2Node_TutorialExcerptComplete::GetNodeTitle(ENodeTitleType::Type TitleType) const
+FText UK2Node_TutorialExcerptComplete::GetNodeTitle(ENodeTitleType::Type TitleType) const
 {
 	if(CachedExcerpt.Len() > 0)
 	{
 		FFormatNamedArguments NamedArgs;
 		NamedArgs.Add(TEXT("ExcerptName"), FText::FromString(CachedExcerpt));
-		return FText::Format(LOCTEXT("TutorialExcerptComplete_TitleWithExcerpt", "Tutorial Excerpt '{ExcerptName}' Complete"), NamedArgs).ToString();
+		return FText::Format(LOCTEXT("TutorialExcerptComplete_TitleWithExcerpt", "Tutorial Excerpt '{ExcerptName}' Complete"), NamedArgs);
 	}
 	else
 	{
-		return LOCTEXT("TutorialExcerptComplete_Title", "Tutorial Excerpt Complete").ToString();
+		return LOCTEXT("TutorialExcerptComplete_Title", "Tutorial Excerpt Complete");
+	}
+}
+
+FString UK2Node_TutorialExcerptComplete::GetNodeNativeTitle(ENodeTitleType::Type TitleType) const
+{
+	if(CachedExcerpt.Len() > 0)
+	{
+		return FString::Printf(TEXT("Tutorial Excerpt '%s' Complete"), *CachedExcerpt);
+	}
+	else
+	{
+		return TEXT("Tutorial Excerpt Complete");
 	}
 }
 
diff --git a/Engine/Plugins/Experimental/LiveEditor/Source/LiveEditor/Classes/K2Node_LiveEditObject.h b/Engine/Plugins/Experimental/LiveEditor/Source/LiveEditor/Classes/K2Node_LiveEditObject.h
--- a/Engine/Plugins/Experimental/LiveEditor/Source/LiveEditor/Classes/K2Node_LiveEditObject.h
+++ b/Engine/Plugins/Experimental/LiveEditor/Source/LiveEditor/Classes/K2Node_LiveEditObject.h
@@ -11,7 +11,8 @@ class UK2Node_LiveEditObject : public UK2Node
 #if WITH_EDITOR
 	// Begin UEdGraphNode interface.
 	virtual void AllocateDefaultPins() OVERRIDE;
-	virtual FString GetNodeTitle(ENodeTitleType::Type TitleType) const OVERRIDE;
+	virtual FText GetNodeTitle(ENodeTitleType::Type TitleType) const OVERRIDE;
+	virtual FString GetNodeNativeTitle(ENodeTitleType::Type TitleType) const OVERRIDE;
 	virtual void PinDefaultValueChanged(UEdGraphPin* Pin) OVERRIDE;
 	// End UEdGraphNode interface.
 
diff --git a/Engine/Plugins/Experimental/LiveEditor/Source/LiveEditor/Private/K2Node_LiveEditObject.cpp b/Engine/Plugins/Experimental/LiveEditor/Source/LiveEditor/Private/K2Node_LiveEditObject.cpp
--- a/Engine/Plugins/Experimental/LiveEditor/Source/LiveEditor/Private/K2Node_LiveEditObject.cpp
+++ b/Engine/Plugins/Experimental/LiveEditor/Source/LiveEditor/Private/K2Node_LiveEditObject.cpp
@@ -177,17 +177,33 @@ void UK2Node_LiveEditObject::AllocateDefaultPins()
 	Super::AllocateDefaultPins();
 }
 
-FString UK2Node_LiveEditObject::GetNodeTitle(ENodeTitleType::Type TitleType) const
+FText UK2Node_LiveEditObject::GetNodeTitle(ENodeTitleType::Type TitleType) const
 {
 	UEdGraphPin* BaseClassPin = GetBaseClassPin();
 
-	FString SpawnString = NSLOCTEXT("K2Node", "None", "NONE").ToString();
+	FText SpawnString = NSLOCTEXT("K2Node", "None", "NONE");
+	if(BaseClassPin != NULL && BaseClassPin->DefaultObject != NULL )
+	{
+		SpawnString = FText::FromString(BaseClassPin->DefaultObject->GetName());
+	}
+
+	FFormatNamedArguments Args;
+	Args.Add(TEXT("SpawnString"), SpawnString);
+	Args.Add(TEXT("ID"), GetUniqueID());
+	return FText::Format(NSLOCTEXT("K2Node", "LiveEditObject", "LiveEditObject {SpawnString}_{ID}"), Args );
+}
+
+FString UK2Node_LiveEditObject::GetNodeNativeTitle(ENodeTitleType::Type TitleType) const
+{
+	UEdGraphPin* BaseClassPin = GetBaseClassPin();
+
+	FString SpawnString = TEXT("NONE");
 	if(BaseClassPin != NULL && BaseClassPin->DefaultObject != NULL )
 	{
 		SpawnString = BaseClassPin->DefaultObject->GetName();
 	}
 
-	return FString::Printf(*NSLOCTEXT("K2Node", "LiveEditObject", "LiveEditObject %s_%d").ToString(), *SpawnString, GetUniqueID());
+	return FString::Printf(TEXT("LiveEditObject %s_%d"), *SpawnString, GetUniqueID());
 }
 
 void UK2Node_LiveEditObject::PinDefaultValueChanged(UEdGraphPin* Pin) 
diff --git a/Engine/Plugins/Experimental/LiveEditor/Source/LiveEditor/Private/LiveEditor.cpp b/Engine/Plugins/Experimental/LiveEditor/Source/LiveEditor/Private/LiveEditor.cpp
--- a/Engine/Plugins/Experimental/LiveEditor/Source/LiveEditor/Private/LiveEditor.cpp
+++ b/Engine/Plugins/Experimental/LiveEditor/Source/LiveEditor/Private/LiveEditor.cpp
@@ -100,8 +100,9 @@ void FLiveEditorNodeInjector::AddContextMenuNodes(FBlueprintGraphActionListBuild
 {
 	UK2Node* NodeTemplate = ContextMenuBuilder.CreateTemplateNode<UK2Node_LiveEditObject>();
 	FString FunctionCategory( TEXT("LiveEditor") );
-	TSharedPtr<FEdGraphSchemaAction_K2NewNode> Action = FK2ActionMenuBuilder::AddNewNodeAction(ContextMenuBuilder, FunctionCategory, TEXT("LiveEdit Object"), NodeTemplate->GetTooltip(), 0, NodeTemplate->GetKeywords());
+	TSharedPtr<FEdGraphSchemaAction_K2NewNode> Action = FK2ActionMenuBuilder::AddNewNodeAction(ContextMenuBuilder, FunctionCategory, NSLOCTEXT("LiveEditorPlugin", "LiveEditObject", "LiveEdit Object"), NodeTemplate->GetTooltip(), 0, NodeTemplate->GetKeywords());
 	Action->NodeTemplate = NodeTemplate;
+	Action->SearchTitle = NodeTemplate->GetNodeSearchTitle();
 }
 
 void FLiveEditorNodeInjector::OnBeginPIE(const bool bIsSimulating)
diff --git a/Engine/Plugins/Experimental/MathExpressions/Source/MathExpressions/Private/K2Node_MathExpression.cpp b/Engine/Plugins/Experimental/MathExpressions/Source/MathExpressions/Private/K2Node_MathExpression.cpp
--- a/Engine/Plugins/Experimental/MathExpressions/Source/MathExpressions/Private/K2Node_MathExpression.cpp
+++ b/Engine/Plugins/Experimental/MathExpressions/Source/MathExpressions/Private/K2Node_MathExpression.cpp
@@ -2252,11 +2252,12 @@ void UK2Node_MathExpression::GetMenuEntries(FGraphContextMenuBuilder& ContextMen
 	UK2Node_MathExpression* TemplateNode = NewObject<UK2Node_MathExpression>(GetTransientPackage(), GetClass());
 
 	const FString Category = TEXT("");
-	const FString MenuDesc = TEXT("Add Math Expression...");
+	const FText MenuDesc = FText::FromString(TEXT("Add Math Expression..."));
 	const FString Tooltip = TEXT("Create a new mathematical expression");
 
 	TSharedPtr<FEdGraphSchemaAction_K2NewNode> NodeAction = FK2ActionMenuBuilder::AddNewNodeAction(ContextMenuBuilder, Category, MenuDesc, Tooltip);
 	NodeAction->NodeTemplate = TemplateNode;
+	NodeAction->SearchTitle = TemplateNode->GetNodeSearchTitle();
 }
 
 TSharedPtr<class INameValidatorInterface> UK2Node_MathExpression::MakeNameValidator() const
@@ -2289,14 +2290,28 @@ void UK2Node_MathExpression::RebuildExpression()
 	}
 }
 
-FString UK2Node_MathExpression::GetNodeTitle(ENodeTitleType::Type TitleType) const
+FText UK2Node_MathExpression::GetNodeTitle(ENodeTitleType::Type TitleType) const
+{
+	FString Result = Expression;
+
+	if (TitleType == ENodeTitleType::FullTitle)
+	{
+		FFormatNamedArguments Args;
+		Args.Add(TEXT("Expression"), FText::FromString(Result));
+		Result = NSLOCTEXT("K2Node", "MathExpressionSecondTitleLine", "{Expression}\nMath Expression").ToString();
+	}
+	
+	return FText::FromString(Result);
+}
+
+FString UK2Node_MathExpression::GetNodeNativeTitle(ENodeTitleType::Type TitleType) const
 {
 	FString Result = Expression;
 
 	if (TitleType == ENodeTitleType::FullTitle)
 	{
 		Result += TEXT("\n");
-		Result += NSLOCTEXT("K2Node", "MathExpressionSecondTitleLine", "Math Expression").ToString();
+		Result += TEXT("Math Expression");
 	}
 
 	return Result;
diff --git a/Engine/Source/Developer/MovieSceneTools/Classes/K2Node_PlayMovieScene.h b/Engine/Source/Developer/MovieSceneTools/Classes/K2Node_PlayMovieScene.h
--- a/Engine/Source/Developer/MovieSceneTools/Classes/K2Node_PlayMovieScene.h
+++ b/Engine/Source/Developer/MovieSceneTools/Classes/K2Node_PlayMovieScene.h
@@ -20,7 +20,7 @@ struct FEdGraphSchemaAction_K2AddPlayMovieScene : public FEdGraphSchemaAction_K2
 	{
 	}
 	
-	FEdGraphSchemaAction_K2AddPlayMovieScene(const FString& InNodeCategory, const FString& InMenuDesc, const FString& InToolTip, const int32 InGrouping)
+	FEdGraphSchemaAction_K2AddPlayMovieScene(const FString& InNodeCategory, const FText& InMenuDesc, const FString& InToolTip, const int32 InGrouping)
 		: FEdGraphSchemaAction_K2NewNode(InNodeCategory, InMenuDesc, InToolTip, InGrouping)
 	{
 	}
@@ -41,7 +41,8 @@ class UK2Node_PlayMovieScene : public UK2Node
 	// UEdGraphNode interface.
 	virtual void AllocateDefaultPins() OVERRIDE;
 	virtual FLinearColor GetNodeTitleColor() const OVERRIDE;
-	virtual FString GetNodeTitle(ENodeTitleType::Type TitleType) const OVERRIDE;
+	virtual FText GetNodeTitle(ENodeTitleType::Type TitleType) const OVERRIDE;
+	virtual FString GetNodeNativeTitle(ENodeTitleType::Type TitleType) const OVERRIDE;
 
 	// UK2Node interface
 	virtual bool NodeCausesStructuralBlueprintChange() const OVERRIDE
diff --git a/Engine/Source/Developer/MovieSceneTools/Private/K2Node_PlayMovieScene.cpp b/Engine/Source/Developer/MovieSceneTools/Private/K2Node_PlayMovieScene.cpp
--- a/Engine/Source/Developer/MovieSceneTools/Private/K2Node_PlayMovieScene.cpp
+++ b/Engine/Source/Developer/MovieSceneTools/Private/K2Node_PlayMovieScene.cpp
@@ -76,12 +76,24 @@ FLinearColor UK2Node_PlayMovieScene::GetNodeTitleColor() const
 	return FLinearColor(0.8f, 0.05f, 0.05f);
 }
 
-FString UK2Node_PlayMovieScene::GetNodeTitle(ENodeTitleType::Type TitleType) const
+FText UK2Node_PlayMovieScene::GetNodeTitle(ENodeTitleType::Type TitleType) const
+{
+	UMovieScene* MovieScene = MovieSceneBindings != NULL ? MovieSceneBindings->GetRootMovieScene() : NULL;
+
+	FFormatNamedArguments Args;
+	Args.Add(TEXT("SceneName"), FText::FromString(MovieScene->GetName()));
+
+	return MovieScene != NULL ?
+		( FText::Format(NSLOCTEXT("PlayMovieSceneNode", "NodeTitle", "Play Movie Scene: {SceneName}"), Args) ) :
+		NSLOCTEXT("PlayMovieSceneNode", "NodeTitleWithNoMovieScene", "Play Movie Scene (No Asset)");
+}
+
+FString UK2Node_PlayMovieScene::GetNodeNativeTitle(ENodeTitleType::Type TitleType) const
 {
 	UMovieScene* MovieScene = MovieSceneBindings != NULL ? MovieSceneBindings->GetRootMovieScene() : NULL;
 	return MovieScene != NULL ?
-		( NSLOCTEXT("PlayMovieSceneNode", "NodeTitle", "Play Movie Scene: ").ToString() + MovieScene->GetName() ) :
-		NSLOCTEXT("PlayMovieSceneNode", "NodeTitleWithNoMovieScene", "Play Movie Scene (No Asset)").ToString();
+		( FString(TEXT("Play Movie Scene: ")) + MovieScene->GetName() ) 
+		: TEXT("Play Movie Scene (No Asset)");
 }
 
 
@@ -262,7 +274,7 @@ void UK2Node_PlayMovieScene::CreatePinForBoundObject( FMovieSceneBoundObject& Bo
 	const FString PinName = BoundObject.GetPossessableGuid().ToString( EGuidFormats::DigitsWithHyphens );
 
 	// For the friendly name, we use the possessable name from the MovieScene asset that is associated with this node
-	FString PinFriendlyName = PinName;
+	FText PinFriendlyName = FText::FromString(PinName);
 	{
 		UMovieScene* MovieScene = GetMovieScene();
 		if( MovieScene != NULL )		// @todo sequencer: Need to refresh the PinFriendlyName if the MovieScene asset changes, or if the possessable slot is renamed within Sequencer
@@ -273,7 +285,7 @@ void UK2Node_PlayMovieScene::CreatePinForBoundObject( FMovieSceneBoundObject& Bo
 				if( Possessable.GetGuid() == BoundObject.GetPossessableGuid() )
 				{
 					// Found a name for this possessable
-					PinFriendlyName = Possessable.GetName();
+					PinFriendlyName = FText::FromString(Possessable.GetName());
 					break;
 				}
 			}
diff --git a/Engine/Source/Editor/AnimGraph/Classes/AnimGraphNode_ApplyAdditive.h b/Engine/Source/Editor/AnimGraph/Classes/AnimGraphNode_ApplyAdditive.h
--- a/Engine/Source/Editor/AnimGraph/Classes/AnimGraphNode_ApplyAdditive.h
+++ b/Engine/Source/Editor/AnimGraph/Classes/AnimGraphNode_ApplyAdditive.h
@@ -13,7 +13,7 @@ class UAnimGraphNode_ApplyAdditive : public UAnimGraphNode_Base
 
 	// UEdGraphNode interface
 	virtual FString GetTooltip() const OVERRIDE;
-	virtual FString GetNodeTitle(ENodeTitleType::Type TitleType) const OVERRIDE;
+	virtual FText GetNodeTitle(ENodeTitleType::Type TitleType) const OVERRIDE;
 	virtual FLinearColor GetNodeTitleColor() const OVERRIDE;
 	// End of UEdGraphNode interface
 
diff --git a/Engine/Source/Editor/AnimGraph/Classes/AnimGraphNode_BlendListByBool.h b/Engine/Source/Editor/AnimGraph/Classes/AnimGraphNode_BlendListByBool.h
--- a/Engine/Source/Editor/AnimGraph/Classes/AnimGraphNode_BlendListByBool.h
+++ b/Engine/Source/Editor/AnimGraph/Classes/AnimGraphNode_BlendListByBool.h
@@ -13,7 +13,7 @@ class UAnimGraphNode_BlendListByBool : public UAnimGraphNode_BlendListBase
 
 	// UEdGraphNode interface
 	virtual FString GetTooltip() const OVERRIDE;
-	virtual FString GetNodeTitle(ENodeTitleType::Type TitleType) const OVERRIDE;
+	virtual FText GetNodeTitle(ENodeTitleType::Type TitleType) const OVERRIDE;
 	// End of UEdGraphNode interface
 
 	// UAnimGraphNode_Base interface
diff --git a/Engine/Source/Editor/AnimGraph/Classes/AnimGraphNode_CopyBone.h b/Engine/Source/Editor/AnimGraph/Classes/AnimGraphNode_CopyBone.h
--- a/Engine/Source/Editor/AnimGraph/Classes/AnimGraphNode_CopyBone.h
+++ b/Engine/Source/Editor/AnimGraph/Classes/AnimGraphNode_CopyBone.h
@@ -13,12 +13,13 @@ class UAnimGraphNode_CopyBone : public UAnimGraphNode_SkeletalControlBase
 
 public:
 	// UEdGraphNode interface
-	virtual FString GetNodeTitle(ENodeTitleType::Type TitleType) const OVERRIDE;
+	virtual FText GetNodeTitle(ENodeTitleType::Type TitleType) const OVERRIDE;
+	virtual FString GetNodeNativeTitle(ENodeTitleType::Type TitleType) const OVERRIDE;
 	virtual FString GetTooltip() const OVERRIDE;
 	// End of UEdGraphNode interface
 
 protected:
 	// UAnimGraphNode_SkeletalControlBase interface
-	virtual FString GetControllerDescription() const;
+	virtual FText GetControllerDescription() const OVERRIDE;
 	// End of UAnimGraphNode_SkeletalControlBase interface
 };
diff --git a/Engine/Source/Editor/AnimGraph/Classes/AnimGraphNode_Fabrik.h b/Engine/Source/Editor/AnimGraph/Classes/AnimGraphNode_Fabrik.h
--- a/Engine/Source/Editor/AnimGraph/Classes/AnimGraphNode_Fabrik.h
+++ b/Engine/Source/Editor/AnimGraph/Classes/AnimGraphNode_Fabrik.h
@@ -15,11 +15,11 @@ class UAnimGraphNode_Fabrik : public UAnimGraphNode_SkeletalControlBase
 
 public:
 	// UEdGraphNode interface
-	virtual FString GetNodeTitle(ENodeTitleType::Type TitleType) const OVERRIDE;
+	virtual FText GetNodeTitle(ENodeTitleType::Type TitleType) const OVERRIDE;
 	// End of UEdGraphNode interface
 
 protected:
 	// UAnimGraphNode_SkeletalControlBase interface
-	virtual FString GetControllerDescription() const;
+	virtual FText GetControllerDescription() const OVERRIDE;
 	// End of UAnimGraphNode_SkeletalControlBase interface
 };
\ No newline at end of file
diff --git a/Engine/Source/Editor/AnimGraph/Classes/AnimGraphNode_ModifyBone.h b/Engine/Source/Editor/AnimGraph/Classes/AnimGraphNode_ModifyBone.h
--- a/Engine/Source/Editor/AnimGraph/Classes/AnimGraphNode_ModifyBone.h
+++ b/Engine/Source/Editor/AnimGraph/Classes/AnimGraphNode_ModifyBone.h
@@ -15,13 +15,13 @@ class UAnimGraphNode_ModifyBone : public UAnimGraphNode_SkeletalControlBase
 
 public:
 	// UEdGraphNode interface
-	virtual FString GetNodeTitle(ENodeTitleType::Type TitleType) const OVERRIDE;
+	virtual FText GetNodeTitle(ENodeTitleType::Type TitleType) const OVERRIDE;
 	virtual FString GetTooltip() const OVERRIDE;
 	virtual FString GetKeywords() const OVERRIDE;
 	// End of UEdGraphNode interface
 
 protected:
 	// UAnimGraphNode_SkeletalControlBase interface
-	virtual FString GetControllerDescription() const;
+	virtual FText GetControllerDescription() const OVERRIDE;
 	// End of UAnimGraphNode_SkeletalControlBase interface
 };
diff --git a/Engine/Source/Editor/AnimGraph/Classes/AnimGraphNode_RotationMultiplier.h b/Engine/Source/Editor/AnimGraph/Classes/AnimGraphNode_RotationMultiplier.h
--- a/Engine/Source/Editor/AnimGraph/Classes/AnimGraphNode_RotationMultiplier.h
+++ b/Engine/Source/Editor/AnimGraph/Classes/AnimGraphNode_RotationMultiplier.h
@@ -13,12 +13,13 @@ class UAnimGraphNode_RotationMultiplier : public UAnimGraphNode_SkeletalControlB
 
 public:
 	// UEdGraphNode interface
-	virtual FString GetNodeTitle(ENodeTitleType::Type TitleType) const OVERRIDE;
+	virtual FText GetNodeTitle(ENodeTitleType::Type TitleType) const OVERRIDE;
+	virtual FString GetNodeNativeTitle(ENodeTitleType::Type TitleType) const OVERRIDE;
 	virtual FString GetTooltip() const OVERRIDE;
 	// End of UEdGraphNode interface
 
 protected:
 	// UAnimGraphNode_SkeletalControlBase interface
-	virtual FString GetControllerDescription() const;
+	virtual FText GetControllerDescription() const OVERRIDE;
 	// End of UAnimGraphNode_SkeletalControlBase interface
 };
diff --git a/Engine/Source/Editor/AnimGraph/Classes/AnimGraphNode_SequenceEvaluator.h b/Engine/Source/Editor/AnimGraph/Classes/AnimGraphNode_SequenceEvaluator.h
--- a/Engine/Source/Editor/AnimGraph/Classes/AnimGraphNode_SequenceEvaluator.h
+++ b/Engine/Source/Editor/AnimGraph/Classes/AnimGraphNode_SequenceEvaluator.h
@@ -12,8 +12,9 @@ class UAnimGraphNode_SequenceEvaluator : public UAnimGraphNode_Base
 	FAnimNode_SequenceEvaluator Node;
 
 	// UEdGraphNode interface
-	FString GetNodeTitle(ENodeTitleType::Type TitleType) const OVERRIDE;
-	FString GetTooltip() const OVERRIDE;
+	virtual FText GetNodeTitle(ENodeTitleType::Type TitleType) const OVERRIDE;
+	virtual FString GetNodeNativeTitle(ENodeTitleType::Type TitleType) const OVERRIDE;
+	virtual FString GetTooltip() const OVERRIDE;
 	// End of UEdGraphNode
 
 	// UAnimGraphNode_Base interface
diff --git a/Engine/Source/Editor/AnimGraph/Classes/AnimGraphNode_SequencePlayer.h b/Engine/Source/Editor/AnimGraph/Classes/AnimGraphNode_SequencePlayer.h
--- a/Engine/Source/Editor/AnimGraph/Classes/AnimGraphNode_SequencePlayer.h
+++ b/Engine/Source/Editor/AnimGraph/Classes/AnimGraphNode_SequencePlayer.h
@@ -18,7 +18,8 @@ class UAnimGraphNode_SequencePlayer : public UAnimGraphNode_Base
 	// UEdGraphNode interface
 	virtual FLinearColor GetNodeTitleColor() const OVERRIDE;
 	virtual FString GetTooltip() const OVERRIDE;
-	virtual FString GetNodeTitle(ENodeTitleType::Type TitleType) const OVERRIDE;
+	virtual FText GetNodeTitle(ENodeTitleType::Type TitleType) const OVERRIDE;
+	virtual FString GetNodeNativeTitle(ENodeTitleType::Type TitleType) const OVERRIDE;
 	// End of UEdGraphNode interface
 
 	// UAnimGraphNode_Base interface
@@ -39,5 +40,5 @@ class UAnimGraphNode_SequencePlayer : public UAnimGraphNode_Base
 	// End of UK2Node interface
 
 private:
-	static FString GetTitleGivenAssetInfo(const FString& AssetName, bool bKnownToBeAdditive);
+	static FText GetTitleGivenAssetInfo(const FText& AssetName, bool bKnownToBeAdditive);
 };
diff --git a/Engine/Source/Editor/AnimGraph/Classes/AnimGraphNode_SkeletalControlBase.h b/Engine/Source/Editor/AnimGraph/Classes/AnimGraphNode_SkeletalControlBase.h
--- a/Engine/Source/Editor/AnimGraph/Classes/AnimGraphNode_SkeletalControlBase.h
+++ b/Engine/Source/Editor/AnimGraph/Classes/AnimGraphNode_SkeletalControlBase.h
@@ -26,5 +26,5 @@ public:
 	ANIMGRAPH_API virtual void Draw( FPrimitiveDrawInterface* PDI, USkeletalMeshComponent * PreviewSkelMeshComp ) const {};
 protected:
 	// Returns the short descriptive name of the controller
-	virtual FString GetControllerDescription() const;
+	virtual FText GetControllerDescription() const;
 };
diff --git a/Engine/Source/Editor/AnimGraph/Classes/AnimationStateMachineSchema.h b/Engine/Source/Editor/AnimGraph/Classes/AnimationStateMachineSchema.h
--- a/Engine/Source/Editor/AnimGraph/Classes/AnimationStateMachineSchema.h
+++ b/Engine/Source/Editor/AnimGraph/Classes/AnimationStateMachineSchema.h
@@ -17,7 +17,7 @@ struct ANIMGRAPH_API FEdGraphSchemaAction_NewStateNode : public FEdGraphSchemaAc
 		, NodeTemplate(NULL)
 	{}
 
-	FEdGraphSchemaAction_NewStateNode(const FString& InNodeCategory, const FString& InMenuDesc, const FString& InToolTip, const int32 InGrouping)
+	FEdGraphSchemaAction_NewStateNode(const FString& InNodeCategory, const FText& InMenuDesc, const FString& InToolTip, const int32 InGrouping)
 		: FEdGraphSchemaAction(InNodeCategory, InMenuDesc, InToolTip, InGrouping)
 		, NodeTemplate(NULL)
 	{}
@@ -45,7 +45,7 @@ struct ANIMGRAPH_API FEdGraphSchemaAction_NewStateComment : public FEdGraphSchem
 		: FEdGraphSchemaAction()
 	{}
 
-	FEdGraphSchemaAction_NewStateComment(const FString& InNodeCategory, const FString& InMenuDesc, const FString& InToolTip, const int32 InGrouping)
+	FEdGraphSchemaAction_NewStateComment(const FString& InNodeCategory, const FText& InMenuDesc, const FString& InToolTip, const int32 InGrouping)
 		: FEdGraphSchemaAction(InNodeCategory, InMenuDesc, InToolTip, InGrouping)
 	{}
 
diff --git a/Engine/Source/Editor/AnimGraph/Classes/K2Node_TransitionRuleGetter.h b/Engine/Source/Editor/AnimGraph/Classes/K2Node_TransitionRuleGetter.h
--- a/Engine/Source/Editor/AnimGraph/Classes/K2Node_TransitionRuleGetter.h
+++ b/Engine/Source/Editor/AnimGraph/Classes/K2Node_TransitionRuleGetter.h
@@ -38,7 +38,8 @@ class UK2Node_TransitionRuleGetter : public UK2Node
 
 	// UEdGraphNode interface
 	virtual void AllocateDefaultPins() OVERRIDE;
-	virtual FString GetNodeTitle(ENodeTitleType::Type TitleType) const OVERRIDE;
+	virtual FText GetNodeTitle(ENodeTitleType::Type TitleType) const OVERRIDE;
+	virtual FString GetNodeNativeTitle(ENodeTitleType::Type TitleType) const OVERRIDE;
 	virtual FString GetTooltip() const OVERRIDE;
 	virtual bool ShowPaletteIconOnNode() const OVERRIDE{ return false; }
 	// End of UEdGraphNode interface
@@ -47,6 +48,6 @@ class UK2Node_TransitionRuleGetter : public UK2Node
 	ANIMGRAPH_API UEdGraphPin* GetOutputPin() const;
 
 	// @todo document
-	ANIMGRAPH_API static FString GetFriendlyName(ETransitionGetter::Type TypeID);
+	ANIMGRAPH_API static FText GetFriendlyName(ETransitionGetter::Type TypeID);
 };
 
diff --git a/Engine/Source/Editor/AnimGraph/Private/AnimGraphNode_ApplyAdditive.cpp b/Engine/Source/Editor/AnimGraph/Private/AnimGraphNode_ApplyAdditive.cpp
--- a/Engine/Source/Editor/AnimGraph/Private/AnimGraphNode_ApplyAdditive.cpp
+++ b/Engine/Source/Editor/AnimGraph/Private/AnimGraphNode_ApplyAdditive.cpp
@@ -8,6 +8,8 @@
 /////////////////////////////////////////////////////
 // UAnimGraphNode_ApplyAdditive
 
+#define LOCTEXT_NAMESPACE "A3Nodes"
+
 UAnimGraphNode_ApplyAdditive::UAnimGraphNode_ApplyAdditive(const FPostConstructInitializeProperties& PCIP)
 	: Super(PCIP)
 {
@@ -20,12 +22,12 @@ FLinearColor UAnimGraphNode_ApplyAdditive::GetNodeTitleColor() const
 
 FString UAnimGraphNode_ApplyAdditive::GetTooltip() const
 {
-	return TEXT("Apply additive animation to normal pose");
+	return LOCTEXT("AnimGraphNode_ApplyAdditive_Tooltip", "Apply additive animation to normal pose").ToString();
 }
 
-FString UAnimGraphNode_ApplyAdditive::GetNodeTitle(ENodeTitleType::Type TitleType) const
+FText UAnimGraphNode_ApplyAdditive::GetNodeTitle(ENodeTitleType::Type TitleType) const
 {
-	return TEXT("Apply Additive");
+	return LOCTEXT("AnimGraphNode_ApplyAdditive_Title", "Apply Additive");
 }
 
 FString UAnimGraphNode_ApplyAdditive::GetNodeCategory() const
@@ -33,3 +35,5 @@ FString UAnimGraphNode_ApplyAdditive::GetNodeCategory() const
 	return TEXT("Blends");
 	//@TODO: TEXT("Apply additive to normal pose"), TEXT("Apply additive pose"));
 }
+
+#undef LOCTEXT_NAMESPACE
diff --git a/Engine/Source/Editor/AnimGraph/Private/AnimGraphNode_Base.cpp b/Engine/Source/Editor/AnimGraph/Private/AnimGraphNode_Base.cpp
--- a/Engine/Source/Editor/AnimGraph/Private/AnimGraphNode_Base.cpp
+++ b/Engine/Source/Editor/AnimGraph/Private/AnimGraphNode_Base.cpp
@@ -220,12 +220,13 @@ TSharedPtr<FEdGraphSchemaAction_K2NewNode> UAnimGraphNode_Base::CreateDefaultMen
 	UAnimGraphNode_Base* TemplateNode = NewObject<UAnimGraphNode_Base>(GetTransientPackage(), GetClass());
 
 	FString Category = TemplateNode->GetNodeCategory();
-	FString MenuDesc = TemplateNode->GetNodeTitle(ENodeTitleType::ListView);
+	FText MenuDesc = TemplateNode->GetNodeTitle(ENodeTitleType::ListView);
 	FString Tooltip = TemplateNode->GetTooltip();
 	FString Keywords = TemplateNode->GetKeywords();
 
 	TSharedPtr<FEdGraphSchemaAction_K2NewNode> NodeAction = FK2ActionMenuBuilder::AddNewNodeAction(ContextMenuBuilder, Category, MenuDesc, Tooltip, 0, Keywords);
 	NodeAction->NodeTemplate = TemplateNode;
+	NodeAction->SearchTitle = TemplateNode->GetNodeSearchTitle();
 
 	return NodeAction;
 }
diff --git a/Engine/Source/Editor/AnimGraph/Private/AnimGraphNode_BlendListByBool.cpp b/Engine/Source/Editor/AnimGraph/Private/AnimGraphNode_BlendListByBool.cpp
--- a/Engine/Source/Editor/AnimGraph/Private/AnimGraphNode_BlendListByBool.cpp
+++ b/Engine/Source/Editor/AnimGraph/Private/AnimGraphNode_BlendListByBool.cpp
@@ -8,19 +8,21 @@
 /////////////////////////////////////////////////////
 // UAnimGraphNode_BlendListByBool
 
+#define LOCTEXT_NAMESPACE "A3Nodes"
+
 UAnimGraphNode_BlendListByBool::UAnimGraphNode_BlendListByBool(const FPostConstructInitializeProperties& PCIP)
 	: Super(PCIP)
 {
 }
 
-FString UAnimGraphNode_BlendListByBool::GetNodeTitle(ENodeTitleType::Type TitleType) const
+FText UAnimGraphNode_BlendListByBool::GetNodeTitle(ENodeTitleType::Type TitleType) const
 {
-	return FString::Printf(TEXT("Blend Poses by bool"));
+	return LOCTEXT("UAnimGraphNode_BlendListByBool_Tooltip", "Blend Poses by bool");
 }
 
 FString UAnimGraphNode_BlendListByBool::GetTooltip() const
 {
-	return FString::Printf(TEXT("Blend List (by bool)"));
+	return LOCTEXT("UAnimGraphNode_BlendListByBool_Title", "Blend List (by bool)").ToString();
 }
 
 void UAnimGraphNode_BlendListByBool::CustomizePinData(UEdGraphPin* Pin, FName SourcePropertyName, int32 ArrayIndex) const 
@@ -31,15 +33,18 @@ void UAnimGraphNode_BlendListByBool::CustomizePinData(UEdGraphPin* Pin, FName So
 	if (ArrayIndex != INDEX_NONE)
 	{
 		// Note: This is intentionally flipped, as it looks better with true as the topmost element!
-		FString TrueFalse = (ArrayIndex == 0) ? TEXT("True") : TEXT("False");
+		FFormatNamedArguments Args;
+		Args.Add(TEXT("TrueFalse"), (ArrayIndex == 0) ? LOCTEXT("True", "True") : LOCTEXT("False", "False"));
 
 		if (SourcePropertyName == BlendPoses)
 		{
-			Pin->PinFriendlyName = TrueFalse + TEXT(" Pose");
+			Pin->PinFriendlyName = FText::Format(LOCTEXT("BoolPoseFriendlyName", "{TrueFalse} Pose"), Args);
 		}
 		else if (SourcePropertyName == BlendTimes)
 		{
-			Pin->PinFriendlyName = TrueFalse + TEXT(" Blend Time");
+			Pin->PinFriendlyName = FText::Format(LOCTEXT("BoolBlendTimeFriendlyName", "{TrueFalse} Blend Time"), Args);
 		}
 	}
-}
\ No newline at end of file
+}
+
+#undef LOCTEXT_NAMESPACE
\ No newline at end of file
diff --git a/Engine/Source/Editor/AnimGraph/Private/AnimGraphNode_BlendListByEnum.cpp b/Engine/Source/Editor/AnimGraph/Private/AnimGraphNode_BlendListByEnum.cpp
--- a/Engine/Source/Editor/AnimGraph/Private/AnimGraphNode_BlendListByEnum.cpp
+++ b/Engine/Source/Editor/AnimGraph/Private/AnimGraphNode_BlendListByEnum.cpp
@@ -26,16 +26,33 @@ FString UAnimGraphNode_BlendListByEnum::GetTooltip() const
 {
 	if (BoundEnum != NULL)
 	{
-		return FString::Printf(TEXT("Blend Poses based on enum (%s). Use context menu to add new entries."), *BoundEnum->GetName());
+		FFormatNamedArguments Args;
+		Args.Add(TEXT("EnumName"), FText::FromString(BoundEnum->GetName()));
+		return FText::Format(LOCTEXT("AnimGraphNode_BlendListByEnum_Tooltip", "Blend Poses based on enum ({EnumName}). Use context menu to add new entries."), Args).ToString();
 	}
 	else
 	{
-		return FString::Printf(TEXT("ERROR: Blend Poses (by missing enum)"));
+		return LOCTEXT("AnimGraphNode_BlendListByEnum_TooltipError", "ERROR: Blend Poses (by missing enum)").ToString();
+	}
+}
+
+FText UAnimGraphNode_BlendListByEnum::GetNodeTitle(ENodeTitleType::Type TitleType) const
+{
+	if (BoundEnum != NULL)
+	{
+		FFormatNamedArguments Args;
+		Args.Add(TEXT("EnumName"), FText::FromString(BoundEnum->GetName()));
+		return FText::Format(LOCTEXT("AnimGraphNode_BlendListByEnum_Title", "Blend Poses ({EnumName})"), Args);
+	}
+	else
+	{
+		return LOCTEXT("AnimGraphNode_BlendListByEnum_TitleError", "ERROR: Blend Poses (by missing enum)");
 	}
 }
 
-FString UAnimGraphNode_BlendListByEnum::GetNodeTitle(ENodeTitleType::Type TitleType) const
+FString UAnimGraphNode_BlendListByEnum::GetNodeNativeTitle(ENodeTitleType::Type TitleType) const
 {
+	// Do not setup this function for localization, intentionally left unlocalized!
 	if (BoundEnum != NULL)
 	{
 		return FString::Printf(TEXT("Blend Poses (%s)"), *BoundEnum->GetName());
@@ -69,6 +86,7 @@ void UAnimGraphNode_BlendListByEnum::GetMenuEntries(FGraphContextMenuBuilder& Co
 
 			TSharedPtr<FEdGraphSchemaAction_K2NewNode> Action = FK2ActionMenuBuilder::AddNewNodeAction(ContextMenuBuilder, EnumTemplate->GetNodeCategory(), EnumTemplate->GetNodeTitle(ENodeTitleType::ListView), EnumTemplate->GetTooltip(), 0, EnumTemplate->GetKeywords());
 			Action->NodeTemplate = EnumTemplate;
+			Action->SearchTitle = EnumTemplate->GetNodeSearchTitle();
 		}
 	}
 }
@@ -219,32 +237,36 @@ void UAnimGraphNode_BlendListByEnum::CustomizePinData(UEdGraphPin* Pin, FName So
 				const int32 EnumIndex = BoundEnum->FindEnumIndex(EnumElementName);
 				if (EnumIndex != INDEX_NONE)
 				{
-					Pin->PinFriendlyName = BoundEnum->GetEnumText(EnumIndex).ToString();
+					Pin->PinFriendlyName = BoundEnum->GetEnumText(EnumIndex);
 				}
 				else
 				{
-					Pin->PinFriendlyName = EnumElementName.ToString();
+					Pin->PinFriendlyName = FText::FromName(EnumElementName);
 				}
 			}
 			else
 			{
-				Pin->PinFriendlyName = TEXT("Invalid index");
+				Pin->PinFriendlyName = LOCTEXT("InvalidIndex", "Invalid index");
 			}
 		}
 		else if (ensure(RawArrayIndex == 0))
 		{
-			Pin->PinFriendlyName = TEXT("Default");
+			Pin->PinFriendlyName = LOCTEXT("Default", "Default");
 		}
 
 		// Append the pin type
 		if (bIsPosePin)
 		{
-			Pin->PinFriendlyName += TEXT(" Pose");
+			FFormatNamedArguments Args;
+			Args.Add(TEXT("PinFriendlyName"), Pin->PinFriendlyName);
+			Pin->PinFriendlyName = FText::Format(LOCTEXT("FriendlyNamePose", "{PinFriendlyName} Pose"), Args);
 		}
 
 		if (bIsTimePin)
 		{
-			Pin->PinFriendlyName += TEXT(" Blend Time");
+			FFormatNamedArguments Args;
+			Args.Add(TEXT("PinFriendlyName"), Pin->PinFriendlyName);
+			Pin->PinFriendlyName = FText::Format(LOCTEXT("FriendlyNameBlendTime", "{PinFriendlyName} Blend Time"), Args);
 		}
 	}
 }
@@ -297,4 +319,4 @@ void UAnimGraphNode_BlendListByEnum::BakeDataDuringCompilation(class FCompilerRe
 
 /////////////////////////////////////////////////////
 
-#undef LOCTEXT_NAMESPACE
\ No newline at end of file
+#undef LOCTEXT_NAMESPACE
diff --git a/Engine/Source/Editor/AnimGraph/Private/AnimGraphNode_BlendListByInt.cpp b/Engine/Source/Editor/AnimGraph/Private/AnimGraphNode_BlendListByInt.cpp
--- a/Engine/Source/Editor/AnimGraph/Private/AnimGraphNode_BlendListByInt.cpp
+++ b/Engine/Source/Editor/AnimGraph/Private/AnimGraphNode_BlendListByInt.cpp
@@ -9,6 +9,8 @@
 /////////////////////////////////////////////////////
 // UAnimGraphNode_BlendListByInt
 
+#define LOCTEXT_NAMESPACE "A3Nodes"
+
 UAnimGraphNode_BlendListByInt::UAnimGraphNode_BlendListByInt(const FPostConstructInitializeProperties& PCIP)
 	: Super(PCIP)
 {
@@ -16,12 +18,12 @@ UAnimGraphNode_BlendListByInt::UAnimGraphNode_BlendListByInt(const FPostConstruc
 
 FString UAnimGraphNode_BlendListByInt::GetTooltip() const
 {
-	return FString::Printf(TEXT("Blend List (by int)"));
+	return LOCTEXT("AnimGraphNode_BlendListByInt_Tooltip", "Blend List (by int)").ToString();
 }
 
-FString UAnimGraphNode_BlendListByInt::GetNodeTitle(ENodeTitleType::Type TitleType) const
+FText UAnimGraphNode_BlendListByInt::GetNodeTitle(ENodeTitleType::Type TitleType) const
 {
-	return FString::Printf(TEXT("Blend Poses by int"));
+	return LOCTEXT("AnimGraphNode_BlendListByInt_Title", "Blend Poses by int");
 }
 
 void UAnimGraphNode_BlendListByInt::PostPlacedNewNode()
@@ -84,4 +86,6 @@ void UAnimGraphNode_BlendListByInt::GetContextMenuActions(const FGraphNodeContex
 		}
 		Context.MenuBuilder->EndSection();	
 	}
-}
\ No newline at end of file
+}
+
+#undef LOCTEXT_NAMESPACE
\ No newline at end of file
diff --git a/Engine/Source/Editor/AnimGraph/Private/AnimGraphNode_BlendSpaceBase.cpp b/Engine/Source/Editor/AnimGraph/Private/AnimGraphNode_BlendSpaceBase.cpp
--- a/Engine/Source/Editor/AnimGraph/Private/AnimGraphNode_BlendSpaceBase.cpp
+++ b/Engine/Source/Editor/AnimGraph/Private/AnimGraphNode_BlendSpaceBase.cpp
@@ -35,6 +35,8 @@ struct FNewBlendSpacePlayerAction : public FEdGraphSchemaAction_K2NewNode
 
 		// Grab extra keywords
 		Keywords = BlendSpace->GetPathName();
+
+		SearchTitle = NodeTemplate->GetNodeSearchTitle();
 	}
 };
 
@@ -86,16 +88,16 @@ void UAnimGraphNode_BlendSpaceBase::CustomizePinData(UEdGraphPin* Pin, FName Sou
 	{
 		if (SourcePropertyName == TEXT("X"))
 		{
-			Pin->PinFriendlyName = BlendSpace->GetBlendParameter(0).DisplayName;
+			Pin->PinFriendlyName = FText::FromString(BlendSpace->GetBlendParameter(0).DisplayName);
 		}
 		else if (SourcePropertyName == TEXT("Y"))
 		{
-			Pin->PinFriendlyName = BlendSpace->GetBlendParameter(1).DisplayName;
+			Pin->PinFriendlyName = FText::FromString(BlendSpace->GetBlendParameter(1).DisplayName);
 			Pin->bHidden = (BlendSpace->NumOfDimension == 1) ? 1 : 0;
 		}
 		else if (SourcePropertyName == TEXT("Z"))
 		{
-			Pin->PinFriendlyName = BlendSpace->GetBlendParameter(2).DisplayName;
+			Pin->PinFriendlyName = FText::FromString(BlendSpace->GetBlendParameter(2).DisplayName);
 		}
 	}
 }
diff --git a/Engine/Source/Editor/AnimGraph/Private/AnimGraphNode_BlendSpaceEvaluator.cpp b/Engine/Source/Editor/AnimGraph/Private/AnimGraphNode_BlendSpaceEvaluator.cpp
--- a/Engine/Source/Editor/AnimGraph/Private/AnimGraphNode_BlendSpaceEvaluator.cpp
+++ b/Engine/Source/Editor/AnimGraph/Private/AnimGraphNode_BlendSpaceEvaluator.cpp
@@ -8,6 +8,8 @@
 /////////////////////////////////////////////////////
 // UAnimGraphNode_BlendSpaceEvaluator
 
+#define LOCTEXT_NAMESPACE "A3Nodes"
+
 UAnimGraphNode_BlendSpaceEvaluator::UAnimGraphNode_BlendSpaceEvaluator(const FPostConstructInitializeProperties& PCIP)
 	: Super(PCIP)
 {
@@ -18,10 +20,39 @@ FString UAnimGraphNode_BlendSpaceEvaluator::GetTooltip() const
 	return FString::Printf(TEXT("Blendspace Evaluator'%s'"), (Node.BlendSpace != NULL) ? *(Node.BlendSpace->GetPathName()) : TEXT("(None)"));
 }
 
-FString UAnimGraphNode_BlendSpaceEvaluator::GetNodeTitle(ENodeTitleType::Type TitleType) const
+FText UAnimGraphNode_BlendSpaceEvaluator::GetNodeTitle(ENodeTitleType::Type TitleType) const
 {
-	const FString BlendSpaceName((Node.BlendSpace != NULL) ? *(Node.BlendSpace->GetName()) : TEXT("(None)"));
+	const FText BlendSpaceName((Node.BlendSpace != NULL) ? FText::FromString(Node.BlendSpace->GetName()) : LOCTEXT("None", "(None)"));
 	
+	if (TitleType == ENodeTitleType::ListView)
+	{
+		FFormatNamedArguments Args;
+		Args.Add(TEXT("BlendSpaceName"), BlendSpaceName);
+		return FText::Format(LOCTEXT("BlendSpaceEvaluatorListTitle", "Blendspace Evaluator '{BlendSpaceName}'"), Args);
+	}
+	else
+	{
+		FFormatNamedArguments TitleArgs;
+		TitleArgs.Add(TEXT("BlendSpaceName"), BlendSpaceName);
+		FText Title = FText::Format(LOCTEXT("BlendSpaceEvaluatorFullTitle", "{BlendSpaceName}\nBlendspace Evaluator"), TitleArgs);
+
+		if ((TitleType == ENodeTitleType::FullTitle) && (SyncGroup.GroupName != NAME_None))
+		{
+			FFormatNamedArguments Args;
+			Args.Add(TEXT("Title"), Title);
+			Args.Add(TEXT("SyncGroupName"), FText::FromName(SyncGroup.GroupName));
+			Title = FText::Format(LOCTEXT("BlendSpaceNodeGroupSubtitle", "{Title}\nSync group {SyncGroupName}"), Args);
+		}
+
+		return Title;
+	}
+}
+
+FString UAnimGraphNode_BlendSpaceEvaluator::GetNodeNativeTitle(ENodeTitleType::Type TitleType) const
+{
+	// Do not setup this function for localization, intentionally left unlocalized!
+	const FString BlendSpaceName((Node.BlendSpace != NULL) ? *(Node.BlendSpace->GetName()) : TEXT("(None)"));
+
 	if (TitleType == ENodeTitleType::ListView)
 	{
 		return FString::Printf(TEXT("Blendspace Evaluator '%s'"), *BlendSpaceName);
@@ -33,7 +64,7 @@ FString UAnimGraphNode_BlendSpaceEvaluator::GetNodeTitle(ENodeTitleType::Type Ti
 		if ((TitleType == ENodeTitleType::FullTitle) && (SyncGroup.GroupName != NAME_None))
 		{
 			Title += TEXT("\n");
-			Title += FString::Printf(*NSLOCTEXT("A3Nodes", "BlendSpaceNodeGroupSubtitle", "Sync group %s").ToString(), *SyncGroup.GroupName.ToString());
+			Title += FString::Printf(TEXT("Sync group %s"), *SyncGroup.GroupName.ToString());
 		}
 
 		return Title;
@@ -102,4 +133,6 @@ const TCHAR* UAnimGraphNode_BlendSpaceEvaluator::GetTimePropertyName() const
 UScriptStruct* UAnimGraphNode_BlendSpaceEvaluator::GetTimePropertyStruct() const 
 {
 	return FAnimNode_BlendSpaceEvaluator::StaticStruct();
-}
\ No newline at end of file
+}
+
+#undef LOCTEXT_NAMESPACE
\ No newline at end of file
diff --git a/Engine/Source/Editor/AnimGraph/Private/AnimGraphNode_BlendSpacePlayer.cpp b/Engine/Source/Editor/AnimGraph/Private/AnimGraphNode_BlendSpacePlayer.cpp
--- a/Engine/Source/Editor/AnimGraph/Private/AnimGraphNode_BlendSpacePlayer.cpp
+++ b/Engine/Source/Editor/AnimGraph/Private/AnimGraphNode_BlendSpacePlayer.cpp
@@ -7,6 +7,8 @@
 /////////////////////////////////////////////////////
 // UAnimGraphNode_BlendSpacePlayer
 
+#define LOCTEXT_NAMESPACE "A3Nodes"
+
 UAnimGraphNode_BlendSpacePlayer::UAnimGraphNode_BlendSpacePlayer(const FPostConstructInitializeProperties& PCIP)
 	: Super(PCIP)
 {
@@ -17,10 +19,39 @@ FString UAnimGraphNode_BlendSpacePlayer::GetTooltip() const
 	return FString::Printf(TEXT("Blendspace Player '%s'"), (Node.BlendSpace != NULL) ? *(Node.BlendSpace->GetPathName()) : TEXT("(None)"));
 }
 
-FString UAnimGraphNode_BlendSpacePlayer::GetNodeTitle(ENodeTitleType::Type TitleType) const
+FText UAnimGraphNode_BlendSpacePlayer::GetNodeTitle(ENodeTitleType::Type TitleType) const
 {
-	const FString BlendSpaceName((Node.BlendSpace != NULL) ? *(Node.BlendSpace->GetName()) : TEXT("(None)"));
+	const FText BlendSpaceName((Node.BlendSpace != NULL) ? FText::FromString(Node.BlendSpace->GetName()) : LOCTEXT("None", "(None)"));
 	
+	if (TitleType == ENodeTitleType::ListView)
+	{
+		FFormatNamedArguments Args;
+		Args.Add(TEXT("BlendSpaceName"), BlendSpaceName);
+		return FText::Format(LOCTEXT("BlendspacePlayer", "Blendspace Player'{BlendSpaceName}'"), Args);
+	}
+	else
+	{
+		FFormatNamedArguments TitleArgs;
+		TitleArgs.Add(TEXT("BlendSpaceName"), BlendSpaceName);
+		FText Title = FText::Format(LOCTEXT("BlendSpacePlayerFullTitle", "{BlendSpaceName}\nBlendspace Player"), TitleArgs);
+
+		if ((TitleType == ENodeTitleType::FullTitle) && (SyncGroup.GroupName != NAME_None))
+		{
+			FFormatNamedArguments Args;
+			Args.Add(TEXT("Title"), Title);
+			Args.Add(TEXT("SyncGroupName"), FText::FromName(SyncGroup.GroupName));
+			Title = FText::Format(LOCTEXT("BlendSpaceNodeGroupSubtitle", "{Title}\nSync group {SyncGroupName}"), Args);
+		}
+
+		return Title;
+	}
+}
+
+FString UAnimGraphNode_BlendSpacePlayer::GetNodeNativeTitle(ENodeTitleType::Type TitleType) const
+{
+	// Do not setup this function for localization, intentionally left unlocalized!
+	const FString BlendSpaceName((Node.BlendSpace != NULL) ? *(Node.BlendSpace->GetName()) : TEXT("(None)"));
+
 	if (TitleType == ENodeTitleType::ListView)
 	{
 		return FString::Printf(TEXT("Blendspace Player'%s'"), *BlendSpaceName);
@@ -32,7 +63,7 @@ FString UAnimGraphNode_BlendSpacePlayer::GetNodeTitle(ENodeTitleType::Type Title
 		if ((TitleType == ENodeTitleType::FullTitle) && (SyncGroup.GroupName != NAME_None))
 		{
 			Title += TEXT("\n");
-			Title += FString::Printf(*NSLOCTEXT("A3Nodes", "BlendSpaceNodeGroupSubtitle", "Sync group %s").ToString(), *SyncGroup.GroupName.ToString());
+			Title += FString::Printf(TEXT("Sync group %s"), *SyncGroup.GroupName.ToString());
 		}
 
 		return Title;
@@ -116,3 +147,5 @@ UScriptStruct* UAnimGraphNode_BlendSpacePlayer::GetTimePropertyStruct() const
 	return FAnimNode_BlendSpacePlayer::StaticStruct();
 }
 
+#undef LOCTEXT_NAMESPACE
+
diff --git a/Engine/Source/Editor/AnimGraph/Private/AnimGraphNode_ComponentToLocalSpace.cpp b/Engine/Source/Editor/AnimGraph/Private/AnimGraphNode_ComponentToLocalSpace.cpp
--- a/Engine/Source/Editor/AnimGraph/Private/AnimGraphNode_ComponentToLocalSpace.cpp
+++ b/Engine/Source/Editor/AnimGraph/Private/AnimGraphNode_ComponentToLocalSpace.cpp
@@ -5,6 +5,8 @@
 /////////////////////////////////////////////////////
 // UAnimGraphNode_ComponentToLocalSpace
 
+#define LOCTEXT_NAMESPACE "A3Nodes"
+
 UAnimGraphNode_ComponentToLocalSpace::UAnimGraphNode_ComponentToLocalSpace(const FPostConstructInitializeProperties& PCIP)
 	: Super(PCIP)
 {
@@ -17,12 +19,12 @@ FLinearColor UAnimGraphNode_ComponentToLocalSpace::GetNodeTitleColor() const
 
 FString UAnimGraphNode_ComponentToLocalSpace::GetTooltip() const
 {
-	return TEXT("Convert Component Space Pose to Local Pose");
+	return LOCTEXT("AnimGraphNode_ComponentToLocalSpace_Tooltip", "Convert Component Space Pose to Local Pose").ToString();
 }
 
-FString UAnimGraphNode_ComponentToLocalSpace::GetNodeTitle(ENodeTitleType::Type TitleType) const
+FText UAnimGraphNode_ComponentToLocalSpace::GetNodeTitle(ENodeTitleType::Type TitleType) const
 {
-	return TEXT("Component To Local");
+	return LOCTEXT("AnimGraphNode_ComponentToLocalSpace_Title", "Component To Local");
 }
 
 FString UAnimGraphNode_ComponentToLocalSpace::GetNodeCategory() const
@@ -34,3 +36,5 @@ void UAnimGraphNode_ComponentToLocalSpace::PostProcessPinName(const UEdGraphPin*
 {
 	DisplayName = TEXT("");
 }
+
+#undef LOCTEXT_NAMESPACE
diff --git a/Engine/Source/Editor/AnimGraph/Private/AnimGraphNode_CopyBone.cpp b/Engine/Source/Editor/AnimGraph/Private/AnimGraphNode_CopyBone.cpp
--- a/Engine/Source/Editor/AnimGraph/Private/AnimGraphNode_CopyBone.cpp
+++ b/Engine/Source/Editor/AnimGraph/Private/AnimGraphNode_CopyBone.cpp
@@ -5,22 +5,36 @@
 /////////////////////////////////////////////////////
 // UAnimGraphNode_CopyBoneSkeletalControl
 
+#define LOCTEXT_NAMESPACE "A3Nodes"
+
 UAnimGraphNode_CopyBone::UAnimGraphNode_CopyBone(const FPostConstructInitializeProperties& PCIP)
 	: Super(PCIP)
 {
 }
 
-FString UAnimGraphNode_CopyBone::GetControllerDescription() const
+FText UAnimGraphNode_CopyBone::GetControllerDescription() const
 {
-	return TEXT("Copy Bone");
+	return LOCTEXT("CopyBone", "Copy Bone");
 }
 
 FString UAnimGraphNode_CopyBone::GetTooltip() const
 {
-	return TEXT("The Copy Bone control copies the Transform data or any component of it - i.e. Translation, Rotation, or Scale - from one bone to another.");
+	return LOCTEXT("AnimGraphNode_CopyBone_Tooltip", "The Copy Bone control copies the Transform data or any component of it - i.e. Translation, Rotation, or Scale - from one bone to another.").ToString();
 }
 
-FString UAnimGraphNode_CopyBone::GetNodeTitle(ENodeTitleType::Type TitleType) const
+FText UAnimGraphNode_CopyBone::GetNodeTitle(ENodeTitleType::Type TitleType) const
 {
-	return FString::Printf(TEXT("%s\nSource Bone: %s\nTarget Bone: %s"), *GetControllerDescription(), *Node.SourceBone.BoneName.ToString(), *Node.TargetBone.BoneName.ToString());
+	FFormatNamedArguments Args;
+	Args.Add(TEXT("ControllerDescription"), GetControllerDescription());
+	Args.Add(TEXT("SourceBoneName"), FText::FromName(Node.SourceBone.BoneName));
+	Args.Add(TEXT("TargetBoneName"), FText::FromName(Node.TargetBone.BoneName));
+	return FText::Format(LOCTEXT("AnimGraphNode_CopyBone_Title", "{ControllerDescription}\nSource Bone: {SourceBoneName}\nTarget Bone: {TargetBoneName}"), Args);
 }
+
+FString UAnimGraphNode_CopyBone::GetNodeNativeTitle(ENodeTitleType::Type TitleType) const
+{
+	// Do not setup this function for localization, intentionally left unlocalized!
+	return FString::Printf(TEXT("%s\nSource Bone: %s\nTarget Bone: %s"), *GetControllerDescription().ToString(), *Node.SourceBone.BoneName.ToString(), *Node.TargetBone.BoneName.ToString());
+}
+
+#undef LOCTEXT_NAMESPACE
diff --git a/Engine/Source/Editor/AnimGraph/Private/AnimGraphNode_CustomTransitionResult.cpp b/Engine/Source/Editor/AnimGraph/Private/AnimGraphNode_CustomTransitionResult.cpp
--- a/Engine/Source/Editor/AnimGraph/Private/AnimGraphNode_CustomTransitionResult.cpp
+++ b/Engine/Source/Editor/AnimGraph/Private/AnimGraphNode_CustomTransitionResult.cpp
@@ -5,6 +5,8 @@
 /////////////////////////////////////////////////////
 // UAnimGraphNode_CustomTransitionResult
 
+#define LOCTEXT_NAMESPACE "A3Nodes"
+
 UAnimGraphNode_CustomTransitionResult::UAnimGraphNode_CustomTransitionResult(const FPostConstructInitializeProperties& PCIP)
 	: Super(PCIP)
 {
@@ -17,10 +19,12 @@ FLinearColor UAnimGraphNode_CustomTransitionResult::GetNodeTitleColor() const
 
 FString UAnimGraphNode_CustomTransitionResult::GetTooltip() const
 {
-	return TEXT("Result node for a custom transition blend graph");
+	return LOCTEXT("AnimGraphNode_CustomTransitionResult_Tooltip", "Result node for a custom transition blend graph").ToString();
 }
 
-FString UAnimGraphNode_CustomTransitionResult::GetNodeTitle(ENodeTitleType::Type TitleType) const
+FText UAnimGraphNode_CustomTransitionResult::GetNodeTitle(ENodeTitleType::Type TitleType) const
 {
-	return TEXT("Custom Transition Blend Result");
+	return LOCTEXT("AnimGraphNode_CustomTransitionResult_Title", "Custom Transition Blend Result");
 }
+
+#undef LOCTEXT_NAMESPACE
\ No newline at end of file
diff --git a/Engine/Source/Editor/AnimGraph/Private/AnimGraphNode_Fabrik.cpp b/Engine/Source/Editor/AnimGraph/Private/AnimGraphNode_Fabrik.cpp
--- a/Engine/Source/Editor/AnimGraph/Private/AnimGraphNode_Fabrik.cpp
+++ b/Engine/Source/Editor/AnimGraph/Private/AnimGraphNode_Fabrik.cpp
@@ -5,18 +5,21 @@
 /////////////////////////////////////////////////////
 // UAnimGraphNode_Fabrik 
 
+#define LOCTEXT_NAMESPACE "A3Nodes"
+
 UAnimGraphNode_Fabrik::UAnimGraphNode_Fabrik(const FPostConstructInitializeProperties& PCIP)
 : Super(PCIP)
 {
 }
 
-FString UAnimGraphNode_Fabrik::GetControllerDescription() const
+FText UAnimGraphNode_Fabrik::GetControllerDescription() const
 {
-	return TEXT("FABRIK");
+	return LOCTEXT("Fabrik", "FABRIK");
 }
 
-FString UAnimGraphNode_Fabrik::GetNodeTitle(ENodeTitleType::Type TitleType) const
+FText UAnimGraphNode_Fabrik::GetNodeTitle(ENodeTitleType::Type TitleType) const
 {
-	FString Result = *GetControllerDescription();
-	return Result;
-}
\ No newline at end of file
+	return GetControllerDescription();
+}
+
+#undef LOCTEXT_NAMESPACE
\ No newline at end of file
diff --git a/Engine/Source/Editor/AnimGraph/Private/AnimGraphNode_IdentityPose.cpp b/Engine/Source/Editor/AnimGraph/Private/AnimGraphNode_IdentityPose.cpp
--- a/Engine/Source/Editor/AnimGraph/Private/AnimGraphNode_IdentityPose.cpp
+++ b/Engine/Source/Editor/AnimGraph/Private/AnimGraphNode_IdentityPose.cpp
@@ -5,6 +5,8 @@
 /////////////////////////////////////////////////////
 // UAnimGraphNode_IdentityPose
 
+#define LOCTEXT_NAMESPACE "A3Nodes"
+
 UAnimGraphNode_IdentityPose::UAnimGraphNode_IdentityPose(const FPostConstructInitializeProperties& PCIP)
 	: Super(PCIP)
 {
@@ -13,11 +15,12 @@ UAnimGraphNode_IdentityPose::UAnimGraphNode_IdentityPose(const FPostConstructIni
 
 FString UAnimGraphNode_IdentityPose::GetTooltip() const
 {
-	return TEXT("Returns identity pose.");
+	return LOCTEXT("AnimGraphNode_IdentityPose_Tooltip", "Returns identity pose.").ToString();
 }
 
-FString UAnimGraphNode_IdentityPose::GetNodeTitle(ENodeTitleType::Type TitleType) const
+FText UAnimGraphNode_IdentityPose::GetNodeTitle(ENodeTitleType::Type TitleType) const
 {
-	return FString::Printf(TEXT("Additive Identity Pose"));
+	return LOCTEXT("AnimGraphNode_IdentityPose_Title", "Additive Identity Pose");
 }
 
+#undef LOCTEXT_NAMESPACE
diff --git a/Engine/Source/Editor/AnimGraph/Private/AnimGraphNode_LayeredBoneBlend.cpp b/Engine/Source/Editor/AnimGraph/Private/AnimGraphNode_LayeredBoneBlend.cpp
--- a/Engine/Source/Editor/AnimGraph/Private/AnimGraphNode_LayeredBoneBlend.cpp
+++ b/Engine/Source/Editor/AnimGraph/Private/AnimGraphNode_LayeredBoneBlend.cpp
@@ -8,6 +8,8 @@
 /////////////////////////////////////////////////////
 // UAnimGraphNode_LayeredBoneBlend
 
+#define LOCTEXT_NAMESPACE "A3Nodes"
+
 UAnimGraphNode_LayeredBoneBlend::UAnimGraphNode_LayeredBoneBlend(const FPostConstructInitializeProperties& PCIP)
 	: Super(PCIP)
 {
@@ -20,12 +22,12 @@ FLinearColor UAnimGraphNode_LayeredBoneBlend::GetNodeTitleColor() const
 
 FString UAnimGraphNode_LayeredBoneBlend::GetTooltip() const
 {
-	return FString::Printf(TEXT("Layered blend per bone"));
+	return LOCTEXT("AnimGraphNode_LayeredBoneBlend_Tooltip", "Layered blend per bone").ToString();
 }
 
-FString UAnimGraphNode_LayeredBoneBlend::GetNodeTitle(ENodeTitleType::Type TitleType) const
+FText UAnimGraphNode_LayeredBoneBlend::GetNodeTitle(ENodeTitleType::Type TitleType) const
 {
-	return FString::Printf(TEXT("Layered blend per bone"));
+	return LOCTEXT("AnimGraphNode_LayeredBoneBlend_Title", "Layered blend per bone");
 }
 
 FString UAnimGraphNode_LayeredBoneBlend::GetNodeCategory() const
@@ -35,7 +37,7 @@ FString UAnimGraphNode_LayeredBoneBlend::GetNodeCategory() const
 
 void UAnimGraphNode_LayeredBoneBlend::AddPinToBlendByFilter()
 {
-	FScopedTransaction Transaction( NSLOCTEXT("A3Nodes", "AddPinToBlend", "AddPinToBlendByFilter") );
+	FScopedTransaction Transaction( LOCTEXT("AddPinToBlend", "AddPinToBlendByFilter") );
 	Modify();
 
 	Node.AddPose();
@@ -45,7 +47,7 @@ void UAnimGraphNode_LayeredBoneBlend::AddPinToBlendByFilter()
 
 void UAnimGraphNode_LayeredBoneBlend::RemovePinFromBlendByFilter(UEdGraphPin* Pin)
 {
-	FScopedTransaction Transaction( NSLOCTEXT("A3Nodes", "RemovePinFromBlend", "RemovePinFromBlendByFilter") );
+	FScopedTransaction Transaction( LOCTEXT("RemovePinFromBlend", "RemovePinFromBlendByFilter") );
 	Modify();
 
 	UProperty* AssociatedProperty;
@@ -67,7 +69,7 @@ void UAnimGraphNode_LayeredBoneBlend::GetContextMenuActions(const FGraphNodeCont
 {
 	if (!Context.bIsDebugging)
 	{
-		Context.MenuBuilder->BeginSection("AnimGraphNodeLayeredBoneblend", NSLOCTEXT("A3Nodes", "LayeredBoneBlend", "Layered Bone Blend"));
+		Context.MenuBuilder->BeginSection("AnimGraphNodeLayeredBoneblend", LOCTEXT("LayeredBoneBlend", "Layered Bone Blend"));
 		{
 			if (Context.Pin != NULL)
 			{
@@ -86,3 +88,5 @@ void UAnimGraphNode_LayeredBoneBlend::GetContextMenuActions(const FGraphNodeCont
 		Context.MenuBuilder->EndSection();
 	}
 }
+
+#undef LOCTEXT_NAMESPACE
\ No newline at end of file
diff --git a/Engine/Source/Editor/AnimGraph/Private/AnimGraphNode_LocalRefPose.cpp b/Engine/Source/Editor/AnimGraph/Private/AnimGraphNode_LocalRefPose.cpp
--- a/Engine/Source/Editor/AnimGraph/Private/AnimGraphNode_LocalRefPose.cpp
+++ b/Engine/Source/Editor/AnimGraph/Private/AnimGraphNode_LocalRefPose.cpp
@@ -5,6 +5,8 @@
 /////////////////////////////////////////////////////
 // UAnimGraphNode_LocalRefPose
 
+#define LOCTEXT_NAMESPACE "A3Nodes"
+
 UAnimGraphNode_LocalRefPose::UAnimGraphNode_LocalRefPose(const FPostConstructInitializeProperties& PCIP)
 	: Super(PCIP)
 {
@@ -13,11 +15,12 @@ UAnimGraphNode_LocalRefPose::UAnimGraphNode_LocalRefPose(const FPostConstructIni
 
 FString UAnimGraphNode_LocalRefPose::GetTooltip() const
 {
-	return TEXT("Returns local space reference pose.");
+	return LOCTEXT("UAnimGraphNode_LocalRefPose_Tooltip", "Returns local space reference pose.").ToString();
 }
 
-FString UAnimGraphNode_LocalRefPose::GetNodeTitle(ENodeTitleType::Type TitleType) const
+FText UAnimGraphNode_LocalRefPose::GetNodeTitle(ENodeTitleType::Type TitleType) const
 {
-	return FString::Printf(TEXT("Local Space Ref Pose"));
+	return LOCTEXT("UAnimGraphNode_LocalRefPose_Title", "Local Space Ref Pose");
 }
 
+#undef LOCTEXT_NAMESPACE
diff --git a/Engine/Source/Editor/AnimGraph/Private/AnimGraphNode_LocalToComponentSpace.cpp b/Engine/Source/Editor/AnimGraph/Private/AnimGraphNode_LocalToComponentSpace.cpp
--- a/Engine/Source/Editor/AnimGraph/Private/AnimGraphNode_LocalToComponentSpace.cpp
+++ b/Engine/Source/Editor/AnimGraph/Private/AnimGraphNode_LocalToComponentSpace.cpp
@@ -8,6 +8,8 @@
 /////////////////////////////////////////////////////
 // UAnimGraphNode_LocalToComponentSpace
 
+#define LOCTEXT_NAMESPACE "A3Nodes"
+
 UAnimGraphNode_LocalToComponentSpace::UAnimGraphNode_LocalToComponentSpace(const FPostConstructInitializeProperties& PCIP)
 	: Super(PCIP)
 {
@@ -20,12 +22,12 @@ FLinearColor UAnimGraphNode_LocalToComponentSpace::GetNodeTitleColor() const
 
 FString UAnimGraphNode_LocalToComponentSpace::GetTooltip() const
 {
-	return TEXT("Convert Local Pose to Component Space Pose");
+	return LOCTEXT("AnimGraphNode_LocalToComponentSpace_Tooltip", "Convert Local Pose to Component Space Pose").ToString();
 }
 
-FString UAnimGraphNode_LocalToComponentSpace::GetNodeTitle(ENodeTitleType::Type TitleType) const
+FText UAnimGraphNode_LocalToComponentSpace::GetNodeTitle(ENodeTitleType::Type TitleType) const
 {
-	return TEXT("Local To Component");
+	return LOCTEXT("AnimGraphNode_LocalToComponentSpace_Title", "Local To Component");
 }
 
 FString UAnimGraphNode_LocalToComponentSpace::GetNodeCategory() const
@@ -43,3 +45,5 @@ void UAnimGraphNode_LocalToComponentSpace::PostProcessPinName(const UEdGraphPin*
 {
 	DisplayName = TEXT("");
 }
+
+#undef LOCTEXT_NAMESPACE
\ No newline at end of file
diff --git a/Engine/Source/Editor/AnimGraph/Private/AnimGraphNode_MeshRefPose.cpp b/Engine/Source/Editor/AnimGraph/Private/AnimGraphNode_MeshRefPose.cpp
--- a/Engine/Source/Editor/AnimGraph/Private/AnimGraphNode_MeshRefPose.cpp
+++ b/Engine/Source/Editor/AnimGraph/Private/AnimGraphNode_MeshRefPose.cpp
@@ -5,6 +5,8 @@
 /////////////////////////////////////////////////////
 // UAnimGraphNode_MeshRefPose
 
+#define LOCTEXT_NAMESPACE "A3Nodes"
+
 UAnimGraphNode_MeshRefPose::UAnimGraphNode_MeshRefPose(const FPostConstructInitializeProperties& PCIP)
 	: Super(PCIP)
 {
@@ -22,16 +24,18 @@ FLinearColor UAnimGraphNode_MeshRefPose::GetNodeTitleColor() const
 
 FString UAnimGraphNode_MeshRefPose::GetTooltip() const
 {
-	return TEXT("Returns mesh space reference pose.");
+	return LOCTEXT("AnimGraphNode_MeshRefPose_Tooltip", "Returns mesh space reference pose.").ToString();
 }
 
-FString UAnimGraphNode_MeshRefPose::GetNodeTitle(ENodeTitleType::Type TitleType) const
+FText UAnimGraphNode_MeshRefPose::GetNodeTitle(ENodeTitleType::Type TitleType) const
 {
-	return FString::Printf(TEXT("Mesh Space Ref Pose"));
+	return LOCTEXT("AnimGraphNode_MeshRefPose_Title", "Mesh Space Ref Pose");
 }
 
 void UAnimGraphNode_MeshRefPose::CreateOutputPins()
 {
 	const UAnimationGraphSchema* Schema = GetDefault<UAnimationGraphSchema>();
 	CreatePin(EGPD_Output, Schema->PC_Struct, TEXT(""), FComponentSpacePoseLink::StaticStruct(), /*bIsArray=*/ false, /*bIsReference=*/ false, TEXT("ComponentPose"));
-}
\ No newline at end of file
+}
+
+#undef LOCTEXT_NAMESPACE
\ No newline at end of file
diff --git a/Engine/Source/Editor/AnimGraph/Private/AnimGraphNode_ModifyBone.cpp b/Engine/Source/Editor/AnimGraph/Private/AnimGraphNode_ModifyBone.cpp
--- a/Engine/Source/Editor/AnimGraph/Private/AnimGraphNode_ModifyBone.cpp
+++ b/Engine/Source/Editor/AnimGraph/Private/AnimGraphNode_ModifyBone.cpp
@@ -5,14 +5,16 @@
 /////////////////////////////////////////////////////
 // UAnimGraphNode_ModifyBone
 
+#define LOCTEXT_NAMESPACE "A3Nodes"
+
 UAnimGraphNode_ModifyBone::UAnimGraphNode_ModifyBone(const FPostConstructInitializeProperties& PCIP)
 	: Super(PCIP)
 {
 }
 
-FString UAnimGraphNode_ModifyBone::GetControllerDescription() const
+FText UAnimGraphNode_ModifyBone::GetControllerDescription() const
 {
-	return TEXT("Transform (Modify) Bone");
+	return LOCTEXT("TransformModifyBone", "Transform (Modify) Bone");
 }
 
 FString UAnimGraphNode_ModifyBone::GetKeywords() const
@@ -22,13 +24,23 @@ FString UAnimGraphNode_ModifyBone::GetKeywords() const
 
 FString UAnimGraphNode_ModifyBone::GetTooltip() const
 {
-	return TEXT("The Transform Bone node alters the transform - i.e. Translation, Rotation, or Scale - of the bone");
+	return LOCTEXT("AnimGraphNode_ModifyBone_Tooltip", "The Transform Bone node alters the transform - i.e. Translation, Rotation, or Scale - of the bone").ToString();
 }
 
-FString UAnimGraphNode_ModifyBone::GetNodeTitle(ENodeTitleType::Type TitleType) const
+FText UAnimGraphNode_ModifyBone::GetNodeTitle(ENodeTitleType::Type TitleType) const
 {
-	FString Result = *GetControllerDescription();
-	Result += (TitleType == ENodeTitleType::ListView) ? TEXT(" - ") : TEXT("\n");
-	Result += FString::Printf(TEXT("Bone: %s"), *Node.BoneToModify.BoneName.ToString());
-	return Result;
+	FFormatNamedArguments Args;
+	Args.Add(TEXT("ControllerDescription"), GetControllerDescription());
+	Args.Add(TEXT("BoneName"), FText::FromName(Node.BoneToModify.BoneName));
+
+	if(TitleType == ENodeTitleType::ListView)
+	{
+		return FText::Format(LOCTEXT("AnimGraphNode_ModifyBone_Title", "{ControllerDescription} - Bone: {BoneName}"), Args);
+	}
+	else
+	{
+		return FText::Format(LOCTEXT("AnimGraphNode_ModifyBone_Title", "{ControllerDescription}\nBone: {BoneName}"), Args);
+	}
 }
+
+#undef LOCTEXT_NAMESPACE
diff --git a/Engine/Source/Editor/AnimGraph/Private/AnimGraphNode_Root.cpp b/Engine/Source/Editor/AnimGraph/Private/AnimGraphNode_Root.cpp
--- a/Engine/Source/Editor/AnimGraph/Private/AnimGraphNode_Root.cpp
+++ b/Engine/Source/Editor/AnimGraph/Private/AnimGraphNode_Root.cpp
@@ -33,6 +33,8 @@ void FPoseLinkMappingRecord::PatchLinkIndex(uint8* DestinationPtr, int32 LinkID,
 /////////////////////////////////////////////////////
 // UAnimGraphNode_Root
 
+#define LOCTEXT_NAMESPACE "A3Nodes"
+
 UAnimGraphNode_Root::UAnimGraphNode_Root(const FPostConstructInitializeProperties& PCIP)
 	: Super(PCIP)
 {
@@ -43,14 +45,14 @@ FLinearColor UAnimGraphNode_Root::GetNodeTitleColor() const
 	return GEditor->AccessEditorUserSettings().ResultNodeTitleColor;
 }
 
-FString UAnimGraphNode_Root::GetNodeTitle(ENodeTitleType::Type TitleType) const
+FText UAnimGraphNode_Root::GetNodeTitle(ENodeTitleType::Type TitleType) const
 {
-	return TEXT("Final Animation Pose");
+	return LOCTEXT("AnimGraphNodeRoot_Title", "Final Animation Pose");
 }
 
 FString UAnimGraphNode_Root::GetTooltip() const
 {
-	return TEXT("Wire the final animation pose into this node");
+	return LOCTEXT("AnimGraphNodeRoot_Tooltip", "Wire the final animation pose into this node").ToString();
 }
 
 bool UAnimGraphNode_Root::IsSinkNode() const
@@ -67,3 +69,5 @@ FString UAnimGraphNode_Root::GetDocumentationLink() const
 {
 	return TEXT("Shared/GraphNodes/AnimationStateMachine");
 }
+
+#undef LOCTEXT_NAMESPACE
diff --git a/Engine/Source/Editor/AnimGraph/Private/AnimGraphNode_RotateRootBone.cpp b/Engine/Source/Editor/AnimGraph/Private/AnimGraphNode_RotateRootBone.cpp
--- a/Engine/Source/Editor/AnimGraph/Private/AnimGraphNode_RotateRootBone.cpp
+++ b/Engine/Source/Editor/AnimGraph/Private/AnimGraphNode_RotateRootBone.cpp
@@ -8,6 +8,8 @@
 /////////////////////////////////////////////////////
 // UAnimGraphNode_RotateRootBone
 
+#define LOCTEXT_NAMESPACE "A3Nodes"
+
 UAnimGraphNode_RotateRootBone::UAnimGraphNode_RotateRootBone(const FPostConstructInitializeProperties& PCIP)
 	: Super(PCIP)
 {
@@ -20,15 +22,17 @@ FLinearColor UAnimGraphNode_RotateRootBone::GetNodeTitleColor() const
 
 FString UAnimGraphNode_RotateRootBone::GetTooltip() const
 {
-	return TEXT("Rotate Root Bone");
+	return LOCTEXT("RotateRootBone", "Rotate Root Bone").ToString();
 }
 
-FString UAnimGraphNode_RotateRootBone::GetNodeTitle(ENodeTitleType::Type TitleType) const
+FText UAnimGraphNode_RotateRootBone::GetNodeTitle(ENodeTitleType::Type TitleType) const
 {
-	return TEXT("Rotate Root Bone");
+	return LOCTEXT("RotateRootBone", "Rotate Root Bone");
 }
 
 FString UAnimGraphNode_RotateRootBone::GetNodeCategory() const
 {
 	return TEXT("Tools");
 }
+
+#undef LOCTEXT_NAMESPACE
diff --git a/Engine/Source/Editor/AnimGraph/Private/AnimGraphNode_RotationMultiplier.cpp b/Engine/Source/Editor/AnimGraph/Private/AnimGraphNode_RotationMultiplier.cpp
--- a/Engine/Source/Editor/AnimGraph/Private/AnimGraphNode_RotationMultiplier.cpp
+++ b/Engine/Source/Editor/AnimGraph/Private/AnimGraphNode_RotationMultiplier.cpp
@@ -5,25 +5,46 @@
 /////////////////////////////////////////////////////
 // UAnimGraphNode_RotationMultiplier
 
+#define LOCTEXT_NAMESPACE "A3Nodes"
+
 UAnimGraphNode_RotationMultiplier::UAnimGraphNode_RotationMultiplier(const FPostConstructInitializeProperties& PCIP)
 	: Super(PCIP)
 {
 }
 
-FString UAnimGraphNode_RotationMultiplier::GetControllerDescription() const
+FText UAnimGraphNode_RotationMultiplier::GetControllerDescription() const
 {
-	return TEXT("Apply a percentage of Rotation");
+	return LOCTEXT("ApplyPercentageOfRotation", "Apply a percentage of Rotation");
 }
 
 FString UAnimGraphNode_RotationMultiplier::GetTooltip() const
 {
-	return TEXT("The Apply a Percentage of Rotation control drives the Rotation of a target bone at some specified percentage of the Rotation of another bone within the Skeleton.");
+	return LOCTEXT("AnimGraphNode_RotationMultiplier_Tooltip", "The Apply a Percentage of Rotation control drives the Rotation of a target bone at some specified percentage of the Rotation of another bone within the Skeleton.").ToString();
+}
+
+FText UAnimGraphNode_RotationMultiplier::GetNodeTitle(ENodeTitleType::Type TitleType) const
+{
+	FFormatNamedArguments Args;
+	Args.Add(TEXT("ControllerDescription"), GetControllerDescription());
+	Args.Add(TEXT("BoneName"), FText::FromName(Node.TargetBone.BoneName));
+
+	if(TitleType == ENodeTitleType::ListView)
+	{
+		return FText::Format(LOCTEXT("AnimGraphNode_ModifyBone_Title", "{ControllerDescription} - Bone: {BoneName}"), Args);
+	}
+	else
+	{
+		return FText::Format(LOCTEXT("AnimGraphNode_ModifyBone_Title", "{ControllerDescription}\nBone: {BoneName}"), Args);
+	}
 }
 
-FString UAnimGraphNode_RotationMultiplier::GetNodeTitle(ENodeTitleType::Type TitleType) const
+FString UAnimGraphNode_RotationMultiplier::GetNodeNativeTitle(ENodeTitleType::Type TitleType) const
 {
-	FString Result = *GetControllerDescription();
+	// Do not setup this function for localization, intentionally left unlocalized!
+	FString Result = GetControllerDescription().ToString();
 	Result += (TitleType == ENodeTitleType::ListView) ? TEXT(" - ") : TEXT("\n");
 	Result += FString::Printf(TEXT("Bone: %s"), *Node.TargetBone.BoneName.ToString());
 	return Result;
 }
+
+#undef LOCTEXT_NAMESPACE
diff --git a/Engine/Source/Editor/AnimGraph/Private/AnimGraphNode_RotationOffsetBlendSpace.cpp b/Engine/Source/Editor/AnimGraph/Private/AnimGraphNode_RotationOffsetBlendSpace.cpp
--- a/Engine/Source/Editor/AnimGraph/Private/AnimGraphNode_RotationOffsetBlendSpace.cpp
+++ b/Engine/Source/Editor/AnimGraph/Private/AnimGraphNode_RotationOffsetBlendSpace.cpp
@@ -7,6 +7,8 @@
 /////////////////////////////////////////////////////
 // UAnimGraphNode_RotationOffsetBlendSpace
 
+#define LOCTEXT_NAMESPACE "A3Nodes"
+
 UAnimGraphNode_RotationOffsetBlendSpace::UAnimGraphNode_RotationOffsetBlendSpace(const FPostConstructInitializeProperties& PCIP)
 	: Super(PCIP)
 {
@@ -17,8 +19,27 @@ FString UAnimGraphNode_RotationOffsetBlendSpace::GetTooltip() const
 	return FString::Printf(TEXT("AimOffset %s"), *(Node.BlendSpace->GetPathName()));
 }
 
-FString UAnimGraphNode_RotationOffsetBlendSpace::GetNodeTitle(ENodeTitleType::Type TitleType) const
+FText UAnimGraphNode_RotationOffsetBlendSpace::GetNodeTitle(ENodeTitleType::Type TitleType) const
+{
+	const FText BlendSpaceName((Node.BlendSpace != NULL) ? FText::FromString(Node.BlendSpace->GetName()) : LOCTEXT("None", "(None)"));
+
+	FFormatNamedArguments Args;
+	Args.Add(TEXT("BlendSpaceName"), BlendSpaceName);
+
+	if (TitleType == ENodeTitleType::ListView)
+	{
+		return FText::Format(LOCTEXT("AimOffsetListTitle", "AimOffset '{BlendSpaceName}'"), Args);
+	}
+	else
+	{
+		return FText::Format(LOCTEXT("AimOffsetFullTitle", "{BlendSpaceName}\nAimOffset"), Args);
+	}
+}
+
+FString UAnimGraphNode_RotationOffsetBlendSpace::GetNodeNativeTitle(ENodeTitleType::Type TitleType) const
 {
+	// Do not setup this function for localization, intentionally left unlocalized!
+	
 	const FString BlendSpaceName((Node.BlendSpace != NULL) ? *(Node.BlendSpace->GetName()) : TEXT("(None)"));
 
 	if (TitleType == ENodeTitleType::ListView)
@@ -70,4 +91,6 @@ void UAnimGraphNode_RotationOffsetBlendSpace::GetContextMenuActions(const FGraph
 		}
 		Context.MenuBuilder->EndSection();
 	}
-}
\ No newline at end of file
+}
+
+#undef LOCTEXT_NAMESPACE
\ No newline at end of file
diff --git a/Engine/Source/Editor/AnimGraph/Private/AnimGraphNode_SaveCachedPose.cpp b/Engine/Source/Editor/AnimGraph/Private/AnimGraphNode_SaveCachedPose.cpp
--- a/Engine/Source/Editor/AnimGraph/Private/AnimGraphNode_SaveCachedPose.cpp
+++ b/Engine/Source/Editor/AnimGraph/Private/AnimGraphNode_SaveCachedPose.cpp
@@ -30,6 +30,8 @@ public:
 /////////////////////////////////////////////////////
 // UAnimGraphNode_ComponentToLocalSpace
 
+#define LOCTEXT_NAMESPACE "A3Nodes"
+
 UAnimGraphNode_SaveCachedPose::UAnimGraphNode_SaveCachedPose(const FPostConstructInitializeProperties& PCIP)
 	: Super(PCIP)
 {
@@ -42,8 +44,23 @@ FString UAnimGraphNode_SaveCachedPose::GetTooltip() const
 	return TEXT("Denotes an animation tree that can be referenced elsewhere in the blueprint, which will be evaluated at most once per frame and then cached.");
 }
 
-FString UAnimGraphNode_SaveCachedPose::GetNodeTitle(ENodeTitleType::Type TitleType) const
+FText UAnimGraphNode_SaveCachedPose::GetNodeTitle(ENodeTitleType::Type TitleType) const
+{
+	if (TitleType == ENodeTitleType::EditableTitle)
+	{
+		return FText::FromString(CacheName);
+	}
+	else
+	{
+		FFormatNamedArguments Args;
+		Args.Add(TEXT("NodeTitle"), FText::FromString(CacheName));
+		return FText::Format(LOCTEXT("AnimGraphNode_SaveCachedPose_Title", "Save cached pose '{NodeTitle}'"), Args);
+	}
+}
+
+FString UAnimGraphNode_SaveCachedPose::GetNodeNativeTitle(ENodeTitleType::Type TitleType) const
 {
+	// Do not setup this function for localization, intentionally left unlocalized!
 	if (TitleType == ENodeTitleType::EditableTitle)
 	{
 		return CacheName;
@@ -73,8 +90,9 @@ void UAnimGraphNode_SaveCachedPose::GetMenuEntries(FGraphContextMenuBuilder& Con
 			UAnimGraphNode_SaveCachedPose* SaveCachedPose = NewObject<UAnimGraphNode_SaveCachedPose>();
 			SaveCachedPose->CacheName += FString::FromInt(FMath::Rand());
 
-			TSharedPtr<FEdGraphSchemaAction_K2NewNode> SaveCachedPoseAction = FK2ActionMenuBuilder::AddNewNodeAction(ContextMenuBuilder, GetNodeCategory(), TEXT("New Save cached pose..."), SaveCachedPose->GetTooltip(), 0, SaveCachedPose->GetKeywords());
+			TSharedPtr<FEdGraphSchemaAction_K2NewNode> SaveCachedPoseAction = FK2ActionMenuBuilder::AddNewNodeAction(ContextMenuBuilder, GetNodeCategory(), LOCTEXT("NewSaveCachedPose", "New Save cached pose..."), SaveCachedPose->GetTooltip(), 0, SaveCachedPose->GetKeywords());
 			SaveCachedPoseAction->NodeTemplate = SaveCachedPose;
+			SaveCachedPoseAction->SearchTitle = SaveCachedPoseAction->NodeTemplate->GetNodeSearchTitle();
 		}
 	}
 }
@@ -88,4 +106,6 @@ void UAnimGraphNode_SaveCachedPose::OnRenameNode(const FString& NewName)
 TSharedPtr<class INameValidatorInterface> UAnimGraphNode_SaveCachedPose::MakeNameValidator() const
 {
 	return MakeShareable(new FCachedPoseNameValidator(GetBlueprint(), CacheName));
-}
\ No newline at end of file
+}
+
+#undef LOCTEXT_NAMESPACE
\ No newline at end of file
diff --git a/Engine/Source/Editor/AnimGraph/Private/AnimGraphNode_SequenceEvaluator.cpp b/Engine/Source/Editor/AnimGraph/Private/AnimGraphNode_SequenceEvaluator.cpp
--- a/Engine/Source/Editor/AnimGraph/Private/AnimGraphNode_SequenceEvaluator.cpp
+++ b/Engine/Source/Editor/AnimGraph/Private/AnimGraphNode_SequenceEvaluator.cpp
@@ -8,6 +8,8 @@
 /////////////////////////////////////////////////////
 // UAnimGraphNode_SequenceEvaluator
 
+#define LOCTEXT_NAMESPACE "A3Nodes"
+
 UAnimGraphNode_SequenceEvaluator::UAnimGraphNode_SequenceEvaluator(const FPostConstructInitializeProperties& PCIP)
 	: Super(PCIP)
 {
@@ -45,8 +47,26 @@ FString UAnimGraphNode_SequenceEvaluator::GetTooltip() const
 	}
 }
 
-FString UAnimGraphNode_SequenceEvaluator::GetNodeTitle(ENodeTitleType::Type TitleType) const
+FText UAnimGraphNode_SequenceEvaluator::GetNodeTitle(ENodeTitleType::Type TitleType) const
+{
+	const FText SequenceName((Node.Sequence != NULL) ? FText::FromString(Node.Sequence->GetName()) : LOCTEXT("None", "(None)"));
+
+	FFormatNamedArguments Args;
+	Args.Add(TEXT("SequenceName"), SequenceName);
+
+	if (Node.Sequence && Node.Sequence->IsValidAdditive())
+	{
+		return FText::Format(LOCTEXT("EvaluateSequence_Additive", "Evaluate {SequenceName} (additive)"), Args);
+	}
+	else
+	{
+		return FText::Format(LOCTEXT("EvaluateSequence", "Evaluate {SequenceName}"), Args);
+	}
+}
+
+FString UAnimGraphNode_SequenceEvaluator::GetNodeNativeTitle(ENodeTitleType::Type TitleType) const
 {
+	// Do not setup this function for localization, intentionally left unlocalized!
 	if (Node.Sequence && Node.Sequence->IsValidAdditive())
 	{
 		return FString::Printf(TEXT("Evaluate %s (additive)"), (Node.Sequence != NULL) ? *(Node.Sequence->GetName()) : TEXT("(None)"));
@@ -111,4 +131,6 @@ const TCHAR* UAnimGraphNode_SequenceEvaluator::GetTimePropertyName() const
 UScriptStruct* UAnimGraphNode_SequenceEvaluator::GetTimePropertyStruct() const 
 {
 	return FAnimNode_SequenceEvaluator::StaticStruct();
-}
\ No newline at end of file
+}
+
+#undef LOCTEXT_NAMESPACE
\ No newline at end of file
diff --git a/Engine/Source/Editor/AnimGraph/Private/AnimGraphNode_SequencePlayer.cpp b/Engine/Source/Editor/AnimGraph/Private/AnimGraphNode_SequencePlayer.cpp
--- a/Engine/Source/Editor/AnimGraph/Private/AnimGraphNode_SequencePlayer.cpp
+++ b/Engine/Source/Editor/AnimGraph/Private/AnimGraphNode_SequencePlayer.cpp
@@ -17,7 +17,7 @@ struct FNewSequencePlayerAction : public FEdGraphSchemaAction_K2NewNode
 protected:
 	FAssetData AssetInfo;
 public:
-	FNewSequencePlayerAction(const FAssetData& InAssetInfo, FString Title)
+	FNewSequencePlayerAction(const FAssetData& InAssetInfo, FText Title)
 	{
 		AssetInfo = InAssetInfo;
 
@@ -71,32 +71,52 @@ FLinearColor UAnimGraphNode_SequencePlayer::GetNodeTitleColor() const
 FString UAnimGraphNode_SequencePlayer::GetTooltip() const
 {
 	const bool bAdditive = ((Node.Sequence != NULL) && Node.Sequence->IsValidAdditive());
-	return GetTitleGivenAssetInfo(Node.Sequence->GetPathName(), bAdditive);
+	return GetTitleGivenAssetInfo(FText::FromString(Node.Sequence->GetPathName()), bAdditive).ToString();
 }
 
-FString UAnimGraphNode_SequencePlayer::GetNodeTitle(ENodeTitleType::Type TitleType) const
+FText UAnimGraphNode_SequencePlayer::GetNodeTitle(ENodeTitleType::Type TitleType) const
 {
 	const bool bAdditive = ((Node.Sequence != NULL) && Node.Sequence->IsValidAdditive());
-	FString Title = GetTitleGivenAssetInfo((Node.Sequence != NULL) ? Node.Sequence->GetName() : FString(TEXT("(None)")), bAdditive);
+	FText Title = GetTitleGivenAssetInfo((Node.Sequence != NULL) ? FText::FromString(Node.Sequence->GetName()) : LOCTEXT("None", "(None)"), bAdditive);
+
+	if ((TitleType == ENodeTitleType::FullTitle) && (SyncGroup.GroupName != NAME_None))
+	{
+		FFormatNamedArguments Args;
+		Args.Add(TEXT("Title"), Title);
+
+		Title = FText::Format(LOCTEXT("SequenceNodeGroupWithSubtitle", "{Title}\nSync group {Title}"), Args);
+	}
+
+	return Title;
+}
+
+FString UAnimGraphNode_SequencePlayer::GetNodeNativeTitle(ENodeTitleType::Type TitleType) const
+{
+	// Do not setup this function for localization, intentionally left unlocalized!
+	const bool bAdditive = ((Node.Sequence != NULL) && Node.Sequence->IsValidAdditive());
+	FString Title = GetTitleGivenAssetInfo((Node.Sequence != NULL) ? FText::FromString(Node.Sequence->GetName()) : LOCTEXT("None", "(None)"), bAdditive).ToString();
 
 	if ((TitleType == ENodeTitleType::FullTitle) && (SyncGroup.GroupName != NAME_None))
 	{
 		Title += TEXT("\n");
-		Title += FString::Printf(*LOCTEXT("SequenceNodeGroupSubtitle", "Sync group %s").ToString(), *SyncGroup.GroupName.ToString());
+		Title += FString::Printf(TEXT("Sync group %s"), *SyncGroup.GroupName.ToString());
 	}
 
 	return Title;
 }
 
-FString UAnimGraphNode_SequencePlayer::GetTitleGivenAssetInfo(const FString& AssetName, bool bKnownToBeAdditive)
+FText UAnimGraphNode_SequencePlayer::GetTitleGivenAssetInfo(const FText& AssetName, bool bKnownToBeAdditive)
 {
+	FFormatNamedArguments Args;
+	Args.Add(TEXT("AssetName"), AssetName);
+
 	if (bKnownToBeAdditive)
 	{
-		return FString::Printf(*LOCTEXT("SequenceNodeTitleAdditive", "Play %s (additive)").ToString(), *AssetName);
+		return FText::Format(LOCTEXT("SequenceNodeTitleAdditive", "Play {AssetName} (additive)"), Args);
 	}
 	else
 	{
-		return FString::Printf(*LOCTEXT("SequenceNodeTitle", "Play %s").ToString(), *AssetName);
+		return FText::Format(LOCTEXT("SequenceNodeTitle", "Play {AssetName}"), Args);
 	}
 }
 
@@ -138,7 +158,7 @@ void UAnimGraphNode_SequencePlayer::GetMenuEntries(FGraphContextMenuBuilder& Con
 				}
 
 				// Create the menu item
-				const FString Title = UAnimGraphNode_SequencePlayer::GetTitleGivenAssetInfo(Asset.AssetName.ToString(), bAdditive);
+				const FText Title = UAnimGraphNode_SequencePlayer::GetTitleGivenAssetInfo(FText::FromName(Asset.AssetName), bAdditive);
 				TSharedPtr<FNewSequencePlayerAction> NewAction(new FNewSequencePlayerAction(Asset, Title));
 				ContextMenuBuilder.AddAction( NewAction );
 			}
diff --git a/Engine/Source/Editor/AnimGraph/Private/AnimGraphNode_SkeletalControlBase.cpp b/Engine/Source/Editor/AnimGraph/Private/AnimGraphNode_SkeletalControlBase.cpp
--- a/Engine/Source/Editor/AnimGraph/Private/AnimGraphNode_SkeletalControlBase.cpp
+++ b/Engine/Source/Editor/AnimGraph/Private/AnimGraphNode_SkeletalControlBase.cpp
@@ -6,6 +6,8 @@
 /////////////////////////////////////////////////////
 // UAnimGraphNode_SkeletalControlBase
 
+#define LOCTEXT_NAMESPACE "A3Nodes"
+
 UAnimGraphNode_SkeletalControlBase::UAnimGraphNode_SkeletalControlBase(const FPostConstructInitializeProperties& PCIP)
 	: Super(PCIP)
 {
@@ -21,14 +23,14 @@ FString UAnimGraphNode_SkeletalControlBase::GetNodeCategory() const
 	return TEXT("Skeletal Controls");
 }
 
-FString UAnimGraphNode_SkeletalControlBase::GetControllerDescription() const
+FText UAnimGraphNode_SkeletalControlBase::GetControllerDescription() const
 {
-	return TEXT("Implement me");
+	return LOCTEXT("ImplementMe", "Implement me");
 }
 
 FString UAnimGraphNode_SkeletalControlBase::GetTooltip() const
 {
-	return GetControllerDescription();
+	return GetControllerDescription().ToString();
 }
 
 void UAnimGraphNode_SkeletalControlBase::CreateOutputPins()
@@ -42,10 +44,13 @@ void UAnimGraphNode_SkeletalControlBase::GetMenuEntries(FGraphContextMenuBuilder
 	UAnimGraphNode_Base* TemplateNode = NewObject<UAnimGraphNode_Base>(GetTransientPackage(), GetClass());
 
 	FString Category = TemplateNode->GetNodeCategory();
-	FString MenuDesc = GetControllerDescription();
-	FString Tooltip = GetControllerDescription();
+	FText MenuDesc = GetControllerDescription();
+	FString Tooltip = GetControllerDescription().ToString();
 	FString Keywords = TemplateNode->GetKeywords();
 
 	TSharedPtr<FEdGraphSchemaAction_K2NewNode> NodeAction = FK2ActionMenuBuilder::AddNewNodeAction(ContextMenuBuilder, Category, MenuDesc, Tooltip, 0, Keywords);
 	NodeAction->NodeTemplate = TemplateNode;
-}
\ No newline at end of file
+	NodeAction->SearchTitle = NodeAction->NodeTemplate->GetNodeSearchTitle();
+}
+
+#undef LOCTEXT_NAMESPACE
\ No newline at end of file
diff --git a/Engine/Source/Editor/AnimGraph/Private/AnimGraphNode_Slot.cpp b/Engine/Source/Editor/AnimGraph/Private/AnimGraphNode_Slot.cpp
--- a/Engine/Source/Editor/AnimGraph/Private/AnimGraphNode_Slot.cpp
+++ b/Engine/Source/Editor/AnimGraph/Private/AnimGraphNode_Slot.cpp
@@ -8,6 +8,8 @@
 /////////////////////////////////////////////////////
 // UAnimGraphNode_Slot
 
+#define LOCTEXT_NAMESPACE "A3Nodes"
+
 UAnimGraphNode_Slot::UAnimGraphNode_Slot(const FPostConstructInitializeProperties& PCIP)
 	: Super(PCIP)
 {
@@ -23,8 +25,26 @@ FString UAnimGraphNode_Slot::GetTooltip() const
 	return TEXT("Plays animation from code using AnimMontage");
 }
 
-FString UAnimGraphNode_Slot::GetNodeTitle(ENodeTitleType::Type TitleType) const
+FText UAnimGraphNode_Slot::GetNodeTitle(ENodeTitleType::Type TitleType) const
 {
+	const FText SlotName = (Node.SlotName != NAME_None) ? FText::FromName(Node.SlotName) : LOCTEXT("NoSlotName", "(No slot name)");
+
+	FFormatNamedArguments Args;
+	Args.Add(TEXT("SlotName"), SlotName);
+
+	if (TitleType == ENodeTitleType::ListView)
+	{
+		return FText::Format(LOCTEXT("SlotNodeListTitle", "Slot '{SlotName}'"), Args);
+	}
+	else
+	{
+		return FText::Format(LOCTEXT("SlotNodeTitle", "{SlotName}\nSlot"), Args);
+	}
+}
+
+FString UAnimGraphNode_Slot::GetNodeNativeTitle(ENodeTitleType::Type TitleType) const
+{
+	// Do not setup this function for localization, intentionally left unlocalized!
 	const FString SlotName = (Node.SlotName != NAME_None) ? Node.SlotName.ToString() : TEXT("(No slot name)");
 
 	if (TitleType == ENodeTitleType::ListView)
@@ -41,3 +61,5 @@ FString UAnimGraphNode_Slot::GetNodeCategory() const
 {
 	return TEXT("Blends");
 }
+
+#undef LOCTEXT_NAMESPACE
diff --git a/Engine/Source/Editor/AnimGraph/Private/AnimGraphNode_SpringBone.cpp b/Engine/Source/Editor/AnimGraph/Private/AnimGraphNode_SpringBone.cpp
--- a/Engine/Source/Editor/AnimGraph/Private/AnimGraphNode_SpringBone.cpp
+++ b/Engine/Source/Editor/AnimGraph/Private/AnimGraphNode_SpringBone.cpp
@@ -5,25 +5,46 @@
 /////////////////////////////////////////////////////
 // UAnimGraphNode_SpringBone
 
+#define LOCTEXT_NAMESPACE "A3Nodes"
+
 UAnimGraphNode_SpringBone::UAnimGraphNode_SpringBone(const FPostConstructInitializeProperties& PCIP)
 	: Super(PCIP)
 {
 }
 
-FString UAnimGraphNode_SpringBone::GetControllerDescription() const
+FText UAnimGraphNode_SpringBone::GetControllerDescription() const
 {
-	return TEXT("Spring controller");
+	return LOCTEXT("SpringController", "Spring controller");
 }
 
 FString UAnimGraphNode_SpringBone::GetTooltip() const
 {
-	return TEXT("The Spring Controller applies a spring solver that can be used to limit how far a bone can stretch from its reference pose position and apply a force in the opposite direction.");
+	return LOCTEXT("AnimGraphNode_SpringBone_Tooltip", "The Spring Controller applies a spring solver that can be used to limit how far a bone can stretch from its reference pose position and apply a force in the opposite direction.").ToString();
+}
+
+FText UAnimGraphNode_SpringBone::GetNodeTitle(ENodeTitleType::Type TitleType) const
+{
+	FFormatNamedArguments Args;
+	Args.Add(TEXT("ControllerDescription"), GetControllerDescription());
+	Args.Add(TEXT("BoneName"), FText::FromName(Node.SpringBone.BoneName));
+
+	if(TitleType == ENodeTitleType::ListView)
+	{
+		return FText::Format(LOCTEXT("AnimGraphNode_SpringBone_Title", "{ControllerDescription} - Bone: {BoneName}"), Args);
+	}
+	else
+	{
+		return FText::Format(LOCTEXT("AnimGraphNode_SpringBone_Title", "{ControllerDescription}\nBone: {BoneName}"), Args);
+	}
 }
 
-FString UAnimGraphNode_SpringBone::GetNodeTitle(ENodeTitleType::Type TitleType) const
+FString UAnimGraphNode_SpringBone::GetNodeNativeTitle(ENodeTitleType::Type TitleType) const
 {
-	FString Result = *GetControllerDescription();
+	// Do not setup this function for localization, intentionally left unlocalized!
+	FString Result = GetControllerDescription().ToString();
 	Result += (TitleType == ENodeTitleType::ListView) ? TEXT(" - ") : TEXT("\n");
 	Result += FString::Printf(TEXT("Bone: %s"), *Node.SpringBone.BoneName.ToString());
 	return Result;
 }
+
+#undef LOCTEXT_NAMESPACE
\ No newline at end of file
diff --git a/Engine/Source/Editor/AnimGraph/Private/AnimGraphNode_StateMachineBase.cpp b/Engine/Source/Editor/AnimGraph/Private/AnimGraphNode_StateMachineBase.cpp
--- a/Engine/Source/Editor/AnimGraph/Private/AnimGraphNode_StateMachineBase.cpp
+++ b/Engine/Source/Editor/AnimGraph/Private/AnimGraphNode_StateMachineBase.cpp
@@ -35,6 +35,8 @@ public:
 /////////////////////////////////////////////////////
 // UAnimGraphNode_StateMachineBase
 
+#define LOCTEXT_NAMESPACE "A3Nodes"
+
 UAnimGraphNode_StateMachineBase::UAnimGraphNode_StateMachineBase(const FPostConstructInitializeProperties& PCIP)
 	: Super(PCIP)
 {
@@ -50,8 +52,23 @@ FString UAnimGraphNode_StateMachineBase::GetTooltip() const
 	return TEXT("Animation State Machine");
 }
 
-FString UAnimGraphNode_StateMachineBase::GetNodeTitle(ENodeTitleType::Type TitleType) const
+FText UAnimGraphNode_StateMachineBase::GetNodeTitle(ENodeTitleType::Type TitleType) const
+{
+	const FText FirstLine = (EditorStateMachineGraph != NULL) ? FText::FromString(EditorStateMachineGraph->GetName()) : LOCTEXT("ErrorNoGraph", "Error: No Graph");
+	if(TitleType == ENodeTitleType::FullTitle)
+	{
+		FFormatNamedArguments Args;
+		Args.Add(TEXT("Title"), FirstLine);
+
+		return FText::Format(LOCTEXT("StateMachineFullTitle", "{Title}\nState Machine"), Args);
+	}
+
+	return FirstLine;
+}
+
+FString UAnimGraphNode_StateMachineBase::GetNodeNativeTitle(ENodeTitleType::Type TitleType) const
 {
+	// Do not setup this function for localization, intentionally left unlocalized!
 	FString FirstLine = (EditorStateMachineGraph != NULL) ? EditorStateMachineGraph->GetName() : TEXT("Error: No Graph");
 	if(TitleType == ENodeTitleType::FullTitle)
 	{
@@ -72,8 +89,8 @@ void UAnimGraphNode_StateMachineBase::GetMenuEntries(FGraphContextMenuBuilder& C
 	{
 		TSharedPtr<FEdGraphSchemaAction_K2NewNode> MenuEntry = CreateDefaultMenuEntry(ContextMenuBuilder);
 
-		MenuEntry->MenuDescription = TEXT("Add New State Machine...");
-		MenuEntry->TooltipDescription = TEXT("Create a new state machine");
+		MenuEntry->MenuDescription = LOCTEXT("AddNewStateMachine", "Add New State Machine...");
+		MenuEntry->TooltipDescription = LOCTEXT("AddNewStateMachine_Tooltip", "Create a new state machine").ToString();
 	}
 }
 
@@ -154,3 +171,4 @@ void UAnimGraphNode_StateMachineBase::OnRenameNode(const FString& NewName)
 	FBlueprintEditorUtils::RenameGraph(EditorStateMachineGraph, NewName);
 }
 
+#undef LOCTEXT_NAMESPACE
diff --git a/Engine/Source/Editor/AnimGraph/Private/AnimGraphNode_TransitionPoseEvaluator.cpp b/Engine/Source/Editor/AnimGraph/Private/AnimGraphNode_TransitionPoseEvaluator.cpp
--- a/Engine/Source/Editor/AnimGraph/Private/AnimGraphNode_TransitionPoseEvaluator.cpp
+++ b/Engine/Source/Editor/AnimGraph/Private/AnimGraphNode_TransitionPoseEvaluator.cpp
@@ -8,6 +8,8 @@
 /////////////////////////////////////////////////////
 // UAnimGraphNode_TransitionPoseEvaluator
 
+#define LOCTEXT_NAMESPACE "UAnimGraphNode_TransitionPoseEvaluator"
+
 UAnimGraphNode_TransitionPoseEvaluator::UAnimGraphNode_TransitionPoseEvaluator(const FPostConstructInitializeProperties& PCIP)
 	: Super(PCIP)
 {
@@ -23,14 +25,14 @@ FString UAnimGraphNode_TransitionPoseEvaluator::GetTooltip() const
 	return (Node.DataSource == EEvaluatorDataSource::EDS_DestinationPose) ? TEXT("Evaluates and returns the pose generated by the destination state of this transition") : TEXT("Evaluates and returns the pose generated by the setup prior to this transition firing");
 }
 
-FString UAnimGraphNode_TransitionPoseEvaluator::GetNodeTitle(ENodeTitleType::Type TitleType) const
+FText UAnimGraphNode_TransitionPoseEvaluator::GetNodeTitle(ENodeTitleType::Type TitleType) const
 {
 	if (TitleType == ENodeTitleType::FullTitle)
 	{
-		return (Node.DataSource == EEvaluatorDataSource::EDS_DestinationPose) ? TEXT("Get Destination State Pose") : TEXT("Get Source State Pose");
+		return (Node.DataSource == EEvaluatorDataSource::EDS_DestinationPose) ? LOCTEXT("GetDestinationStatePose", "Get Destination State Pose") : LOCTEXT("GetSourceStatePose", "Get Source State Pose");
 	}
 
-	return TEXT("Input Pose Evaluator");
+	return LOCTEXT("InputPoseEvaluator", "Input Pose Evaluator");
 }
 
 void UAnimGraphNode_TransitionPoseEvaluator::ValidateAnimNodeDuringCompilation(class USkeleton* ForSkeleton, class FCompilerResultsLog& MessageLog)
@@ -58,3 +60,5 @@ void UAnimGraphNode_TransitionPoseEvaluator::GetMenuEntries(FGraphContextMenuBui
 {
 	// Intentionally empty: Don't allow an option to create them, as they're auto-created when custom blend graphs are made
 }
+
+#undef LOCTEXT_NAMESPACE
diff --git a/Engine/Source/Editor/AnimGraph/Private/AnimGraphNode_TransitionResult.cpp b/Engine/Source/Editor/AnimGraph/Private/AnimGraphNode_TransitionResult.cpp
--- a/Engine/Source/Editor/AnimGraph/Private/AnimGraphNode_TransitionResult.cpp
+++ b/Engine/Source/Editor/AnimGraph/Private/AnimGraphNode_TransitionResult.cpp
@@ -8,6 +8,8 @@
 /////////////////////////////////////////////////////
 // UAnimGraphNode_TransitionResult
 
+#define LOCTEXT_NAMESPACE "UAnimGraphNode_TransitionResult"
+
 UAnimGraphNode_TransitionResult::UAnimGraphNode_TransitionResult(const FPostConstructInitializeProperties& PCIP)
 	: Super(PCIP)
 {
@@ -23,12 +25,14 @@ FString UAnimGraphNode_TransitionResult::GetTooltip() const
 	return TEXT("This expression is evaluated to determine if the state transition can be taken");
 }
 
-FString UAnimGraphNode_TransitionResult::GetNodeTitle(ENodeTitleType::Type TitleType) const
+FText UAnimGraphNode_TransitionResult::GetNodeTitle(ENodeTitleType::Type TitleType) const
 {
-	return TEXT("Result");
+	return LOCTEXT("Result", "Result");
 }
 
 void UAnimGraphNode_TransitionResult::GetMenuEntries(FGraphContextMenuBuilder& ContextMenuBuilder) const
 {
 	// Intentionally empty.  This node is autogenerated when a transition graph is created.
 }
+
+#undef LOCTEXT_NAMESPACE
diff --git a/Engine/Source/Editor/AnimGraph/Private/AnimGraphNode_TwoBoneIK.cpp b/Engine/Source/Editor/AnimGraph/Private/AnimGraphNode_TwoBoneIK.cpp
--- a/Engine/Source/Editor/AnimGraph/Private/AnimGraphNode_TwoBoneIK.cpp
+++ b/Engine/Source/Editor/AnimGraph/Private/AnimGraphNode_TwoBoneIK.cpp
@@ -5,24 +5,43 @@
 /////////////////////////////////////////////////////
 // UAnimGraphNode_TwoBoneIK
 
+#define LOCTEXT_NAMESPACE "A3Nodes"
+
 UAnimGraphNode_TwoBoneIK::UAnimGraphNode_TwoBoneIK(const FPostConstructInitializeProperties& PCIP)
 	: Super(PCIP)
 {
 }
 
-FString UAnimGraphNode_TwoBoneIK::GetControllerDescription() const
+FText UAnimGraphNode_TwoBoneIK::GetControllerDescription() const
 {
-	return TEXT("Two Bone IK");
+	return LOCTEXT("TwoBoneIK", "Two Bone IK");
 }
 
 FString UAnimGraphNode_TwoBoneIK::GetTooltip() const
 {
-	return TEXT("The Two Bone IK control applies an inverse kinematic (IK) solver to a 3-joint chain, such as the limbs of a character.");
+	return LOCTEXT("AnimGraphNode_TwoBoneIK_Tooltip", "The Two Bone IK control applies an inverse kinematic (IK) solver to a 3-joint chain, such as the limbs of a character.").ToString();
 }
 
-FString UAnimGraphNode_TwoBoneIK::GetNodeTitle(ENodeTitleType::Type TitleType) const
+FText UAnimGraphNode_TwoBoneIK::GetNodeTitle(ENodeTitleType::Type TitleType) const
 {
-	FString Result = *GetControllerDescription();
+	FFormatNamedArguments Args;
+	Args.Add(TEXT("ControllerDescription"), GetControllerDescription());
+	Args.Add(TEXT("BoneName"), FText::FromName(Node.IKBone.BoneName));
+
+	if(TitleType == ENodeTitleType::ListView)
+	{
+		return FText::Format(LOCTEXT("AnimGraphNode_IKBone_Title", "{ControllerDescription} - Bone: {BoneName}"), Args);
+	}
+	else
+	{
+		return FText::Format(LOCTEXT("AnimGraphNode_IKBone_Title", "{ControllerDescription}\nBone: {BoneName}"), Args);
+	}
+}
+
+FString UAnimGraphNode_TwoBoneIK::GetNodeNativeTitle(ENodeTitleType::Type TitleType) const
+{
+	// Do not setup this function for localization, intentionally left unlocalized!
+	FString Result = GetControllerDescription().ToString();
 	Result += (TitleType == ENodeTitleType::ListView) ? TEXT(" - ") : TEXT("\n");
 	Result += FString::Printf(TEXT("Bone: %s"), *Node.IKBone.BoneName.ToString());
 	return Result;
@@ -128,3 +147,6 @@ void UAnimGraphNode_TwoBoneIK::ConvertToComponentSpaceTransform(USkeletalMeshCom
 		break;
 	}
 }
+
+
+#undef LOCTEXT_NAMESPACE
diff --git a/Engine/Source/Editor/AnimGraph/Private/AnimGraphNode_TwoWayBlend.cpp b/Engine/Source/Editor/AnimGraph/Private/AnimGraphNode_TwoWayBlend.cpp
--- a/Engine/Source/Editor/AnimGraph/Private/AnimGraphNode_TwoWayBlend.cpp
+++ b/Engine/Source/Editor/AnimGraph/Private/AnimGraphNode_TwoWayBlend.cpp
@@ -5,6 +5,8 @@
 /////////////////////////////////////////////////////
 // UAnimGraphNode_TwoWayBlend
 
+#define LOCTEXT_NAMESPACE "AnimGraphNode_TwoWayBlend"
+
 UAnimGraphNode_TwoWayBlend::UAnimGraphNode_TwoWayBlend(const FPostConstructInitializeProperties& PCIP)
 	: Super(PCIP)
 {
@@ -25,7 +27,10 @@ FString UAnimGraphNode_TwoWayBlend::GetTooltip() const
 	return TEXT("Blend two poses together");
 }
 
-FString UAnimGraphNode_TwoWayBlend::GetNodeTitle(ENodeTitleType::Type TitleType) const
+FText UAnimGraphNode_TwoWayBlend::GetNodeTitle(ENodeTitleType::Type TitleType) const
 {
-	return TEXT("Blend");
+	return LOCTEXT("Blend", "Blend");
 }
+
+#undef LOCTEXT_NAMESPACE
+
diff --git a/Engine/Source/Editor/AnimGraph/Private/AnimGraphNode_UseCachedPose.cpp b/Engine/Source/Editor/AnimGraph/Private/AnimGraphNode_UseCachedPose.cpp
--- a/Engine/Source/Editor/AnimGraph/Private/AnimGraphNode_UseCachedPose.cpp
+++ b/Engine/Source/Editor/AnimGraph/Private/AnimGraphNode_UseCachedPose.cpp
@@ -9,6 +9,8 @@
 /////////////////////////////////////////////////////
 // UAnimGraphNode_UseCachedPose
 
+#define LOCTEXT_NAMESPACE "A3Nodes"
+
 UAnimGraphNode_UseCachedPose::UAnimGraphNode_UseCachedPose(const FPostConstructInitializeProperties& PCIP)
 	: Super(PCIP)
 {
@@ -16,11 +18,19 @@ UAnimGraphNode_UseCachedPose::UAnimGraphNode_UseCachedPose(const FPostConstructI
 
 FString UAnimGraphNode_UseCachedPose::GetTooltip() const
 {
-	return TEXT("References an animation tree elsewhere in the blueprint, which will be evaluated at most once per frame.");
+	return LOCTEXT("AnimGraphNode_UseCachedPose_Tooltip", "References an animation tree elsewhere in the blueprint, which will be evaluated at most once per frame.").ToString();
+}
+
+FText UAnimGraphNode_UseCachedPose::GetNodeTitle(ENodeTitleType::Type TitleType) const
+{
+	FFormatNamedArguments Args;
+	Args.Add(TEXT("CachePoseName"), FText::FromString(NameOfCache));
+	return FText::Format(LOCTEXT("AnimGraphNode_UseCachedPose_Title", "Use cached pose '{CachePoseName}'"), Args);
 }
 
-FString UAnimGraphNode_UseCachedPose::GetNodeTitle(ENodeTitleType::Type TitleType) const
+FString UAnimGraphNode_UseCachedPose::GetNodeNativeTitle(ENodeTitleType::Type TitleType) const
 {
+	// Do not setup this function for localization, intentionally left unlocalized!
 	return FString::Printf(TEXT("Use cached pose '%s'"), *NameOfCache);
 }
 
@@ -46,6 +56,9 @@ void UAnimGraphNode_UseCachedPose::GetMenuEntries(FGraphContextMenuBuilder& Cont
 
 			TSharedPtr<FEdGraphSchemaAction_K2NewNode> UseCachedPoseAction = FK2ActionMenuBuilder::AddNewNodeAction(ContextMenuBuilder, GetNodeCategory(), UseCachedPose->GetNodeTitle(ENodeTitleType::ListView), UseCachedPose->GetTooltip(), 0, UseCachedPose->GetKeywords());
 			UseCachedPoseAction->NodeTemplate = UseCachedPose;
+			UseCachedPoseAction->SearchTitle = UseCachedPose->GetNodeSearchTitle();
 		}
 	}
-}
\ No newline at end of file
+}
+
+#undef LOCTEXT_NAMESPACE
\ No newline at end of file
diff --git a/Engine/Source/Editor/AnimGraph/Private/AnimStateConduitNode.cpp b/Engine/Source/Editor/AnimGraph/Private/AnimStateConduitNode.cpp
--- a/Engine/Source/Editor/AnimGraph/Private/AnimStateConduitNode.cpp
+++ b/Engine/Source/Editor/AnimGraph/Private/AnimStateConduitNode.cpp
@@ -33,9 +33,9 @@ void UAnimStateConduitNode::AutowireNewNode(UEdGraphPin* FromPin)
 	}
 }
 
-FString UAnimStateConduitNode::GetNodeTitle(ENodeTitleType::Type TitleType) const
+FText UAnimStateConduitNode::GetNodeTitle(ENodeTitleType::Type TitleType) const
 {
-	return GetStateName();
+	return FText::FromString(GetStateName());
 }
 
 FString UAnimStateConduitNode::GetTooltip() const
diff --git a/Engine/Source/Editor/AnimGraph/Private/AnimStateEntryNode.cpp b/Engine/Source/Editor/AnimGraph/Private/AnimStateEntryNode.cpp
--- a/Engine/Source/Editor/AnimGraph/Private/AnimStateEntryNode.cpp
+++ b/Engine/Source/Editor/AnimGraph/Private/AnimStateEntryNode.cpp
@@ -22,10 +22,10 @@ void UAnimStateEntryNode::AllocateDefaultPins()
 	UEdGraphPin* Outputs = CreatePin(EGPD_Output, Schema->PC_Exec, TEXT(""), NULL, false, false, TEXT("Entry"));
 }
 
-FString UAnimStateEntryNode::GetNodeTitle(ENodeTitleType::Type TitleType) const
+FText UAnimStateEntryNode::GetNodeTitle(ENodeTitleType::Type TitleType) const
 {
 	UEdGraph* Graph = GetGraph();
-	return Graph->GetName();
+	return FText::FromString(Graph->GetName());
 }
 
 FString UAnimStateEntryNode::GetTooltip() const
diff --git a/Engine/Source/Editor/AnimGraph/Private/AnimStateNode.cpp b/Engine/Source/Editor/AnimGraph/Private/AnimStateNode.cpp
--- a/Engine/Source/Editor/AnimGraph/Private/AnimStateNode.cpp
+++ b/Engine/Source/Editor/AnimGraph/Private/AnimStateNode.cpp
@@ -38,9 +38,9 @@ void UAnimStateNode::AutowireNewNode(UEdGraphPin* FromPin)
 	}
 }
 
-FString UAnimStateNode::GetNodeTitle(ENodeTitleType::Type TitleType) const
+FText UAnimStateNode::GetNodeTitle(ENodeTitleType::Type TitleType) const
 {
-	return GetStateName();
+	return FText::FromString(GetStateName());
 }
 
 FString UAnimStateNode::GetTooltip() const
diff --git a/Engine/Source/Editor/AnimGraph/Private/AnimStateTransitionNode.cpp b/Engine/Source/Editor/AnimGraph/Private/AnimStateTransitionNode.cpp
--- a/Engine/Source/Editor/AnimGraph/Private/AnimStateTransitionNode.cpp
+++ b/Engine/Source/Editor/AnimGraph/Private/AnimStateTransitionNode.cpp
@@ -12,6 +12,8 @@
 //////////////////////////////////////////////////////////////////////////
 // IAnimStateTransitionNodeSharedDataHelper
 
+#define LOCTEXT_NAMESPACE "A3Nodes"
+
 class ANIMGRAPH_API IAnimStateTransitionNodeSharedDataHelper
 {
 public:
@@ -122,8 +124,34 @@ void UAnimStateTransitionNode::PostPasteNode()
 	Super::PostPasteNode();
 }
 
-FString UAnimStateTransitionNode::GetNodeTitle(ENodeTitleType::Type TitleType) const
+FText UAnimStateTransitionNode::GetNodeTitle(ENodeTitleType::Type TitleType) const
+{
+	UAnimStateNodeBase* PrevState = GetPreviousState();
+	UAnimStateNodeBase* NextState = GetNextState();
+
+	if (!SharedRulesName.IsEmpty())
+	{
+		return FText::FromString(SharedRulesName);
+	}
+	else if ((PrevState != NULL) && (NextState != NULL))
+	{
+		FFormatNamedArguments Args;
+		Args.Add(TEXT("PrevState"), FText::FromString(PrevState->GetStateName()));
+		Args.Add(TEXT("NextState"), FText::FromString(NextState->GetStateName()));
+
+		return FText::Format(LOCTEXT("PrevStateToNewState", "{PrevState} to {NextState}"), Args);
+	}
+	else
+	{
+		FFormatNamedArguments Args;
+		Args.Add(TEXT("BoundGraph"), (BoundGraph != NULL) ? FText::FromString(BoundGraph->GetName()) : LOCTEXT("Null", "(null)") );
+		return FText::Format(LOCTEXT("TransitioNState", "Trans {BoundGraph}}"), Args);
+	}
+}
+
+FString UAnimStateTransitionNode::GetNodeNativeTitle(ENodeTitleType::Type TitleType) const
 {
+	// Do not setup this function for localization, intentionally left unlocalized!
 	UAnimStateNodeBase* PrevState = GetPreviousState();
 	UAnimStateNodeBase* NextState = GetNextState();
 
@@ -577,3 +605,5 @@ FGuid& FAnimStateTransitionNodeSharedCrossfadeHelper::AccessShareDataGuid(UAnimS
 {
 	return Node->SharedCrossfadeGuid;
 }
+
+#undef LOCTEXT_NAMESPACE
diff --git a/Engine/Source/Editor/AnimGraph/Private/AnimationConduitGraphSchema.cpp b/Engine/Source/Editor/AnimGraph/Private/AnimationConduitGraphSchema.cpp
--- a/Engine/Source/Editor/AnimGraph/Private/AnimationConduitGraphSchema.cpp
+++ b/Engine/Source/Editor/AnimGraph/Private/AnimationConduitGraphSchema.cpp
@@ -23,10 +23,15 @@ void UAnimationConduitGraphSchema::CreateDefaultNodesForGraph(UEdGraph& Graph) c
 
 void UAnimationConduitGraphSchema::GetGraphDisplayInformation(const UEdGraph& Graph, /*out*/ FGraphDisplayInfo& DisplayInfo) const
 {
-	DisplayInfo.DisplayName = FText::FromString( Graph.GetName() );
+	DisplayInfo.PlainName = FText::FromString( Graph.GetName() );
 	
 	if (const UAnimStateConduitNode* ConduitNode = Cast<const UAnimStateConduitNode>(Graph.GetOuter()))
 	{
-		DisplayInfo.DisplayName = FText::Format( NSLOCTEXT("Animation", "ConduitRuleGraphTitle", "{0} (conduit rule)"), FText::FromString( ConduitNode->GetNodeTitle(ENodeTitleType::FullTitle) ) );
+		FFormatNamedArguments Args;
+		Args.Add(TEXT("NodeTitle"), ConduitNode->GetNodeTitle(ENodeTitleType::FullTitle) );
+
+		DisplayInfo.PlainName = FText::Format( NSLOCTEXT("Animation", "ConduitRuleGraphTitle", "{NodeTitle} (conduit rule)"), Args);
 	}
+
+	DisplayInfo.DisplayName = DisplayInfo.PlainName;
 }
