Fix for build break caused by CL#2055369 (RootEditor_OnMac error).

diff --git a/Engine/Source/Editor/AnimGraph/Private/AnimGraphNode_Base.cpp b/Engine/Source/Editor/AnimGraph/Private/AnimGraphNode_Base.cpp
--- a/Engine/Source/Editor/AnimGraph/Private/AnimGraphNode_Base.cpp
+++ b/Engine/Source/Editor/AnimGraph/Private/AnimGraphNode_Base.cpp
@@ -2,6 +2,7 @@
 
 
 #include "AnimGraphPrivatePCH.h"
+#include "AnimGraphNode_Base.h"
 #include "K2ActionMenuBuilder.h" // for FK2ActionMenuBuilder::AddNewNodeAction()
 
 /////////////////////////////////////////////////////
