Fix for build break caused by CL#2055369.

diff --git a/Engine/Source/Editor/AnimGraph/Classes/AnimGraphNode_ApplyAdditive.h b/Engine/Source/Editor/AnimGraph/Classes/AnimGraphNode_ApplyAdditive.h
--- a/Engine/Source/Editor/AnimGraph/Classes/AnimGraphNode_ApplyAdditive.h
+++ b/Engine/Source/Editor/AnimGraph/Classes/AnimGraphNode_ApplyAdditive.h
@@ -1,6 +1,7 @@
 // Copyright 1998-2014 Epic Games, Inc. All Rights Reserved.
 
 #pragma once
+#include "AnimGraphNode_Base.h"
 #include "AnimGraphNode_ApplyAdditive.generated.h"
 
 UCLASS(MinimalAPI)
diff --git a/Engine/Source/Editor/AnimGraph/Classes/AnimGraphNode_BlendListBase.h b/Engine/Source/Editor/AnimGraph/Classes/AnimGraphNode_BlendListBase.h
--- a/Engine/Source/Editor/AnimGraph/Classes/AnimGraphNode_BlendListBase.h
+++ b/Engine/Source/Editor/AnimGraph/Classes/AnimGraphNode_BlendListBase.h
@@ -1,6 +1,7 @@
 // Copyright 1998-2014 Epic Games, Inc. All Rights Reserved.
 
 #pragma once
+#include "AnimGraphNode_Base.h"
 #include "AnimGraphNode_BlendListBase.generated.h"
 
 UCLASS(MinimalAPI, Abstract)
diff --git a/Engine/Source/Editor/AnimGraph/Classes/AnimGraphNode_BlendListByBool.h b/Engine/Source/Editor/AnimGraph/Classes/AnimGraphNode_BlendListByBool.h
--- a/Engine/Source/Editor/AnimGraph/Classes/AnimGraphNode_BlendListByBool.h
+++ b/Engine/Source/Editor/AnimGraph/Classes/AnimGraphNode_BlendListByBool.h
@@ -1,6 +1,7 @@
 // Copyright 1998-2014 Epic Games, Inc. All Rights Reserved.
 
 #pragma once
+#include "AnimGraphNode_BlendListBase.h"
 #include "AnimGraphNode_BlendListByBool.generated.h"
 
 UCLASS(MinimalAPI)
diff --git a/Engine/Source/Editor/AnimGraph/Classes/AnimGraphNode_BlendListByEnum.h b/Engine/Source/Editor/AnimGraph/Classes/AnimGraphNode_BlendListByEnum.h
--- a/Engine/Source/Editor/AnimGraph/Classes/AnimGraphNode_BlendListByEnum.h
+++ b/Engine/Source/Editor/AnimGraph/Classes/AnimGraphNode_BlendListByEnum.h
@@ -1,6 +1,7 @@
 // Copyright 1998-2014 Epic Games, Inc. All Rights Reserved.
 
 #pragma once
+#include "AnimGraphNode_BlendListBase.h"
 #include "AnimGraphNode_BlendListByEnum.generated.h"
 
 UCLASS(MinimalAPI)
diff --git a/Engine/Source/Editor/AnimGraph/Classes/AnimGraphNode_ComponentToLocalSpace.h b/Engine/Source/Editor/AnimGraph/Classes/AnimGraphNode_ComponentToLocalSpace.h
--- a/Engine/Source/Editor/AnimGraph/Classes/AnimGraphNode_ComponentToLocalSpace.h
+++ b/Engine/Source/Editor/AnimGraph/Classes/AnimGraphNode_ComponentToLocalSpace.h
@@ -1,6 +1,7 @@
 // Copyright 1998-2014 Epic Games, Inc. All Rights Reserved.
 
 #pragma once
+#include "AnimGraphNode_Base.h"
 #include "AnimGraphNode_ComponentToLocalSpace.generated.h"
 
 UCLASS(MinimalAPI)
diff --git a/Engine/Source/Editor/AnimGraph/Classes/AnimGraphNode_LayeredBoneBlend.h b/Engine/Source/Editor/AnimGraph/Classes/AnimGraphNode_LayeredBoneBlend.h
--- a/Engine/Source/Editor/AnimGraph/Classes/AnimGraphNode_LayeredBoneBlend.h
+++ b/Engine/Source/Editor/AnimGraph/Classes/AnimGraphNode_LayeredBoneBlend.h
@@ -1,6 +1,7 @@
 // Copyright 1998-2014 Epic Games, Inc. All Rights Reserved.
 
 #pragma once
+#include "AnimGraphNode_BlendListBase.h"
 #include "AnimGraphNode_LayeredBoneBlend.generated.h"
 
 UCLASS(MinimalAPI)
diff --git a/Engine/Source/Editor/AnimGraph/Classes/AnimGraphNode_LocalToComponentSpace.h b/Engine/Source/Editor/AnimGraph/Classes/AnimGraphNode_LocalToComponentSpace.h
--- a/Engine/Source/Editor/AnimGraph/Classes/AnimGraphNode_LocalToComponentSpace.h
+++ b/Engine/Source/Editor/AnimGraph/Classes/AnimGraphNode_LocalToComponentSpace.h
@@ -1,6 +1,7 @@
 // Copyright 1998-2014 Epic Games, Inc. All Rights Reserved.
 
 #pragma once
+#include "AnimGraphNode_Base.h"
 #include "AnimGraphNode_LocalToComponentSpace.generated.h"
 
 UCLASS(MinimalAPI)
diff --git a/Engine/Source/Editor/AnimGraph/Classes/AnimGraphNode_RefPoseBase.h b/Engine/Source/Editor/AnimGraph/Classes/AnimGraphNode_RefPoseBase.h
--- a/Engine/Source/Editor/AnimGraph/Classes/AnimGraphNode_RefPoseBase.h
+++ b/Engine/Source/Editor/AnimGraph/Classes/AnimGraphNode_RefPoseBase.h
@@ -1,6 +1,7 @@
 // Copyright 1998-2014 Epic Games, Inc. All Rights Reserved.
 
 #pragma once
+#include "AnimGraphNode_Base.h"
 #include "AnimGraphNode_RefPoseBase.generated.h"
 
 UCLASS(MinimalAPI, abstract)
diff --git a/Engine/Source/Editor/AnimGraph/Classes/AnimGraphNode_SequencePlayer.h b/Engine/Source/Editor/AnimGraph/Classes/AnimGraphNode_SequencePlayer.h
--- a/Engine/Source/Editor/AnimGraph/Classes/AnimGraphNode_SequencePlayer.h
+++ b/Engine/Source/Editor/AnimGraph/Classes/AnimGraphNode_SequencePlayer.h
@@ -1,6 +1,7 @@
 // Copyright 1998-2014 Epic Games, Inc. All Rights Reserved.
 
 #pragma once
+#include "AnimGraphNode_Base.h"
 #include "AnimGraphNode_SequencePlayer.generated.h"
 
 UCLASS(MinimalAPI)
diff --git a/Engine/Source/Editor/AnimGraph/Classes/AnimStateEntryNode.h b/Engine/Source/Editor/AnimGraph/Classes/AnimStateEntryNode.h
--- a/Engine/Source/Editor/AnimGraph/Classes/AnimStateEntryNode.h
+++ b/Engine/Source/Editor/AnimGraph/Classes/AnimStateEntryNode.h
@@ -1,10 +1,9 @@
 // Copyright 1998-2014 Epic Games, Inc. All Rights Reserved.
 
 #pragma once
+#include "EdGraph/EdGraphNode.h"
 #include "AnimStateEntryNode.generated.h"
 
-
-
 UCLASS(MinimalAPI)
 class UAnimStateEntryNode : public UEdGraphNode
 {
diff --git a/Engine/Source/Editor/AnimGraph/Classes/AnimationGraphSchema.h b/Engine/Source/Editor/AnimGraph/Classes/AnimationGraphSchema.h
--- a/Engine/Source/Editor/AnimGraph/Classes/AnimationGraphSchema.h
+++ b/Engine/Source/Editor/AnimGraph/Classes/AnimationGraphSchema.h
@@ -2,6 +2,7 @@
 
 
 #pragma once
+#include "EdGraphSchema_K2.h"
 #include "AnimationGraphSchema.generated.h"
 
 UCLASS(MinimalAPI)
diff --git a/Engine/Source/Editor/AnimGraph/Classes/AnimationStateGraphSchema.h b/Engine/Source/Editor/AnimGraph/Classes/AnimationStateGraphSchema.h
--- a/Engine/Source/Editor/AnimGraph/Classes/AnimationStateGraphSchema.h
+++ b/Engine/Source/Editor/AnimGraph/Classes/AnimationStateGraphSchema.h
@@ -1,7 +1,7 @@
 // Copyright 1998-2014 Epic Games, Inc. All Rights Reserved.
 
-
 #pragma once
+#include "AnimationGraphSchema.h"
 #include "AnimationStateGraphSchema.generated.h"
 
 UCLASS(MinimalAPI)
diff --git a/Engine/Source/Editor/AnimGraph/Private/AnimGraphNode_Base.cpp b/Engine/Source/Editor/AnimGraph/Private/AnimGraphNode_Base.cpp
--- a/Engine/Source/Editor/AnimGraph/Private/AnimGraphNode_Base.cpp
+++ b/Engine/Source/Editor/AnimGraph/Private/AnimGraphNode_Base.cpp
@@ -4,6 +4,7 @@
 #include "AnimGraphPrivatePCH.h"
 #include "AnimGraphNode_Base.h"
 #include "K2ActionMenuBuilder.h" // for FK2ActionMenuBuilder::AddNewNodeAction()
+#include "AnimationGraphSchema.h"
 
 /////////////////////////////////////////////////////
 // FA3NodeOptionalPinManager
diff --git a/Engine/Source/Editor/AnimGraph/Private/AnimGraphNode_BlendListByEnum.cpp b/Engine/Source/Editor/AnimGraph/Private/AnimGraphNode_BlendListByEnum.cpp
--- a/Engine/Source/Editor/AnimGraph/Private/AnimGraphNode_BlendListByEnum.cpp
+++ b/Engine/Source/Editor/AnimGraph/Private/AnimGraphNode_BlendListByEnum.cpp
@@ -7,6 +7,7 @@
 #include "CompilerResultsLog.h"
 #include "K2ActionMenuBuilder.h" // for FK2ActionMenuBuilder::AddNewNodeAction()
 #include "AnimGraphNode_BlendListByEnum.h"
+#include "AnimationGraphSchema.h"
 
 #define LOCTEXT_NAMESPACE "BlendListByEnum"
 
diff --git a/Engine/Source/Editor/AnimGraph/Private/AnimGraphNode_BlendSpaceBase.cpp b/Engine/Source/Editor/AnimGraph/Private/AnimGraphNode_BlendSpaceBase.cpp
--- a/Engine/Source/Editor/AnimGraph/Private/AnimGraphNode_BlendSpaceBase.cpp
+++ b/Engine/Source/Editor/AnimGraph/Private/AnimGraphNode_BlendSpaceBase.cpp
@@ -1,7 +1,10 @@
 // Copyright 1998-2014 Epic Games, Inc. All Rights Reserved.
 
 #include "AnimGraphPrivatePCH.h"
-
+#include "AnimationGraphSchema.h"
+#include "AnimGraphNode_BlendSpaceBase.h"
+#include "AnimGraphNode_BlendSpacePlayer.h"
+#include "AnimGraphNode_RotationOffsetBlendSpace.h"
 #include "GraphEditorActions.h"
 
 /////////////////////////////////////////////////////
diff --git a/Engine/Source/Editor/AnimGraph/Private/AnimGraphNode_BlendSpacePlayer.cpp b/Engine/Source/Editor/AnimGraph/Private/AnimGraphNode_BlendSpacePlayer.cpp
--- a/Engine/Source/Editor/AnimGraph/Private/AnimGraphNode_BlendSpacePlayer.cpp
+++ b/Engine/Source/Editor/AnimGraph/Private/AnimGraphNode_BlendSpacePlayer.cpp
@@ -3,6 +3,7 @@
 #include "AnimGraphPrivatePCH.h"
 #include "GraphEditorActions.h"
 #include "CompilerResultsLog.h"
+#include "AnimGraphNode_BlendSpacePlayer.h"
 
 /////////////////////////////////////////////////////
 // UAnimGraphNode_BlendSpacePlayer
diff --git a/Engine/Source/Editor/AnimGraph/Private/AnimGraphNode_ComponentToLocalSpace.cpp b/Engine/Source/Editor/AnimGraph/Private/AnimGraphNode_ComponentToLocalSpace.cpp
--- a/Engine/Source/Editor/AnimGraph/Private/AnimGraphNode_ComponentToLocalSpace.cpp
+++ b/Engine/Source/Editor/AnimGraph/Private/AnimGraphNode_ComponentToLocalSpace.cpp
@@ -1,6 +1,7 @@
 // Copyright 1998-2014 Epic Games, Inc. All Rights Reserved.
 
 #include "AnimGraphPrivatePCH.h"
+#include "AnimGraphNode_ComponentToLocalSpace.h"
 
 /////////////////////////////////////////////////////
 // UAnimGraphNode_ComponentToLocalSpace
diff --git a/Engine/Source/Editor/AnimGraph/Private/AnimGraphNode_CustomTransitionResult.cpp b/Engine/Source/Editor/AnimGraph/Private/AnimGraphNode_CustomTransitionResult.cpp
--- a/Engine/Source/Editor/AnimGraph/Private/AnimGraphNode_CustomTransitionResult.cpp
+++ b/Engine/Source/Editor/AnimGraph/Private/AnimGraphNode_CustomTransitionResult.cpp
@@ -1,6 +1,7 @@
 // Copyright 1998-2014 Epic Games, Inc. All Rights Reserved.
 
 #include "AnimGraphPrivatePCH.h"
+#include "AnimGraphNode_CustomTransitionResult.h"
 
 /////////////////////////////////////////////////////
 // UAnimGraphNode_CustomTransitionResult
diff --git a/Engine/Source/Editor/AnimGraph/Private/AnimGraphNode_LocalToComponentSpace.cpp b/Engine/Source/Editor/AnimGraph/Private/AnimGraphNode_LocalToComponentSpace.cpp
--- a/Engine/Source/Editor/AnimGraph/Private/AnimGraphNode_LocalToComponentSpace.cpp
+++ b/Engine/Source/Editor/AnimGraph/Private/AnimGraphNode_LocalToComponentSpace.cpp
@@ -1,9 +1,10 @@
 // Copyright 1998-2014 Epic Games, Inc. All Rights Reserved.
 
 #include "AnimGraphPrivatePCH.h"
-
+#include "AnimationGraphSchema.h"
 #include "GraphEditorActions.h"
 #include "ScopedTransaction.h"
+#include "AnimGraphNode_LocalToComponentSpace.h"
 
 /////////////////////////////////////////////////////
 // UAnimGraphNode_LocalToComponentSpace
diff --git a/Engine/Source/Editor/AnimGraph/Private/AnimGraphNode_MeshRefPose.cpp b/Engine/Source/Editor/AnimGraph/Private/AnimGraphNode_MeshRefPose.cpp
--- a/Engine/Source/Editor/AnimGraph/Private/AnimGraphNode_MeshRefPose.cpp
+++ b/Engine/Source/Editor/AnimGraph/Private/AnimGraphNode_MeshRefPose.cpp
@@ -2,6 +2,7 @@
 
 #include "AnimGraphPrivatePCH.h"
 #include "AnimGraphNode_MeshRefPose.h"
+#include "AnimationGraphSchema.h"
 
 /////////////////////////////////////////////////////
 // UAnimGraphNode_MeshRefPose
diff --git a/Engine/Source/Editor/AnimGraph/Private/AnimGraphNode_RotationOffsetBlendSpace.cpp b/Engine/Source/Editor/AnimGraph/Private/AnimGraphNode_RotationOffsetBlendSpace.cpp
--- a/Engine/Source/Editor/AnimGraph/Private/AnimGraphNode_RotationOffsetBlendSpace.cpp
+++ b/Engine/Source/Editor/AnimGraph/Private/AnimGraphNode_RotationOffsetBlendSpace.cpp
@@ -1,6 +1,7 @@
 // Copyright 1998-2014 Epic Games, Inc. All Rights Reserved.
 
 #include "AnimGraphPrivatePCH.h"
+#include "AnimGraphNode_RotationOffsetBlendSpace.h"
 #include "GraphEditorActions.h"
 #include "CompilerResultsLog.h"
 
diff --git a/Engine/Source/Editor/AnimGraph/Private/AnimGraphNode_SequencePlayer.cpp b/Engine/Source/Editor/AnimGraph/Private/AnimGraphNode_SequencePlayer.cpp
--- a/Engine/Source/Editor/AnimGraph/Private/AnimGraphNode_SequencePlayer.cpp
+++ b/Engine/Source/Editor/AnimGraph/Private/AnimGraphNode_SequencePlayer.cpp
@@ -1,11 +1,12 @@
 // Copyright 1998-2014 Epic Games, Inc. All Rights Reserved.
 
 #include "AnimGraphPrivatePCH.h"
+#include "AnimGraphNode_SequencePlayer.h"
 
 #include "CompilerResultsLog.h"
 #include "GraphEditorActions.h"
 #include "AssetRegistryModule.h"
-
+#include "AnimationGraphSchema.h"
 #define LOCTEXT_NAMESPACE "A3Nodes"
 
 /////////////////////////////////////////////////////
diff --git a/Engine/Source/Editor/AnimGraph/Private/AnimGraphNode_SkeletalControlBase.cpp b/Engine/Source/Editor/AnimGraph/Private/AnimGraphNode_SkeletalControlBase.cpp
--- a/Engine/Source/Editor/AnimGraph/Private/AnimGraphNode_SkeletalControlBase.cpp
+++ b/Engine/Source/Editor/AnimGraph/Private/AnimGraphNode_SkeletalControlBase.cpp
@@ -1,6 +1,8 @@
 // Copyright 1998-2014 Epic Games, Inc. All Rights Reserved.
 
 #include "AnimGraphPrivatePCH.h"
+#include "AnimGraphNode_SkeletalControlBase.h"
+#include "AnimationGraphSchema.h"
 #include "K2ActionMenuBuilder.h" // for FK2ActionMenuBuilder::AddNewNodeAction()
 
 /////////////////////////////////////////////////////
diff --git a/Engine/Source/Editor/AnimGraph/Private/AnimGraphNode_StateMachineBase.cpp b/Engine/Source/Editor/AnimGraph/Private/AnimGraphNode_StateMachineBase.cpp
--- a/Engine/Source/Editor/AnimGraph/Private/AnimGraphNode_StateMachineBase.cpp
+++ b/Engine/Source/Editor/AnimGraph/Private/AnimGraphNode_StateMachineBase.cpp
@@ -6,6 +6,11 @@
 #include "ScopedTransaction.h"
 #include "EdGraphUtilities.h"
 #include "Kismet2NameValidators.h"
+
+#include "AnimationGraph.h"
+#include "AnimationStateMachineGraph.h"
+#include "AnimationStateMachineSchema.h"
+#include "AnimGraphNode_StateMachine.h"
 #include "AnimGraphNode_StateMachineBase.h"
 
 /////////////////////////////////////////////////////
diff --git a/Engine/Source/Editor/AnimGraph/Private/AnimGraphNode_TransitionPoseEvaluator.cpp b/Engine/Source/Editor/AnimGraph/Private/AnimGraphNode_TransitionPoseEvaluator.cpp
--- a/Engine/Source/Editor/AnimGraph/Private/AnimGraphNode_TransitionPoseEvaluator.cpp
+++ b/Engine/Source/Editor/AnimGraph/Private/AnimGraphNode_TransitionPoseEvaluator.cpp
@@ -1,6 +1,8 @@
 // Copyright 1998-2014 Epic Games, Inc. All Rights Reserved.
 
 #include "AnimGraphPrivatePCH.h"
+#include "AnimGraphNode_TransitionPoseEvaluator.h"
+#include "AnimationCustomTransitionGraph.h"
 
 #include "CompilerResultsLog.h"
 #include "GraphEditorActions.h"
diff --git a/Engine/Source/Editor/AnimGraph/Private/AnimGraphNode_TransitionResult.cpp b/Engine/Source/Editor/AnimGraph/Private/AnimGraphNode_TransitionResult.cpp
--- a/Engine/Source/Editor/AnimGraph/Private/AnimGraphNode_TransitionResult.cpp
+++ b/Engine/Source/Editor/AnimGraph/Private/AnimGraphNode_TransitionResult.cpp
@@ -1,7 +1,7 @@
 // Copyright 1998-2014 Epic Games, Inc. All Rights Reserved.
 
 #include "AnimGraphPrivatePCH.h"
-
+#include "AnimGraphNode_TransitionResult.h"
 #include "GraphEditorActions.h"
 #include "ScopedTransaction.h"
 
diff --git a/Engine/Source/Editor/AnimGraph/Private/AnimGraphNode_UseCachedPose.cpp b/Engine/Source/Editor/AnimGraph/Private/AnimGraphNode_UseCachedPose.cpp
--- a/Engine/Source/Editor/AnimGraph/Private/AnimGraphNode_UseCachedPose.cpp
+++ b/Engine/Source/Editor/AnimGraph/Private/AnimGraphNode_UseCachedPose.cpp
@@ -5,6 +5,7 @@
 #include "GraphEditorActions.h"
 #include "ScopedTransaction.h"
 #include "K2ActionMenuBuilder.h" // for FK2ActionMenuBuilder::AddNewNodeAction()
+#include "AnimGraphNode_SaveCachedPose.h"
 #include "AnimGraphNode_UseCachedPose.h"
 
 /////////////////////////////////////////////////////
diff --git a/Engine/Source/Editor/AnimGraph/Private/AnimStateConduitNode.cpp b/Engine/Source/Editor/AnimGraph/Private/AnimStateConduitNode.cpp
--- a/Engine/Source/Editor/AnimGraph/Private/AnimStateConduitNode.cpp
+++ b/Engine/Source/Editor/AnimGraph/Private/AnimStateConduitNode.cpp
@@ -1,9 +1,13 @@
 // Copyright 1998-2014 Epic Games, Inc. All Rights Reserved.
 
 #include "AnimGraphPrivatePCH.h"
-
+#include "AnimStateConduitNode.h"
+#include "AnimStateTransitionNode.h"
 #include "Kismet2NameValidators.h"
 #include "CompilerResultsLog.h"
+#include "AnimationTransitionGraph.h"
+#include "AnimationConduitGraphSchema.h"
+#include "AnimGraphNode_TransitionResult.h"
 
 /////////////////////////////////////////////////////
 // UAnimStateConduitNode
diff --git a/Engine/Source/Editor/AnimGraph/Private/AnimStateEntryNode.cpp b/Engine/Source/Editor/AnimGraph/Private/AnimStateEntryNode.cpp
--- a/Engine/Source/Editor/AnimGraph/Private/AnimStateEntryNode.cpp
+++ b/Engine/Source/Editor/AnimGraph/Private/AnimStateEntryNode.cpp
@@ -5,8 +5,9 @@
 =============================================================================*/
 
 #include "AnimGraphPrivatePCH.h"
-
+#include "AnimStateEntryNode.h"
 #include "EdGraphUtilities.h"
+#include "AnimationStateMachineSchema.h"
 
 /////////////////////////////////////////////////////
 // UAnimStateEntryNode
diff --git a/Engine/Source/Editor/AnimGraph/Private/AnimStateNode.cpp b/Engine/Source/Editor/AnimGraph/Private/AnimStateNode.cpp
--- a/Engine/Source/Editor/AnimGraph/Private/AnimStateNode.cpp
+++ b/Engine/Source/Editor/AnimGraph/Private/AnimStateNode.cpp
@@ -5,10 +5,13 @@
 =============================================================================*/
 
 #include "AnimGraphPrivatePCH.h"
-
+#include "AnimStateNode.h"
+#include "AnimStateTransitionNode.h"
+#include "AnimationStateGraph.h"
 #include "EdGraphUtilities.h"
 #include "Kismet2/Kismet2NameValidators.h"
 #include "AnimGraphNode_StateResult.h"
+#include "AnimationStateGraphSchema.h"
 
 /////////////////////////////////////////////////////
 // UAnimStateNode
diff --git a/Engine/Source/Editor/AnimGraph/Private/AnimStateNodeBase.cpp b/Engine/Source/Editor/AnimGraph/Private/AnimStateNodeBase.cpp
--- a/Engine/Source/Editor/AnimGraph/Private/AnimStateNodeBase.cpp
+++ b/Engine/Source/Editor/AnimGraph/Private/AnimStateNodeBase.cpp
@@ -5,17 +5,18 @@
 =============================================================================*/
 
 #include "AnimGraphPrivatePCH.h"
-
+#include "AnimationStateMachineGraph.h"
+#include "AnimationStateMachineSchema.h"
 #include "EdGraphUtilities.h"
 #include "Kismet2/Kismet2NameValidators.h"
-
+#include "AnimStateNodeBase.h"
 /////////////////////////////////////////////////////
 // FAnimStateNodeNameValidator
 
 class FAnimStateNodeNameValidator : public FStringSetNameValidator
 {
 public:
-	FAnimStateNodeNameValidator(const class UAnimStateNodeBase* InStateNode)
+	FAnimStateNodeNameValidator(const UAnimStateNodeBase* InStateNode)
 		: FStringSetNameValidator(FString())
 	{
 		TArray<UAnimStateNodeBase*> Nodes;
diff --git a/Engine/Source/Editor/AnimGraph/Private/AnimStateTransitionNode.cpp b/Engine/Source/Editor/AnimGraph/Private/AnimStateTransitionNode.cpp
--- a/Engine/Source/Editor/AnimGraph/Private/AnimStateTransitionNode.cpp
+++ b/Engine/Source/Editor/AnimGraph/Private/AnimStateTransitionNode.cpp
@@ -5,7 +5,12 @@
 =============================================================================*/
 
 #include "AnimGraphPrivatePCH.h"
-
+#include "AnimStateTransitionNode.h"
+#include "AnimationTransitionGraph.h"
+#include "AnimationTransitionSchema.h"
+#include "AnimationCustomTransitionGraph.h"
+#include "AnimationCustomTransitionSchema.h"
+#include "AnimGraphNode_TransitionResult.h"
 #include "EdGraphUtilities.h"
 #include "Kismet2/Kismet2NameValidators.h"
 
diff --git a/Engine/Source/Editor/AnimGraph/Private/AnimationCustomTransitionSchema.cpp b/Engine/Source/Editor/AnimGraph/Private/AnimationCustomTransitionSchema.cpp
--- a/Engine/Source/Editor/AnimGraph/Private/AnimationCustomTransitionSchema.cpp
+++ b/Engine/Source/Editor/AnimGraph/Private/AnimationCustomTransitionSchema.cpp
@@ -5,6 +5,7 @@
 #include "AnimGraphNode_CustomTransitionResult.h"
 #include "AnimGraphNode_TransitionPoseEvaluator.h"
 #include "AnimStateTransitionNode.h"
+#include "AnimationCustomTransitionGraph.h"
 
 /////////////////////////////////////////////////////
 // UAnimationCustomTransitionSchema
diff --git a/Engine/Source/Editor/AnimGraph/Private/AnimationGraphSchema.cpp b/Engine/Source/Editor/AnimGraph/Private/AnimationGraphSchema.cpp
--- a/Engine/Source/Editor/AnimGraph/Private/AnimationGraphSchema.cpp
+++ b/Engine/Source/Editor/AnimGraph/Private/AnimationGraphSchema.cpp
@@ -16,6 +16,7 @@
 #include "AnimGraphNode_BlendSpacePlayer.h"
 #include "AnimGraphNode_ComponentToLocalSpace.h"
 #include "AnimGraphNode_LocalToComponentSpace.h"
+#include "AnimGraphNode_Root.h"
 #include "AnimGraphNode_RotationOffsetBlendSpace.h"
 #include "AnimGraphNode_SequencePlayer.h"
 
diff --git a/Engine/Source/Editor/AnimGraph/Private/AnimationStateGraphSchema.cpp b/Engine/Source/Editor/AnimGraph/Private/AnimationStateGraphSchema.cpp
--- a/Engine/Source/Editor/AnimGraph/Private/AnimationStateGraphSchema.cpp
+++ b/Engine/Source/Editor/AnimGraph/Private/AnimationStateGraphSchema.cpp
@@ -6,7 +6,9 @@
 #include "ScopedTransaction.h"
 #include "AssetData.h"
 #include "AnimationStateGraphSchema.h"
-
+#include "AnimationStateGraph.h"
+#include "AnimStateNode.h"
+#include "AnimGraphNode_StateResult.h"
 #define LOCTEXT_NAMESPACE "AnimationStateGraphSchema"
 
 
diff --git a/Engine/Source/Editor/AnimGraph/Private/AnimationStateMachineSchema.cpp b/Engine/Source/Editor/AnimGraph/Private/AnimationStateMachineSchema.cpp
--- a/Engine/Source/Editor/AnimGraph/Private/AnimationStateMachineSchema.cpp
+++ b/Engine/Source/Editor/AnimGraph/Private/AnimationStateMachineSchema.cpp
@@ -6,8 +6,14 @@
 
 #include "AnimGraphPrivatePCH.h"
 
+#include "AnimationStateMachineGraph.h"
 #include "AnimationStateMachineSchema.h"
+#include "AnimationGraphSchema.h"
 #include "AnimStateEntryNode.h"
+#include "AnimStateNode.h"
+#include "AnimStateTransitionNode.h"
+#include "AnimStateConduitNode.h"
+#include "AnimStateNodeBase.h"
 
 #include "ScopedTransaction.h"
 #include "GraphEditorActions.h"
diff --git a/Engine/Source/Editor/AnimGraph/Private/AnimationTransitionGraph.cpp b/Engine/Source/Editor/AnimGraph/Private/AnimationTransitionGraph.cpp
--- a/Engine/Source/Editor/AnimGraph/Private/AnimationTransitionGraph.cpp
+++ b/Engine/Source/Editor/AnimGraph/Private/AnimationTransitionGraph.cpp
@@ -1,7 +1,7 @@
 // Copyright 1998-2014 Epic Games, Inc. All Rights Reserved.
 
 #include "AnimGraphPrivatePCH.h"
-
+#include "AnimationTransitionGraph.h"
 
 #define LOCTEXT_NAMESPACE "AnimationStateGraph"
 
diff --git a/Engine/Source/Editor/AnimGraph/Private/AnimationTransitionSchema.cpp b/Engine/Source/Editor/AnimGraph/Private/AnimationTransitionSchema.cpp
--- a/Engine/Source/Editor/AnimGraph/Private/AnimationTransitionSchema.cpp
+++ b/Engine/Source/Editor/AnimGraph/Private/AnimationTransitionSchema.cpp
@@ -5,7 +5,12 @@
 =============================================================================*/
 
 #include "AnimGraphPrivatePCH.h"
+#include "AnimStateNode.h"
+#include "AnimStateTransitionNode.h"
+#include "AnimationTransitionGraph.h"
 #include "AnimationTransitionSchema.h"
+#include "AnimGraphNode_Base.h"
+#include "AnimGraphNode_TransitionResult.h"
 
 #include "BlueprintUtilities.h"
 #include "AnimGraphDefinitions.h"
@@ -13,6 +18,7 @@
 #include "ScopedTransaction.h"
 #include "Kismet2/KismetDebugUtilities.h"
 #include "K2ActionMenuBuilder.h" // for FK2ActionMenuBuilder::AddNewNodeAction()
+#include "K2Node_TransitionRuleGetter.h"
 
 /////////////////////////////////////////////////////
 // UAnimationTransitionSchema
diff --git a/Engine/Source/Editor/AnimGraph/Private/K2Node_TransitionRuleGetter.cpp b/Engine/Source/Editor/AnimGraph/Private/K2Node_TransitionRuleGetter.cpp
--- a/Engine/Source/Editor/AnimGraph/Private/K2Node_TransitionRuleGetter.cpp
+++ b/Engine/Source/Editor/AnimGraph/Private/K2Node_TransitionRuleGetter.cpp
@@ -2,9 +2,10 @@
 
 
 #include "AnimGraphPrivatePCH.h"
-
+#include "K2Node_TransitionRuleGetter.h"
 #include "BlueprintUtilities.h"
-
+#include "AnimGraphNode_Base.h"
+#include "AnimStateNode.h"
 /////////////////////////////////////////////////////
 // UK2Node_TransitionRuleGetter
 
diff --git a/Engine/Source/Editor/BlueprintGraph/Classes/EdGraphSchema_K2.h b/Engine/Source/Editor/BlueprintGraph/Classes/EdGraphSchema_K2.h
--- a/Engine/Source/Editor/BlueprintGraph/Classes/EdGraphSchema_K2.h
+++ b/Engine/Source/Editor/BlueprintGraph/Classes/EdGraphSchema_K2.h
@@ -1,7 +1,7 @@
 // Copyright 1998-2014 Epic Games, Inc. All Rights Reserved.
 
 #pragma once
-
+#include "EdGraph/EdGraphSchema.h"
 #include "EdGraphSchema_K2.generated.h"
 
 /** Reference to an structure (only used in 'docked' palette) */
diff --git a/Engine/Source/Editor/GraphEditor/Private/AnimationNodes/SGraphNodeAnimationResult.cpp b/Engine/Source/Editor/GraphEditor/Private/AnimationNodes/SGraphNodeAnimationResult.cpp
--- a/Engine/Source/Editor/GraphEditor/Private/AnimationNodes/SGraphNodeAnimationResult.cpp
+++ b/Engine/Source/Editor/GraphEditor/Private/AnimationNodes/SGraphNodeAnimationResult.cpp
@@ -4,6 +4,7 @@
 #include "AnimGraphDefinitions.h"
 #include "SGraphNodeAnimationResult.h"
 #include "EdGraph/EdGraphNode.h"
+#include "AnimGraphNode_base.h"
 
 /////////////////////////////////////////////////////
 // SGraphNodeAnimationResult
diff --git a/Engine/Source/Editor/GraphEditor/Private/AnimationStateNodes/SGraphNodeAnimState.cpp b/Engine/Source/Editor/GraphEditor/Private/AnimationStateNodes/SGraphNodeAnimState.cpp
--- a/Engine/Source/Editor/GraphEditor/Private/AnimationStateNodes/SGraphNodeAnimState.cpp
+++ b/Engine/Source/Editor/GraphEditor/Private/AnimationStateNodes/SGraphNodeAnimState.cpp
@@ -8,6 +8,7 @@
 #include "IDocumentation.h"
 #include "AnimStateConduitNode.h"
 #include "AnimationStateMachineGraph.h"
+#include "AnimGraphNode_StateMachineBase.h"
 
 /////////////////////////////////////////////////////
 // SStateMachineOutputPin
diff --git a/Engine/Source/Editor/GraphEditor/Private/AnimationStateNodes/SGraphNodeAnimStateEntry.cpp b/Engine/Source/Editor/GraphEditor/Private/AnimationStateNodes/SGraphNodeAnimStateEntry.cpp
--- a/Engine/Source/Editor/GraphEditor/Private/AnimationStateNodes/SGraphNodeAnimStateEntry.cpp
+++ b/Engine/Source/Editor/GraphEditor/Private/AnimationStateNodes/SGraphNodeAnimStateEntry.cpp
@@ -6,6 +6,7 @@
 #include "SGraphPreviewer.h"
 #include "Editor/UnrealEd/Public/Kismet2/BlueprintEditorUtils.h"
 #include "NodeFactory.h"
+#include "AnimStateEntryNode.h"
 
 /////////////////////////////////////////////////////
 // SGraphNodeAnimStateEntry
diff --git a/Engine/Source/Editor/GraphEditor/Private/AnimationStateNodes/SGraphNodeAnimStateEntry.h b/Engine/Source/Editor/GraphEditor/Private/AnimationStateNodes/SGraphNodeAnimStateEntry.h
--- a/Engine/Source/Editor/GraphEditor/Private/AnimationStateNodes/SGraphNodeAnimStateEntry.h
+++ b/Engine/Source/Editor/GraphEditor/Private/AnimationStateNodes/SGraphNodeAnimStateEntry.h
@@ -2,6 +2,11 @@
 
 #pragma once
 
+//
+// Forward declarations.
+//
+class UAnimStateEntryNode;
+
 class SGraphNodeAnimStateEntry : public SGraphNode
 {
 public:
diff --git a/Engine/Source/Editor/GraphEditor/Private/NodeFactory.cpp b/Engine/Source/Editor/GraphEditor/Private/NodeFactory.cpp
--- a/Engine/Source/Editor/GraphEditor/Private/NodeFactory.cpp
+++ b/Engine/Source/Editor/GraphEditor/Private/NodeFactory.cpp
@@ -28,7 +28,11 @@
 
 #include "AnimGraphNode_Root.h"
 #include "AnimGraphNode_SequencePlayer.h"
+#include "AnimGraphNode_StateMachineBase.h"
 #include "AnimStateNode.h"
+#include "AnimStateEntryNode.h"
+#include "AnimStateConduitNode.h"
+#include "AnimStateTransitionNode.h"
 
 #include "AnimationStateMachineSchema.h"
 
diff --git a/Engine/Source/Editor/GraphEditor/Private/SGraphPanel.cpp b/Engine/Source/Editor/GraphEditor/Private/SGraphPanel.cpp
--- a/Engine/Source/Editor/GraphEditor/Private/SGraphPanel.cpp
+++ b/Engine/Source/Editor/GraphEditor/Private/SGraphPanel.cpp
@@ -16,6 +16,7 @@
 #include "GraphDiffControl.h"
 
 #include "AnimationGraphSchema.h"
+#include "AnimationStateMachineSchema.h"
 
 DEFINE_LOG_CATEGORY_STATIC(LogGraphPanel, Log, All);
 
diff --git a/Engine/Source/Runtime/Engine/Classes/Animation/AnimNode_StateMachine.h b/Engine/Source/Runtime/Engine/Classes/Animation/AnimNode_StateMachine.h
--- a/Engine/Source/Runtime/Engine/Classes/Animation/AnimNode_StateMachine.h
+++ b/Engine/Source/Runtime/Engine/Classes/Animation/AnimNode_StateMachine.h
@@ -1,11 +1,10 @@
 // Copyright 1998-2014 Epic Games, Inc. All Rights Reserved.
 
 #pragma once
-
+#include "AnimNodeBase.h"
 #include "AnimStateMachineTypes.h"
 #include "AnimNode_StateMachine.generated.h"
 
-
 struct FAnimNode_TransitionPoseEvaluator;
 struct FAnimNode_StateMachine;
 
