Fix non-unity

diff --git a/Engine/Source/Editor/AnimGraph/Private/AnimGraphNode_WheelHandler.cpp b/Engine/Source/Editor/AnimGraph/Private/AnimGraphNode_WheelHandler.cpp
--- a/Engine/Source/Editor/AnimGraph/Private/AnimGraphNode_WheelHandler.cpp
+++ b/Engine/Source/Editor/AnimGraph/Private/AnimGraphNode_WheelHandler.cpp
@@ -2,6 +2,7 @@
 
 #include "AnimGraphPrivatePCH.h"
 #include "AnimGraphNode_WheelHandler.h"
+#include "CompilerResultsLog.h"
 
 /////////////////////////////////////////////////////
 // UAnimGraphNode_WheelHandler
