#ue4 Fix for AnimGraph linker error

diff --git a/Engine/Source/Editor/AnimGraph/Classes/AnimGraphNode_Base.h b/Engine/Source/Editor/AnimGraph/Classes/AnimGraphNode_Base.h
--- a/Engine/Source/Editor/AnimGraph/Classes/AnimGraphNode_Base.h
+++ b/Engine/Source/Editor/AnimGraph/Classes/AnimGraphNode_Base.h
@@ -112,7 +112,7 @@ class UAnimGraphNode_Base : public UK2Node
 	ANIMGRAPH_API virtual void ReallocatePinsDuringReconstruction(TArray<UEdGraphPin*>& OldPins) OVERRIDE;
 	ANIMGRAPH_API virtual bool CanCreateUnderSpecifiedSchema(const UEdGraphSchema* DesiredSchema) const OVERRIDE;
 	ANIMGRAPH_API virtual void GetMenuEntries(FGraphContextMenuBuilder& ContextMenuBuilder) const OVERRIDE;
-	virtual void GetNodeAttributes( TArray<TKeyValuePair<FString, FString>>& OutNodeAttributes ) const OVERRIDE;
+	ANIMGRAPH_API virtual void GetNodeAttributes(TArray<TKeyValuePair<FString, FString>>& OutNodeAttributes) const OVERRIDE;
 	// End of UK2Node interface
 
 	// UAnimGraphNode_Base interface
