UE4 Refactoring. Changed OVERRIDE and FINAL macros to keywords override and final.

diff --git a/Engine/Source/Developer/Windows/WindowsTargetPlatform/Private/GenericWindowsTargetPlatform.h b/Engine/Source/Developer/Windows/WindowsTargetPlatform/Private/GenericWindowsTargetPlatform.h
--- a/Engine/Source/Developer/Windows/WindowsTargetPlatform/Private/GenericWindowsTargetPlatform.h
+++ b/Engine/Source/Developer/Windows/WindowsTargetPlatform/Private/GenericWindowsTargetPlatform.h
@@ -55,25 +55,25 @@ public:
 
 	// Begin ITargetPlatform interface
 
-	virtual void EnableDeviceCheck(bool OnOff) OVERRIDE {}
+	virtual void EnableDeviceCheck(bool OnOff) override {}
 
-	virtual void GetAllDevices( TArray<ITargetDevicePtr>& OutDevices ) const OVERRIDE
+	virtual void GetAllDevices( TArray<ITargetDevicePtr>& OutDevices ) const override
 	{
 		OutDevices.Reset();
 		OutDevices.Add(LocalDevice);
 	}
 
-	virtual ECompressionFlags GetBaseCompressionMethod( ) const OVERRIDE
+	virtual ECompressionFlags GetBaseCompressionMethod( ) const override
 	{
 		return COMPRESS_ZLIB;
 	}
 
-	virtual bool GenerateStreamingInstallManifest(const TMultiMap<FString, int32>& ChunkMap, const TSet<int32>& ChunkIDsInUse) const OVERRIDE
+	virtual bool GenerateStreamingInstallManifest(const TMultiMap<FString, int32>& ChunkMap, const TSet<int32>& ChunkIDsInUse) const override
 	{
 		return true;
 	}
 
-	virtual ITargetDevicePtr GetDefaultDevice( ) const OVERRIDE
+	virtual ITargetDevicePtr GetDefaultDevice( ) const override
 	{
 		return LocalDevice;
 	}
@@ -88,13 +88,13 @@ public:
 		return NULL;
 	}
 
-	virtual bool IsRunningPlatform( ) const OVERRIDE
+	virtual bool IsRunningPlatform( ) const override
 	{
 		// Must be Windows platform as editor for this to be considered a running platform
 		return PLATFORM_WINDOWS && !UE_SERVER && !UE_GAME && WITH_EDITOR && HAS_EDITOR_DATA;
 	}
 
-	virtual bool SupportsFeature( ETargetPlatformFeatures::Type Feature ) const OVERRIDE
+	virtual bool SupportsFeature( ETargetPlatformFeatures::Type Feature ) const override
 	{
 		// we currently do not have a build target for WindowsServer
 		if (Feature == ETargetPlatformFeatures::Packaging)
@@ -106,7 +106,7 @@ public:
 	}
 
 #if WITH_ENGINE
-	virtual void GetAllPossibleShaderFormats( TArray<FName>& OutFormats ) const OVERRIDE
+	virtual void GetAllPossibleShaderFormats( TArray<FName>& OutFormats ) const override
 	{
 		// no shaders needed for dedicated server target
 		if (!IS_DEDICATED_SERVER)
@@ -126,7 +126,7 @@ public:
 		}
 	}
 
-	virtual void GetAllTargetedShaderFormats( TArray<FName>& OutFormats ) const OVERRIDE
+	virtual void GetAllTargetedShaderFormats( TArray<FName>& OutFormats ) const override
 	{
 		for(const FString& ShaderFormat : TargetedShaderFormats)
 		{
@@ -134,12 +134,12 @@ public:
 		}
 	}
 
-	virtual const class FStaticMeshLODSettings& GetStaticMeshLODSettings( ) const OVERRIDE
+	virtual const class FStaticMeshLODSettings& GetStaticMeshLODSettings( ) const override
 	{
 		return StaticMeshLODSettings;
 	}
 
-	virtual void GetTextureFormats( const UTexture* InTexture, TArray<FName>& OutFormats ) const OVERRIDE
+	virtual void GetTextureFormats( const UTexture* InTexture, TArray<FName>& OutFormats ) const override
 	{
 		if (!IS_DEDICATED_SERVER)
 		{
@@ -148,12 +148,12 @@ public:
 		}
 	}
 
-	virtual const struct FTextureLODSettings& GetTextureLODSettings( ) const OVERRIDE
+	virtual const struct FTextureLODSettings& GetTextureLODSettings( ) const override
 	{
 		return TextureLODSettings;
 	}
 
-	virtual FName GetWaveFormat( class USoundWave* Wave ) const OVERRIDE
+	virtual FName GetWaveFormat( class USoundWave* Wave ) const override
 	{
 		static FName NAME_OGG(TEXT("OGG"));
 		static FName NAME_OPUS(TEXT("OPUS"));
@@ -168,13 +168,13 @@ public:
 #endif //WITH_ENGINE
 
 	DECLARE_DERIVED_EVENT(TGenericWindowsTargetPlatform, ITargetPlatform::FOnTargetDeviceDiscovered, FOnTargetDeviceDiscovered);
-	virtual FOnTargetDeviceDiscovered& OnDeviceDiscovered( ) OVERRIDE
+	virtual FOnTargetDeviceDiscovered& OnDeviceDiscovered( ) override
 	{
 		return DeviceDiscoveredEvent;
 	}
 
 	DECLARE_DERIVED_EVENT(TGenericWindowsTargetPlatform, ITargetPlatform::FOnTargetDeviceLost, FOnTargetDeviceLost);
-	virtual FOnTargetDeviceLost& OnDeviceLost( ) OVERRIDE
+	virtual FOnTargetDeviceLost& OnDeviceLost( ) override
 	{
 		return DeviceLostEvent;
 	}
diff --git a/Engine/Source/Developer/Windows/WindowsTargetPlatform/Private/LocalPcTargetDevice.h b/Engine/Source/Developer/Windows/WindowsTargetPlatform/Private/LocalPcTargetDevice.h
--- a/Engine/Source/Developer/Windows/WindowsTargetPlatform/Private/LocalPcTargetDevice.h
+++ b/Engine/Source/Developer/Windows/WindowsTargetPlatform/Private/LocalPcTargetDevice.h
@@ -29,12 +29,12 @@ public:
 
 public:
 
-	virtual bool Connect( ) OVERRIDE
+	virtual bool Connect( ) override
 	{
 		return true;
 	}
 
-	virtual bool Deploy( const FString& SourceFolder, FString& OutAppId ) OVERRIDE
+	virtual bool Deploy( const FString& SourceFolder, FString& OutAppId ) override
 	{
 		OutAppId = TEXT("");
 
@@ -63,7 +63,7 @@ public:
 	virtual void Disconnect( )
 	{ }
 
-	virtual ETargetDeviceTypes::Type GetDeviceType( ) const OVERRIDE
+	virtual ETargetDeviceTypes::Type GetDeviceType( ) const override
 	{
 		if (::GetSystemMetrics(SM_TABLETPC) != 0)
 		{
@@ -73,17 +73,17 @@ public:
 		return ETargetDeviceTypes::Desktop;
 	}
 
-	virtual FTargetDeviceId GetId( ) const OVERRIDE
+	virtual FTargetDeviceId GetId( ) const override
 	{
 		return FTargetDeviceId(TargetPlatform.PlatformName(), GetName());
 	}
 
-	virtual FString GetName( ) const OVERRIDE
+	virtual FString GetName( ) const override
 	{
 		return FPlatformProcess::ComputerName();
 	}
 
-	virtual FString GetOperatingSystemName( ) OVERRIDE
+	virtual FString GetOperatingSystemName( ) override
 	{
 		if (WIN64)
 		{
@@ -93,7 +93,7 @@ public:
 		return TEXT("Windows (32-bit)");
 	}
 
-	virtual int32 GetProcessSnapshot( TArray<FTargetDeviceProcessInfo>& OutProcessInfos ) OVERRIDE
+	virtual int32 GetProcessSnapshot( TArray<FTargetDeviceProcessInfo>& OutProcessInfos ) override
 	{
 		// enumerate processes
 		HANDLE ProcessSnapshot = ::CreateToolhelp32Snapshot(TH32CS_SNAPPROCESS, 0);
@@ -216,7 +216,7 @@ public:
 		return OutProcessInfos.Num();
 	}
 
-	virtual const class ITargetPlatform& GetTargetPlatform( ) const OVERRIDE
+	virtual const class ITargetPlatform& GetTargetPlatform( ) const override
 	{
 		return TargetPlatform;
 	}
@@ -226,12 +226,12 @@ public:
 		return true;
 	}
 
-	virtual bool IsDefault( ) const OVERRIDE
+	virtual bool IsDefault( ) const override
 	{
 		return true;
 	}
 
-	virtual bool Launch( const FString& AppId, EBuildConfigurations::Type BuildConfiguration, EBuildTargets::Type BuildTarget, const FString& Params, uint32* OutProcessId ) OVERRIDE
+	virtual bool Launch( const FString& AppId, EBuildConfigurations::Type BuildConfiguration, EBuildTargets::Type BuildTarget, const FString& Params, uint32* OutProcessId ) override
 	{
 		// build executable path
 		FString PlatformName = WIN64 ? TEXT("Win64") : TEXT("Win32");
@@ -262,7 +262,7 @@ public:
 		return ProcessHandle.Close();
 	}
 
-	virtual bool PowerOff( bool Force ) OVERRIDE
+	virtual bool PowerOff( bool Force ) override
 	{
 		if (!AdjustShutdownPrivileges())
 		{
@@ -272,12 +272,12 @@ public:
 		return (::ExitWindowsEx(EWX_POWEROFF | EWX_FORCE, SHTDN_REASON_MINOR_MAINTENANCE | SHTDN_REASON_FLAG_PLANNED) != 0);
 	}
 
-	virtual bool PowerOn( ) OVERRIDE
+	virtual bool PowerOn( ) override
 	{
 		return false;
 	}
 
-	virtual bool Reboot( bool bReconnect = false ) OVERRIDE
+	virtual bool Reboot( bool bReconnect = false ) override
 	{
 		if (!AdjustShutdownPrivileges())
 		{
@@ -293,7 +293,7 @@ public:
 		return ProcessHandle.Close();
 	}
 
-	virtual bool SupportsFeature( ETargetDeviceFeatures::Type Feature ) const OVERRIDE
+	virtual bool SupportsFeature( ETargetDeviceFeatures::Type Feature ) const override
 	{
 		switch (Feature)
 		{
@@ -317,22 +317,22 @@ public:
 		return false;
 	}
 
-	virtual bool SupportsSdkVersion( const FString& VersionString ) const OVERRIDE
+	virtual bool SupportsSdkVersion( const FString& VersionString ) const override
 	{
 		// @todo filter SDK versions
 		return true;
 	}
 
-	virtual void SetUserCredentials( const FString & UserName, const FString & UserPassword ) OVERRIDE
+	virtual void SetUserCredentials( const FString & UserName, const FString & UserPassword ) override
 	{
 	}
 
-	virtual bool GetUserCredentials( FString & OutUserName, FString & OutUserPassword ) OVERRIDE
+	virtual bool GetUserCredentials( FString & OutUserName, FString & OutUserPassword ) override
 	{
 		return false;
 	}
 
-	virtual bool TerminateProcess( const int32 ProcessId ) OVERRIDE
+	virtual bool TerminateProcess( const int32 ProcessId ) override
 	{
 		HANDLE ProcessHandle = OpenProcess(PROCESS_TERMINATE, false, ProcessId);
 
diff --git a/Engine/Source/Developer/iOS/IOSTargetPlatform/Private/IOSTargetDevice.h b/Engine/Source/Developer/iOS/IOSTargetPlatform/Private/IOSTargetDevice.h
--- a/Engine/Source/Developer/iOS/IOSTargetPlatform/Private/IOSTargetDevice.h
+++ b/Engine/Source/Developer/iOS/IOSTargetPlatform/Private/IOSTargetDevice.h
@@ -39,27 +39,27 @@ public:
 	}
 public:
 	// ITargetDevice interface
-	virtual bool Connect() OVERRIDE;
-	virtual bool Deploy(const FString& SourceFolder, FString& OutAppId) OVERRIDE;
-	virtual void Disconnect() OVERRIDE;
-	virtual int32 GetProcessSnapshot(TArray<FTargetDeviceProcessInfo>& OutProcessInfos) OVERRIDE;
-	virtual ETargetDeviceTypes::Type GetDeviceType() const OVERRIDE;
-	virtual FTargetDeviceId GetId() const OVERRIDE;
-	virtual FString GetName() const OVERRIDE;
-	virtual FString GetOperatingSystemName() OVERRIDE;
-	virtual const class ITargetPlatform& GetTargetPlatform() const OVERRIDE;
-	virtual bool IsConnected() OVERRIDE;
-	virtual bool IsDefault() const OVERRIDE;
-	virtual bool Launch(const FString& InAppId, EBuildConfigurations::Type InBuildConfiguration, EBuildTargets::Type BuildTarget, const FString& Params, uint32* OutProcessId) OVERRIDE;
-	virtual bool PowerOff(bool Force) OVERRIDE;
-	virtual bool PowerOn() OVERRIDE;
-	virtual bool Reboot(bool bReconnect = false) OVERRIDE;
-	virtual bool Run(const FString& ExecutablePath, const FString& Params, uint32* OutProcessId) OVERRIDE;
+	virtual bool Connect() override;
+	virtual bool Deploy(const FString& SourceFolder, FString& OutAppId) override;
+	virtual void Disconnect() override;
+	virtual int32 GetProcessSnapshot(TArray<FTargetDeviceProcessInfo>& OutProcessInfos) override;
+	virtual ETargetDeviceTypes::Type GetDeviceType() const override;
+	virtual FTargetDeviceId GetId() const override;
+	virtual FString GetName() const override;
+	virtual FString GetOperatingSystemName() override;
+	virtual const class ITargetPlatform& GetTargetPlatform() const override;
+	virtual bool IsConnected() override;
+	virtual bool IsDefault() const override;
+	virtual bool Launch(const FString& InAppId, EBuildConfigurations::Type InBuildConfiguration, EBuildTargets::Type BuildTarget, const FString& Params, uint32* OutProcessId) override;
+	virtual bool PowerOff(bool Force) override;
+	virtual bool PowerOn() override;
+	virtual bool Reboot(bool bReconnect = false) override;
+	virtual bool Run(const FString& ExecutablePath, const FString& Params, uint32* OutProcessId) override;
 	virtual bool SupportsFeature(ETargetDeviceFeatures::Type Feature) const;
-	virtual bool SupportsSdkVersion(const FString& VersionString) const OVERRIDE;
-	virtual bool TerminateProcess(const int32 ProcessId) OVERRIDE;
-	virtual void SetUserCredentials(const FString & UserName, const FString & UserPassword) OVERRIDE;
-	virtual bool GetUserCredentials(FString & OutUserName, FString & OutUserPassword) OVERRIDE;
+	virtual bool SupportsSdkVersion(const FString& VersionString) const override;
+	virtual bool TerminateProcess(const int32 ProcessId) override;
+	virtual void SetUserCredentials(const FString & UserName, const FString & UserPassword) override;
+	virtual bool GetUserCredentials(FString & OutUserName, FString & OutUserPassword) override;
 	// End of ITargetDevice interface
 
 public:
diff --git a/Engine/Source/Developer/iOS/IOSTargetPlatform/Private/IOSTargetPlatform.h b/Engine/Source/Developer/iOS/IOSTargetPlatform/Private/IOSTargetPlatform.h
--- a/Engine/Source/Developer/iOS/IOSTargetPlatform/Private/IOSTargetPlatform.h
+++ b/Engine/Source/Developer/iOS/IOSTargetPlatform/Private/IOSTargetPlatform.h
@@ -34,7 +34,7 @@ public:
 
 	// Begin TTargetPlatformBase interface
 
-	virtual bool IsServerOnly( ) const OVERRIDE
+	virtual bool IsServerOnly( ) const override
 	{
 		return false;
 	}
@@ -45,25 +45,25 @@ public:
 
 	// Begin ITargetPlatform interface
 
-	virtual void EnableDeviceCheck(bool OnOff) OVERRIDE;
+	virtual void EnableDeviceCheck(bool OnOff) override;
 
-	virtual void GetAllDevices( TArray<ITargetDevicePtr>& OutDevices ) const OVERRIDE;
+	virtual void GetAllDevices( TArray<ITargetDevicePtr>& OutDevices ) const override;
 
-	virtual ECompressionFlags GetBaseCompressionMethod() const OVERRIDE
+	virtual ECompressionFlags GetBaseCompressionMethod() const override
 	{
 		return COMPRESS_ZLIB;
 	}
 
-	virtual bool GenerateStreamingInstallManifest(const TMultiMap<FString, int32>& ChunkMap, const TSet<int32>& ChunkIDsInUse) const OVERRIDE
+	virtual bool GenerateStreamingInstallManifest(const TMultiMap<FString, int32>& ChunkMap, const TSet<int32>& ChunkIDsInUse) const override
 	{
 		return true;
 	}
 
-	virtual ITargetDevicePtr GetDefaultDevice( ) const OVERRIDE;
+	virtual ITargetDevicePtr GetDefaultDevice( ) const override;
 
-	virtual ITargetDevicePtr GetDevice( const FTargetDeviceId& DeviceId ) OVERRIDE;
+	virtual ITargetDevicePtr GetDevice( const FTargetDeviceId& DeviceId ) override;
 
-	virtual bool IsRunningPlatform( ) const OVERRIDE
+	virtual bool IsRunningPlatform( ) const override
 	{
 		#if PLATFORM_IOS && WITH_EDITOR
 			return true;
@@ -72,7 +72,7 @@ public:
 		#endif
 	}
 
-	virtual bool SupportsFeature( ETargetPlatformFeatures::Type Feature ) const OVERRIDE
+	virtual bool SupportsFeature( ETargetPlatformFeatures::Type Feature ) const override
 	{
 		if (Feature == ETargetPlatformFeatures::Packaging)
 		{
@@ -83,40 +83,40 @@ public:
 		return TTargetPlatformBase<FIOSPlatformProperties>::SupportsFeature(Feature);
 	}
 
-	virtual bool IsSdkInstalled(bool bProjectHasCode, FString& OutDocumentationPath) const OVERRIDE;
+	virtual bool IsSdkInstalled(bool bProjectHasCode, FString& OutDocumentationPath) const override;
 
 
 #if WITH_ENGINE
-	virtual void GetReflectionCaptureFormats( TArray<FName>& OutFormats ) const OVERRIDE
+	virtual void GetReflectionCaptureFormats( TArray<FName>& OutFormats ) const override
 	{
 		OutFormats.Add(FName(TEXT("EncodedHDR")));
 	}
 
-	virtual void GetAllPossibleShaderFormats( TArray<FName>& OutFormats ) const OVERRIDE;
+	virtual void GetAllPossibleShaderFormats( TArray<FName>& OutFormats ) const override;
 
-	virtual void GetAllTargetedShaderFormats( TArray<FName>& OutFormats ) const OVERRIDE;
+	virtual void GetAllTargetedShaderFormats( TArray<FName>& OutFormats ) const override;
 
-	virtual const class FStaticMeshLODSettings& GetStaticMeshLODSettings( ) const OVERRIDE
+	virtual const class FStaticMeshLODSettings& GetStaticMeshLODSettings( ) const override
 	{
 		return StaticMeshLODSettings;
 	}
 
-	virtual void GetTextureFormats( const UTexture* Texture, TArray<FName>& OutFormats ) const OVERRIDE;
+	virtual void GetTextureFormats( const UTexture* Texture, TArray<FName>& OutFormats ) const override;
 
-	virtual const struct FTextureLODSettings& GetTextureLODSettings( ) const OVERRIDE;
+	virtual const struct FTextureLODSettings& GetTextureLODSettings( ) const override;
 
-	virtual FName GetWaveFormat( class USoundWave* Wave ) const OVERRIDE;
+	virtual FName GetWaveFormat( class USoundWave* Wave ) const override;
 #endif // WITH_ENGINE
 
 
 	DECLARE_DERIVED_EVENT(FIOSTargetPlatform, ITargetPlatform::FOnTargetDeviceDiscovered, FOnTargetDeviceDiscovered);
-	virtual FOnTargetDeviceDiscovered& OnDeviceDiscovered( ) OVERRIDE
+	virtual FOnTargetDeviceDiscovered& OnDeviceDiscovered( ) override
 	{
 		return DeviceDiscoveredEvent;
 	}
 
 	DECLARE_DERIVED_EVENT(FIOSTargetPlatform, ITargetPlatform::FOnTargetDeviceLost, FOnTargetDeviceLost);
-	virtual FOnTargetDeviceLost& OnDeviceLost( ) OVERRIDE
+	virtual FOnTargetDeviceLost& OnDeviceLost( ) override
 	{
 		return DeviceLostEvent;
 	}
diff --git a/Engine/Source/Developer/iOS/IOSTargetPlatform/Private/Windows/IOSDeviceHelperWindows.cpp b/Engine/Source/Developer/iOS/IOSTargetPlatform/Private/Windows/IOSDeviceHelperWindows.cpp
--- a/Engine/Source/Developer/iOS/IOSTargetPlatform/Private/Windows/IOSDeviceHelperWindows.cpp
+++ b/Engine/Source/Developer/iOS/IOSTargetPlatform/Private/Windows/IOSDeviceHelperWindows.cpp
@@ -46,12 +46,12 @@ public:
 		, bCheckDevices(true)
 	{}
 
-	virtual bool Init() OVERRIDE
+	virtual bool Init() override
 	{
 		return true;
 	}
 
-	virtual uint32 Run() OVERRIDE
+	virtual uint32 Run() override
 	{
 		while (!Stopping)
 		{
@@ -66,12 +66,12 @@ public:
 		return 0;
 	}
 
-	virtual void Stop() OVERRIDE
+	virtual void Stop() override
 	{
 		Stopping = true;
 	}
 
-	virtual void Exit() OVERRIDE
+	virtual void Exit() override
 	{}
 
 	FDeviceNotification& OnDeviceNotification()
diff --git a/Engine/Source/Editor/ActorPickerMode/Private/EditorModeActorPicker.h b/Engine/Source/Editor/ActorPickerMode/Private/EditorModeActorPicker.h
--- a/Engine/Source/Editor/ActorPickerMode/Private/EditorModeActorPicker.h
+++ b/Engine/Source/Editor/ActorPickerMode/Private/EditorModeActorPicker.h
@@ -28,14 +28,14 @@ public:
 	static void UnregisterEditMode();
 
 	/** Begin FEdMode interface */
-	virtual void Tick(FLevelEditorViewportClient* ViewportClient, float DeltaTime) OVERRIDE;
-	virtual bool MouseEnter(FLevelEditorViewportClient* ViewportClient, FViewport* Viewport,int32 x, int32 y) OVERRIDE;
-	virtual bool MouseLeave(FLevelEditorViewportClient* ViewportClient, FViewport* Viewport) OVERRIDE;
-	virtual bool MouseMove(FLevelEditorViewportClient* ViewportClient, FViewport* Viewport, int32 x, int32 y) OVERRIDE;
-	virtual bool LostFocus(FLevelEditorViewportClient* ViewportClient, FViewport* Viewport) OVERRIDE;
-	virtual bool InputKey(FLevelEditorViewportClient* ViewportClient, FViewport* Viewport, FKey Key, EInputEvent Event) OVERRIDE;
-	virtual bool GetCursor(EMouseCursor::Type& OutCursor) const OVERRIDE;
-	virtual bool UsesToolkits() const OVERRIDE;
+	virtual void Tick(FLevelEditorViewportClient* ViewportClient, float DeltaTime) override;
+	virtual bool MouseEnter(FLevelEditorViewportClient* ViewportClient, FViewport* Viewport,int32 x, int32 y) override;
+	virtual bool MouseLeave(FLevelEditorViewportClient* ViewportClient, FViewport* Viewport) override;
+	virtual bool MouseMove(FLevelEditorViewportClient* ViewportClient, FViewport* Viewport, int32 x, int32 y) override;
+	virtual bool LostFocus(FLevelEditorViewportClient* ViewportClient, FViewport* Viewport) override;
+	virtual bool InputKey(FLevelEditorViewportClient* ViewportClient, FViewport* Viewport, FKey Key, EInputEvent Event) override;
+	virtual bool GetCursor(EMouseCursor::Type& OutCursor) const override;
+	virtual bool UsesToolkits() const override;
 	/** End FEdMode interface */
 
 	/** Start up the mode, reset state etc. */
diff --git a/Engine/Source/Editor/AnimGraph/Classes/AnimGraphNode_ApplyAdditive.h b/Engine/Source/Editor/AnimGraph/Classes/AnimGraphNode_ApplyAdditive.h
--- a/Engine/Source/Editor/AnimGraph/Classes/AnimGraphNode_ApplyAdditive.h
+++ b/Engine/Source/Editor/AnimGraph/Classes/AnimGraphNode_ApplyAdditive.h
@@ -14,12 +14,12 @@ class UAnimGraphNode_ApplyAdditive : public UAnimGraphNode_Base
 	FAnimNode_ApplyAdditive Node;
 
 	// UEdGraphNode interface
-	virtual FString GetTooltip() const OVERRIDE;
-	virtual FText GetNodeTitle(ENodeTitleType::Type TitleType) const OVERRIDE;
-	virtual FLinearColor GetNodeTitleColor() const OVERRIDE;
+	virtual FString GetTooltip() const override;
+	virtual FText GetNodeTitle(ENodeTitleType::Type TitleType) const override;
+	virtual FLinearColor GetNodeTitleColor() const override;
 	// End of UEdGraphNode interface
 
 	// UAnimGraphNode_Base interface
-	virtual FString GetNodeCategory() const OVERRIDE;
+	virtual FString GetNodeCategory() const override;
 	// End of UAnimGraphNode_Base interface
 };
diff --git a/Engine/Source/Editor/AnimGraph/Classes/AnimGraphNode_Base.h b/Engine/Source/Editor/AnimGraph/Classes/AnimGraphNode_Base.h
--- a/Engine/Source/Editor/AnimGraph/Classes/AnimGraphNode_Base.h
+++ b/Engine/Source/Editor/AnimGraph/Classes/AnimGraphNode_Base.h
@@ -94,25 +94,25 @@ class UAnimGraphNode_Base : public UK2Node
 	TArray<FOptionalPinFromProperty> ShowPinForProperties;
 
 	// UObject interface
-	ANIMGRAPH_API virtual void PostEditChangeProperty(struct FPropertyChangedEvent& PropertyChangedEvent) OVERRIDE;
+	ANIMGRAPH_API virtual void PostEditChangeProperty(struct FPropertyChangedEvent& PropertyChangedEvent) override;
 	// End of UObject interface
 
 	// UEdGraphNode interface
-	ANIMGRAPH_API virtual void AllocateDefaultPins() OVERRIDE;
-	ANIMGRAPH_API virtual FLinearColor GetNodeTitleColor() const OVERRIDE;
-	ANIMGRAPH_API virtual FString GetDocumentationLink() const OVERRIDE;
-	ANIMGRAPH_API virtual void GetPinHoverText(const UEdGraphPin& Pin, FString& HoverTextOut) const OVERRIDE;
-	ANIMGRAPH_API virtual bool ShowPaletteIconOnNode() const OVERRIDE{ return false; }
+	ANIMGRAPH_API virtual void AllocateDefaultPins() override;
+	ANIMGRAPH_API virtual FLinearColor GetNodeTitleColor() const override;
+	ANIMGRAPH_API virtual FString GetDocumentationLink() const override;
+	ANIMGRAPH_API virtual void GetPinHoverText(const UEdGraphPin& Pin, FString& HoverTextOut) const override;
+	ANIMGRAPH_API virtual bool ShowPaletteIconOnNode() const override{ return false; }
 	// End of UEdGraphNode interface
 
 	// UK2Node interface
-	virtual bool NodeCausesStructuralBlueprintChange() const OVERRIDE { return true; }
-	virtual bool ShouldShowNodeProperties() const OVERRIDE { return true; }
-	virtual bool CanPlaceBreakpoints() const OVERRIDE { return false; }
-	ANIMGRAPH_API virtual void ReallocatePinsDuringReconstruction(TArray<UEdGraphPin*>& OldPins) OVERRIDE;
-	ANIMGRAPH_API virtual bool CanCreateUnderSpecifiedSchema(const UEdGraphSchema* DesiredSchema) const OVERRIDE;
-	ANIMGRAPH_API virtual void GetMenuEntries(FGraphContextMenuBuilder& ContextMenuBuilder) const OVERRIDE;
-	ANIMGRAPH_API virtual void GetNodeAttributes(TArray<TKeyValuePair<FString, FString>>& OutNodeAttributes) const OVERRIDE;
+	virtual bool NodeCausesStructuralBlueprintChange() const override { return true; }
+	virtual bool ShouldShowNodeProperties() const override { return true; }
+	virtual bool CanPlaceBreakpoints() const override { return false; }
+	ANIMGRAPH_API virtual void ReallocatePinsDuringReconstruction(TArray<UEdGraphPin*>& OldPins) override;
+	ANIMGRAPH_API virtual bool CanCreateUnderSpecifiedSchema(const UEdGraphSchema* DesiredSchema) const override;
+	ANIMGRAPH_API virtual void GetMenuEntries(FGraphContextMenuBuilder& ContextMenuBuilder) const override;
+	ANIMGRAPH_API virtual void GetNodeAttributes(TArray<TKeyValuePair<FString, FString>>& OutNodeAttributes) const override;
 	// End of UK2Node interface
 
 	// UAnimGraphNode_Base interface
diff --git a/Engine/Source/Editor/AnimGraph/Classes/AnimGraphNode_BlendListBase.h b/Engine/Source/Editor/AnimGraph/Classes/AnimGraphNode_BlendListBase.h
--- a/Engine/Source/Editor/AnimGraph/Classes/AnimGraphNode_BlendListBase.h
+++ b/Engine/Source/Editor/AnimGraph/Classes/AnimGraphNode_BlendListBase.h
@@ -10,13 +10,13 @@ class UAnimGraphNode_BlendListBase : public UAnimGraphNode_Base
 	GENERATED_UCLASS_BODY()
 
 	// UEdGraphNode interface
-	virtual FLinearColor GetNodeTitleColor() const OVERRIDE;
-	virtual void PostEditChangeProperty(struct FPropertyChangedEvent& PropertyChangedEvent) OVERRIDE;
-	virtual void ReallocatePinsDuringReconstruction(TArray<UEdGraphPin*>& OldPins) OVERRIDE;
+	virtual FLinearColor GetNodeTitleColor() const override;
+	virtual void PostEditChangeProperty(struct FPropertyChangedEvent& PropertyChangedEvent) override;
+	virtual void ReallocatePinsDuringReconstruction(TArray<UEdGraphPin*>& OldPins) override;
 	// End of UEdGraphNode interface
 
 	// UAnimGraphNode_Base interface
-	virtual FString GetNodeCategory() const OVERRIDE;
+	virtual FString GetNodeCategory() const override;
 	// End of UAnimGraphNode_Base interface
 
 protected:
diff --git a/Engine/Source/Editor/AnimGraph/Classes/AnimGraphNode_BlendListByBool.h b/Engine/Source/Editor/AnimGraph/Classes/AnimGraphNode_BlendListByBool.h
--- a/Engine/Source/Editor/AnimGraph/Classes/AnimGraphNode_BlendListByBool.h
+++ b/Engine/Source/Editor/AnimGraph/Classes/AnimGraphNode_BlendListByBool.h
@@ -14,11 +14,11 @@ class UAnimGraphNode_BlendListByBool : public UAnimGraphNode_BlendListBase
 	FAnimNode_BlendListByBool Node;
 
 	// UEdGraphNode interface
-	virtual FString GetTooltip() const OVERRIDE;
-	virtual FText GetNodeTitle(ENodeTitleType::Type TitleType) const OVERRIDE;
+	virtual FString GetTooltip() const override;
+	virtual FText GetNodeTitle(ENodeTitleType::Type TitleType) const override;
 	// End of UEdGraphNode interface
 
 	// UAnimGraphNode_Base interface
-	virtual void CustomizePinData(UEdGraphPin* Pin, FName SourcePropertyName, int32 ArrayIndex) const OVERRIDE;
+	virtual void CustomizePinData(UEdGraphPin* Pin, FName SourcePropertyName, int32 ArrayIndex) const override;
 	// End of UAnimGraphNode_Base interface
 };
diff --git a/Engine/Source/Editor/AnimGraph/Classes/AnimGraphNode_BlendListByEnum.h b/Engine/Source/Editor/AnimGraph/Classes/AnimGraphNode_BlendListByEnum.h
--- a/Engine/Source/Editor/AnimGraph/Classes/AnimGraphNode_BlendListByEnum.h
+++ b/Engine/Source/Editor/AnimGraph/Classes/AnimGraphNode_BlendListByEnum.h
@@ -22,31 +22,31 @@ protected:
 	TArray<FName> VisibleEnumEntries;
 public:
 	// UEdGraphNode interface
-	virtual FString GetTooltip() const OVERRIDE;
-	virtual FText GetNodeTitle(ENodeTitleType::Type TitleType) const OVERRIDE;
-	virtual FString GetNodeNativeTitle(ENodeTitleType::Type TitleType) const OVERRIDE;
-	virtual void PostPlacedNewNode() OVERRIDE;
+	virtual FString GetTooltip() const override;
+	virtual FText GetNodeTitle(ENodeTitleType::Type TitleType) const override;
+	virtual FString GetNodeNativeTitle(ENodeTitleType::Type TitleType) const override;
+	virtual void PostPlacedNewNode() override;
 	// End of UEdGraphNode interface
 
 	// UK2Node interface
-	virtual void GetContextMenuActions(const FGraphNodeContextMenuBuilder& Context) const OVERRIDE;
+	virtual void GetContextMenuActions(const FGraphNodeContextMenuBuilder& Context) const override;
 	// End of UK2Node interface
 
 	// UAnimGraphNode_Base interface
-	virtual FString GetNodeCategory() const OVERRIDE;
-	virtual void GetMenuEntries(FGraphContextMenuBuilder& ContextMenuBuilder) const OVERRIDE;
-	virtual void CustomizePinData(UEdGraphPin* Pin, FName SourcePropertyName, int32 ArrayIndex) const OVERRIDE;
-	virtual void ValidateAnimNodeDuringCompilation(class USkeleton* ForSkeleton, class FCompilerResultsLog& MessageLog) OVERRIDE;
-	virtual void BakeDataDuringCompilation(class FCompilerResultsLog& MessageLog) OVERRIDE;
-	virtual void PreloadRequiredAssets() OVERRIDE;
+	virtual FString GetNodeCategory() const override;
+	virtual void GetMenuEntries(FGraphContextMenuBuilder& ContextMenuBuilder) const override;
+	virtual void CustomizePinData(UEdGraphPin* Pin, FName SourcePropertyName, int32 ArrayIndex) const override;
+	virtual void ValidateAnimNodeDuringCompilation(class USkeleton* ForSkeleton, class FCompilerResultsLog& MessageLog) override;
+	virtual void BakeDataDuringCompilation(class FCompilerResultsLog& MessageLog) override;
+	virtual void PreloadRequiredAssets() override;
 	// End of UAnimGraphNode_Base interface
 
 	//@TODO: Generalize this behavior (returning a list of actions/delegates maybe?)
 	virtual void RemovePinFromBlendList(UEdGraphPin* Pin);
 
 	// INodeDependingOnEnumInterface
-	virtual class UEnum* GetEnum() const OVERRIDE { return BoundEnum; }
-	virtual bool ShouldBeReconstructedAfterEnumChanged() const OVERRIDE {return true;}
+	virtual class UEnum* GetEnum() const override { return BoundEnum; }
+	virtual bool ShouldBeReconstructedAfterEnumChanged() const override {return true;}
 	// End of INodeDependingOnEnumInterface
 protected:
 	// Exposes a pin corresponding to the specified element name
diff --git a/Engine/Source/Editor/AnimGraph/Classes/AnimGraphNode_BlendListByInt.h b/Engine/Source/Editor/AnimGraph/Classes/AnimGraphNode_BlendListByInt.h
--- a/Engine/Source/Editor/AnimGraph/Classes/AnimGraphNode_BlendListByInt.h
+++ b/Engine/Source/Editor/AnimGraph/Classes/AnimGraphNode_BlendListByInt.h
@@ -19,12 +19,12 @@ class UAnimGraphNode_BlendListByInt : public UAnimGraphNode_BlendListBase
 	ANIMGRAPH_API virtual void RemovePinFromBlendList(UEdGraphPin* Pin);
 
 	// UEdGraphNode interface
-	virtual FString GetTooltip() const OVERRIDE;
-	virtual FText GetNodeTitle(ENodeTitleType::Type TitleType) const OVERRIDE;
-	virtual void PostPlacedNewNode() OVERRIDE;
+	virtual FString GetTooltip() const override;
+	virtual FText GetNodeTitle(ENodeTitleType::Type TitleType) const override;
+	virtual void PostPlacedNewNode() override;
 	// End of UEdGraphNode interface
 
 	// UK2Node interface
-	virtual void GetContextMenuActions(const FGraphNodeContextMenuBuilder& Context) const OVERRIDE;
+	virtual void GetContextMenuActions(const FGraphNodeContextMenuBuilder& Context) const override;
 	// End of UK2Node interface
 };
diff --git a/Engine/Source/Editor/AnimGraph/Classes/AnimGraphNode_BlendSpaceBase.h b/Engine/Source/Editor/AnimGraph/Classes/AnimGraphNode_BlendSpaceBase.h
--- a/Engine/Source/Editor/AnimGraph/Classes/AnimGraphNode_BlendSpaceBase.h
+++ b/Engine/Source/Editor/AnimGraph/Classes/AnimGraphNode_BlendSpaceBase.h
@@ -10,12 +10,12 @@ class UAnimGraphNode_BlendSpaceBase : public UAnimGraphNode_Base
 	GENERATED_UCLASS_BODY()
 
 	// UEdGraphNode interface
-	virtual FLinearColor GetNodeTitleColor() const OVERRIDE;
+	virtual FLinearColor GetNodeTitleColor() const override;
 	// End of UEdGraphNode interface
 
 	// UAnimGraphNode_Base interface
-	virtual void CustomizePinData(UEdGraphPin* Pin, FName SourcePropertyName, int32 ArrayIndex) const OVERRIDE;
-	virtual void PreloadRequiredAssets() OVERRIDE;
+	virtual void CustomizePinData(UEdGraphPin* Pin, FName SourcePropertyName, int32 ArrayIndex) const override;
+	virtual void PreloadRequiredAssets() override;
 	// End of UAnimGraphNode_Base interface
 
 protected:
diff --git a/Engine/Source/Editor/AnimGraph/Classes/AnimGraphNode_BlendSpaceEvaluator.h b/Engine/Source/Editor/AnimGraph/Classes/AnimGraphNode_BlendSpaceEvaluator.h
--- a/Engine/Source/Editor/AnimGraph/Classes/AnimGraphNode_BlendSpaceEvaluator.h
+++ b/Engine/Source/Editor/AnimGraph/Classes/AnimGraphNode_BlendSpaceEvaluator.h
@@ -17,23 +17,23 @@ class UAnimGraphNode_BlendSpaceEvaluator : public UAnimGraphNode_BlendSpaceBase
 	FAnimationGroupReference SyncGroup;
 
 	// UEdGraphNode interface
-	virtual FString GetTooltip() const OVERRIDE;
-	virtual FText GetNodeTitle(ENodeTitleType::Type TitleType) const OVERRIDE;
-	virtual FString GetNodeNativeTitle(ENodeTitleType::Type TitleType) const OVERRIDE;
+	virtual FString GetTooltip() const override;
+	virtual FText GetNodeTitle(ENodeTitleType::Type TitleType) const override;
+	virtual FString GetNodeNativeTitle(ENodeTitleType::Type TitleType) const override;
 	// End of UEdGraphNode interface
 
 	// UAnimGraphNode_Base interface
-	virtual void GetMenuEntries(FGraphContextMenuBuilder& ContextMenuBuilder) const OVERRIDE;
-	virtual void ValidateAnimNodeDuringCompilation(class USkeleton* ForSkeleton, class FCompilerResultsLog& MessageLog) OVERRIDE;
-	virtual void BakeDataDuringCompilation(class FCompilerResultsLog& MessageLog) OVERRIDE;
+	virtual void GetMenuEntries(FGraphContextMenuBuilder& ContextMenuBuilder) const override;
+	virtual void ValidateAnimNodeDuringCompilation(class USkeleton* ForSkeleton, class FCompilerResultsLog& MessageLog) override;
+	virtual void BakeDataDuringCompilation(class FCompilerResultsLog& MessageLog) override;
 	// Interface to support transition getter
-	virtual bool DoesSupportTimeForTransitionGetter() const OVERRIDE;
-	virtual UAnimationAsset* GetAnimationAsset() const OVERRIDE;
-	virtual const TCHAR* GetTimePropertyName() const OVERRIDE;
-	virtual UScriptStruct* GetTimePropertyStruct() const OVERRIDE;
+	virtual bool DoesSupportTimeForTransitionGetter() const override;
+	virtual UAnimationAsset* GetAnimationAsset() const override;
+	virtual const TCHAR* GetTimePropertyName() const override;
+	virtual UScriptStruct* GetTimePropertyStruct() const override;
 	// End of UAnimGraphNode_Base interface
 
 	// UK2Node interface
-	virtual void GetContextMenuActions(const FGraphNodeContextMenuBuilder& Context) const OVERRIDE;
+	virtual void GetContextMenuActions(const FGraphNodeContextMenuBuilder& Context) const override;
 	// End of UK2Node interface
 };
diff --git a/Engine/Source/Editor/AnimGraph/Classes/AnimGraphNode_BlendSpacePlayer.h b/Engine/Source/Editor/AnimGraph/Classes/AnimGraphNode_BlendSpacePlayer.h
--- a/Engine/Source/Editor/AnimGraph/Classes/AnimGraphNode_BlendSpacePlayer.h
+++ b/Engine/Source/Editor/AnimGraph/Classes/AnimGraphNode_BlendSpacePlayer.h
@@ -18,24 +18,24 @@ class UAnimGraphNode_BlendSpacePlayer : public UAnimGraphNode_BlendSpaceBase
 	FAnimationGroupReference SyncGroup;
 
 	// UEdGraphNode interface
-	virtual FString GetTooltip() const OVERRIDE;
-	virtual FText GetNodeTitle(ENodeTitleType::Type TitleType) const OVERRIDE;
-	virtual FString GetNodeNativeTitle(ENodeTitleType::Type TitleType) const OVERRIDE;
+	virtual FString GetTooltip() const override;
+	virtual FText GetNodeTitle(ENodeTitleType::Type TitleType) const override;
+	virtual FString GetNodeNativeTitle(ENodeTitleType::Type TitleType) const override;
 	// End of UEdGraphNode interface
 
 	// UAnimGraphNode_Base interface
-	virtual void GetMenuEntries(FGraphContextMenuBuilder& ContextMenuBuilder) const OVERRIDE;
-	virtual void ValidateAnimNodeDuringCompilation(class USkeleton* ForSkeleton, class FCompilerResultsLog& MessageLog) OVERRIDE;
-	virtual void BakeDataDuringCompilation(class FCompilerResultsLog& MessageLog) OVERRIDE;
-	virtual bool DoesSupportTimeForTransitionGetter() const OVERRIDE;
-	virtual UAnimationAsset* GetAnimationAsset() const OVERRIDE;
-	virtual const TCHAR* GetTimePropertyName() const OVERRIDE;
-	virtual UScriptStruct* GetTimePropertyStruct() const OVERRIDE;
-	virtual void GetAllAnimationSequencesReferred(TArray<UAnimationAsset*>& ComplexAnims, TArray<UAnimSequence*>& AnimationSequences) const OVERRIDE;
-	virtual void ReplaceReferredAnimations(const TMap<UAnimationAsset*, UAnimationAsset*>& ComplexAnimsMap, const TMap<UAnimSequence*, UAnimSequence*>& AnimSequenceMap) OVERRIDE;
+	virtual void GetMenuEntries(FGraphContextMenuBuilder& ContextMenuBuilder) const override;
+	virtual void ValidateAnimNodeDuringCompilation(class USkeleton* ForSkeleton, class FCompilerResultsLog& MessageLog) override;
+	virtual void BakeDataDuringCompilation(class FCompilerResultsLog& MessageLog) override;
+	virtual bool DoesSupportTimeForTransitionGetter() const override;
+	virtual UAnimationAsset* GetAnimationAsset() const override;
+	virtual const TCHAR* GetTimePropertyName() const override;
+	virtual UScriptStruct* GetTimePropertyStruct() const override;
+	virtual void GetAllAnimationSequencesReferred(TArray<UAnimationAsset*>& ComplexAnims, TArray<UAnimSequence*>& AnimationSequences) const override;
+	virtual void ReplaceReferredAnimations(const TMap<UAnimationAsset*, UAnimationAsset*>& ComplexAnimsMap, const TMap<UAnimSequence*, UAnimSequence*>& AnimSequenceMap) override;
 	// End of UAnimGraphNode_Base interface
 
 	// UK2Node interface
-	virtual void GetContextMenuActions(const FGraphNodeContextMenuBuilder& Context) const OVERRIDE;
+	virtual void GetContextMenuActions(const FGraphNodeContextMenuBuilder& Context) const override;
 	// End of UK2Node interface
 };
diff --git a/Engine/Source/Editor/AnimGraph/Classes/AnimGraphNode_ComponentToLocalSpace.h b/Engine/Source/Editor/AnimGraph/Classes/AnimGraphNode_ComponentToLocalSpace.h
--- a/Engine/Source/Editor/AnimGraph/Classes/AnimGraphNode_ComponentToLocalSpace.h
+++ b/Engine/Source/Editor/AnimGraph/Classes/AnimGraphNode_ComponentToLocalSpace.h
@@ -14,13 +14,13 @@ class UAnimGraphNode_ComponentToLocalSpace : public UAnimGraphNode_Base
 	FAnimNode_ConvertComponentToLocalSpace Node;
 
 	// UEdGraphNode interface
-	virtual FLinearColor GetNodeTitleColor() const OVERRIDE;
-	virtual FString GetTooltip() const OVERRIDE;
-	virtual FText GetNodeTitle(ENodeTitleType::Type TitleType) const OVERRIDE;
+	virtual FLinearColor GetNodeTitleColor() const override;
+	virtual FString GetTooltip() const override;
+	virtual FText GetNodeTitle(ENodeTitleType::Type TitleType) const override;
 	// End of UEdGraphNode interface
 
 	// UAnimGraphNode_Base interface
-	virtual FString GetNodeCategory() const OVERRIDE;
-	virtual void PostProcessPinName(const UEdGraphPin* Pin, FString& DisplayName) const OVERRIDE;
+	virtual FString GetNodeCategory() const override;
+	virtual void PostProcessPinName(const UEdGraphPin* Pin, FString& DisplayName) const override;
 	// End of UAnimGraphNode_Base interface
 };
diff --git a/Engine/Source/Editor/AnimGraph/Classes/AnimGraphNode_CopyBone.h b/Engine/Source/Editor/AnimGraph/Classes/AnimGraphNode_CopyBone.h
--- a/Engine/Source/Editor/AnimGraph/Classes/AnimGraphNode_CopyBone.h
+++ b/Engine/Source/Editor/AnimGraph/Classes/AnimGraphNode_CopyBone.h
@@ -15,13 +15,13 @@ class UAnimGraphNode_CopyBone : public UAnimGraphNode_SkeletalControlBase
 
 public:
 	// UEdGraphNode interface
-	virtual FText GetNodeTitle(ENodeTitleType::Type TitleType) const OVERRIDE;
-	virtual FString GetNodeNativeTitle(ENodeTitleType::Type TitleType) const OVERRIDE;
-	virtual FString GetTooltip() const OVERRIDE;
+	virtual FText GetNodeTitle(ENodeTitleType::Type TitleType) const override;
+	virtual FString GetNodeNativeTitle(ENodeTitleType::Type TitleType) const override;
+	virtual FString GetTooltip() const override;
 	// End of UEdGraphNode interface
 
 protected:
 	// UAnimGraphNode_SkeletalControlBase interface
-	virtual FText GetControllerDescription() const OVERRIDE;
+	virtual FText GetControllerDescription() const override;
 	// End of UAnimGraphNode_SkeletalControlBase interface
 };
diff --git a/Engine/Source/Editor/AnimGraph/Classes/AnimGraphNode_CustomTransitionResult.h b/Engine/Source/Editor/AnimGraph/Classes/AnimGraphNode_CustomTransitionResult.h
--- a/Engine/Source/Editor/AnimGraph/Classes/AnimGraphNode_CustomTransitionResult.h
+++ b/Engine/Source/Editor/AnimGraph/Classes/AnimGraphNode_CustomTransitionResult.h
@@ -10,8 +10,8 @@ class UAnimGraphNode_CustomTransitionResult : public UAnimGraphNode_Root
 	GENERATED_UCLASS_BODY()
 
 	// UEdGraphNode interface
-	virtual FLinearColor GetNodeTitleColor() const OVERRIDE;
-	virtual FString GetTooltip() const OVERRIDE;
-	virtual FText GetNodeTitle(ENodeTitleType::Type TitleType) const OVERRIDE;
+	virtual FLinearColor GetNodeTitleColor() const override;
+	virtual FString GetTooltip() const override;
+	virtual FText GetNodeTitle(ENodeTitleType::Type TitleType) const override;
 	// End of UEdGraphNode interface
 };
diff --git a/Engine/Source/Editor/AnimGraph/Classes/AnimGraphNode_Fabrik.h b/Engine/Source/Editor/AnimGraph/Classes/AnimGraphNode_Fabrik.h
--- a/Engine/Source/Editor/AnimGraph/Classes/AnimGraphNode_Fabrik.h
+++ b/Engine/Source/Editor/AnimGraph/Classes/AnimGraphNode_Fabrik.h
@@ -16,11 +16,11 @@ class UAnimGraphNode_Fabrik : public UAnimGraphNode_SkeletalControlBase
 
 public:
 	// UEdGraphNode interface
-	virtual FText GetNodeTitle(ENodeTitleType::Type TitleType) const OVERRIDE;
+	virtual FText GetNodeTitle(ENodeTitleType::Type TitleType) const override;
 	// End of UEdGraphNode interface
 
 protected:
 	// UAnimGraphNode_SkeletalControlBase interface
-	virtual FText GetControllerDescription() const OVERRIDE;
+	virtual FText GetControllerDescription() const override;
 	// End of UAnimGraphNode_SkeletalControlBase interface
 };
\ No newline at end of file
diff --git a/Engine/Source/Editor/AnimGraph/Classes/AnimGraphNode_HandIKRetargeting.h b/Engine/Source/Editor/AnimGraph/Classes/AnimGraphNode_HandIKRetargeting.h
--- a/Engine/Source/Editor/AnimGraph/Classes/AnimGraphNode_HandIKRetargeting.h
+++ b/Engine/Source/Editor/AnimGraph/Classes/AnimGraphNode_HandIKRetargeting.h
@@ -15,13 +15,13 @@ class UAnimGraphNode_HandIKRetargeting : public UAnimGraphNode_SkeletalControlBa
 
 public:
 	// UEdGraphNode interface
-	virtual FText GetNodeTitle(ENodeTitleType::Type TitleType) const OVERRIDE;
-	virtual FString GetNodeNativeTitle(ENodeTitleType::Type TitleType) const OVERRIDE;
-	virtual FString GetTooltip() const OVERRIDE;
+	virtual FText GetNodeTitle(ENodeTitleType::Type TitleType) const override;
+	virtual FString GetNodeNativeTitle(ENodeTitleType::Type TitleType) const override;
+	virtual FString GetTooltip() const override;
 	// End of UEdGraphNode interface
 
 protected:
 	// UAnimGraphNode_SkeletalControlBase interface
-	virtual FText GetControllerDescription() const OVERRIDE;
+	virtual FText GetControllerDescription() const override;
 	// End of UAnimGraphNode_SkeletalControlBase interface
 };
diff --git a/Engine/Source/Editor/AnimGraph/Classes/AnimGraphNode_IdentityPose.h b/Engine/Source/Editor/AnimGraph/Classes/AnimGraphNode_IdentityPose.h
--- a/Engine/Source/Editor/AnimGraph/Classes/AnimGraphNode_IdentityPose.h
+++ b/Engine/Source/Editor/AnimGraph/Classes/AnimGraphNode_IdentityPose.h
@@ -10,7 +10,7 @@ class UAnimGraphNode_IdentityPose : public UAnimGraphNode_RefPoseBase
 	GENERATED_UCLASS_BODY()
 
 	// UEdGraphNode interface
-	virtual FString GetTooltip() const OVERRIDE;
-	virtual FText GetNodeTitle(ENodeTitleType::Type TitleType) const OVERRIDE;
+	virtual FString GetTooltip() const override;
+	virtual FText GetNodeTitle(ENodeTitleType::Type TitleType) const override;
 	// End of UEdGraphNode interface
 };
diff --git a/Engine/Source/Editor/AnimGraph/Classes/AnimGraphNode_LayeredBoneBlend.h b/Engine/Source/Editor/AnimGraph/Classes/AnimGraphNode_LayeredBoneBlend.h
--- a/Engine/Source/Editor/AnimGraph/Classes/AnimGraphNode_LayeredBoneBlend.h
+++ b/Engine/Source/Editor/AnimGraph/Classes/AnimGraphNode_LayeredBoneBlend.h
@@ -19,16 +19,16 @@ class UAnimGraphNode_LayeredBoneBlend : public UAnimGraphNode_BlendListBase
 	ANIMGRAPH_API virtual void RemovePinFromBlendByFilter(UEdGraphPin* Pin);
 
 	// UEdGraphNode interface
-	virtual FLinearColor GetNodeTitleColor() const OVERRIDE;
-	virtual FString GetTooltip() const OVERRIDE;
-	virtual FText GetNodeTitle(ENodeTitleType::Type TitleType) const OVERRIDE;
+	virtual FLinearColor GetNodeTitleColor() const override;
+	virtual FString GetTooltip() const override;
+	virtual FText GetNodeTitle(ENodeTitleType::Type TitleType) const override;
 	// End of UEdGraphNode interface
 
 	// UK2Node interface
-	virtual void GetContextMenuActions(const FGraphNodeContextMenuBuilder& Context) const OVERRIDE;
+	virtual void GetContextMenuActions(const FGraphNodeContextMenuBuilder& Context) const override;
 	// End of UK2Node interface
 
 	// UAnimGraphNode_Base interface
-	virtual FString GetNodeCategory() const OVERRIDE;
+	virtual FString GetNodeCategory() const override;
 	// End of UAnimGraphNode_Base interface
 };
diff --git a/Engine/Source/Editor/AnimGraph/Classes/AnimGraphNode_LocalRefPose.h b/Engine/Source/Editor/AnimGraph/Classes/AnimGraphNode_LocalRefPose.h
--- a/Engine/Source/Editor/AnimGraph/Classes/AnimGraphNode_LocalRefPose.h
+++ b/Engine/Source/Editor/AnimGraph/Classes/AnimGraphNode_LocalRefPose.h
@@ -10,7 +10,7 @@ class UAnimGraphNode_LocalRefPose : public UAnimGraphNode_RefPoseBase
 	GENERATED_UCLASS_BODY()
 
 	// UEdGraphNode interface
-	virtual FString GetTooltip() const OVERRIDE;
-	virtual FText GetNodeTitle(ENodeTitleType::Type TitleType) const OVERRIDE;
+	virtual FString GetTooltip() const override;
+	virtual FText GetNodeTitle(ENodeTitleType::Type TitleType) const override;
 	// End of UEdGraphNode interface
 };
diff --git a/Engine/Source/Editor/AnimGraph/Classes/AnimGraphNode_LocalToComponentSpace.h b/Engine/Source/Editor/AnimGraph/Classes/AnimGraphNode_LocalToComponentSpace.h
--- a/Engine/Source/Editor/AnimGraph/Classes/AnimGraphNode_LocalToComponentSpace.h
+++ b/Engine/Source/Editor/AnimGraph/Classes/AnimGraphNode_LocalToComponentSpace.h
@@ -14,14 +14,14 @@ class UAnimGraphNode_LocalToComponentSpace : public UAnimGraphNode_Base
 	FAnimNode_ConvertLocalToComponentSpace Node;
 
 	// UEdGraphNode interface
-	virtual FLinearColor GetNodeTitleColor() const OVERRIDE;
-	virtual FString GetTooltip() const OVERRIDE;
-	virtual FText GetNodeTitle(ENodeTitleType::Type TitleType) const OVERRIDE;
+	virtual FLinearColor GetNodeTitleColor() const override;
+	virtual FString GetTooltip() const override;
+	virtual FText GetNodeTitle(ENodeTitleType::Type TitleType) const override;
 	// End of UEdGraphNode interface
 
 	// UAnimGraphNode_Base interface
-	virtual FString GetNodeCategory() const OVERRIDE;
-	virtual void CreateOutputPins() OVERRIDE;
-	virtual void PostProcessPinName(const UEdGraphPin* Pin, FString& DisplayName) const OVERRIDE;
+	virtual FString GetNodeCategory() const override;
+	virtual void CreateOutputPins() override;
+	virtual void PostProcessPinName(const UEdGraphPin* Pin, FString& DisplayName) const override;
 	// End of UAnimGraphNode_Base interface
 };
diff --git a/Engine/Source/Editor/AnimGraph/Classes/AnimGraphNode_LookAt.h b/Engine/Source/Editor/AnimGraph/Classes/AnimGraphNode_LookAt.h
--- a/Engine/Source/Editor/AnimGraph/Classes/AnimGraphNode_LookAt.h
+++ b/Engine/Source/Editor/AnimGraph/Classes/AnimGraphNode_LookAt.h
@@ -16,17 +16,17 @@ class UAnimGraphNode_LookAt : public UAnimGraphNode_SkeletalControlBase
 
 public:
 	// UEdGraphNode interface
-	virtual FText GetNodeTitle(ENodeTitleType::Type TitleType) const OVERRIDE;
-	virtual FString GetTooltip() const OVERRIDE;
-	virtual FString GetKeywords() const OVERRIDE;
+	virtual FText GetNodeTitle(ENodeTitleType::Type TitleType) const override;
+	virtual FString GetTooltip() const override;
+	virtual FString GetKeywords() const override;
 	// End of UEdGraphNode interface
 
 	// UAnimGraphNode_SkeletalControlBase interface
-	ANIMGRAPH_API virtual void Draw(FPrimitiveDrawInterface* PDI, USkeletalMeshComponent* SkelMeshComp) const OVERRIDE;
+	ANIMGRAPH_API virtual void Draw(FPrimitiveDrawInterface* PDI, USkeletalMeshComponent* SkelMeshComp) const override;
 	// End of UAnimGraphNode_SkeletalControlBase interface
 
 protected:
 	// UAnimGraphNode_SkeletalControlBase interface
-	virtual FText GetControllerDescription() const OVERRIDE;
+	virtual FText GetControllerDescription() const override;
 	// End of UAnimGraphNode_SkeletalControlBase interface
 };
diff --git a/Engine/Source/Editor/AnimGraph/Classes/AnimGraphNode_MeshRefPose.h b/Engine/Source/Editor/AnimGraph/Classes/AnimGraphNode_MeshRefPose.h
--- a/Engine/Source/Editor/AnimGraph/Classes/AnimGraphNode_MeshRefPose.h
+++ b/Engine/Source/Editor/AnimGraph/Classes/AnimGraphNode_MeshRefPose.h
@@ -14,13 +14,13 @@ class UAnimGraphNode_MeshRefPose : public UAnimGraphNode_Base
 	FAnimNode_MeshSpaceRefPose Node;
 
 	// UEdGraphNode interface
-	virtual FLinearColor GetNodeTitleColor() const OVERRIDE;
-	virtual FString GetTooltip() const OVERRIDE;
-	virtual FText GetNodeTitle(ENodeTitleType::Type TitleType) const OVERRIDE;
+	virtual FLinearColor GetNodeTitleColor() const override;
+	virtual FString GetTooltip() const override;
+	virtual FText GetNodeTitle(ENodeTitleType::Type TitleType) const override;
 	// End of UEdGraphNode interface
 
 	// UAnimGraphNode_Base interface
-	virtual FString GetNodeCategory() const OVERRIDE;
-	virtual void CreateOutputPins() OVERRIDE;
+	virtual FString GetNodeCategory() const override;
+	virtual void CreateOutputPins() override;
 	// End of UAnimGraphNode_Base interface
 };
diff --git a/Engine/Source/Editor/AnimGraph/Classes/AnimGraphNode_ModifyBone.h b/Engine/Source/Editor/AnimGraph/Classes/AnimGraphNode_ModifyBone.h
--- a/Engine/Source/Editor/AnimGraph/Classes/AnimGraphNode_ModifyBone.h
+++ b/Engine/Source/Editor/AnimGraph/Classes/AnimGraphNode_ModifyBone.h
@@ -16,13 +16,13 @@ class UAnimGraphNode_ModifyBone : public UAnimGraphNode_SkeletalControlBase
 
 public:
 	// UEdGraphNode interface
-	virtual FText GetNodeTitle(ENodeTitleType::Type TitleType) const OVERRIDE;
-	virtual FString GetTooltip() const OVERRIDE;
-	virtual FString GetKeywords() const OVERRIDE;
+	virtual FText GetNodeTitle(ENodeTitleType::Type TitleType) const override;
+	virtual FString GetTooltip() const override;
+	virtual FString GetKeywords() const override;
 	// End of UEdGraphNode interface
 
 protected:
 	// UAnimGraphNode_SkeletalControlBase interface
-	virtual FText GetControllerDescription() const OVERRIDE;
+	virtual FText GetControllerDescription() const override;
 	// End of UAnimGraphNode_SkeletalControlBase interface
 };
diff --git a/Engine/Source/Editor/AnimGraph/Classes/AnimGraphNode_RefPoseBase.h b/Engine/Source/Editor/AnimGraph/Classes/AnimGraphNode_RefPoseBase.h
--- a/Engine/Source/Editor/AnimGraph/Classes/AnimGraphNode_RefPoseBase.h
+++ b/Engine/Source/Editor/AnimGraph/Classes/AnimGraphNode_RefPoseBase.h
@@ -14,7 +14,7 @@ class UAnimGraphNode_RefPoseBase : public UAnimGraphNode_Base
 	FAnimNode_RefPose Node;
 
 	// UEdGraphNode interface
-	virtual FString GetNodeCategory() const OVERRIDE;
-	virtual FLinearColor GetNodeTitleColor() const OVERRIDE;
+	virtual FString GetNodeCategory() const override;
+	virtual FLinearColor GetNodeTitleColor() const override;
 	// End of UEdGraphNode interface
 };
diff --git a/Engine/Source/Editor/AnimGraph/Classes/AnimGraphNode_Root.h b/Engine/Source/Editor/AnimGraph/Classes/AnimGraphNode_Root.h
--- a/Engine/Source/Editor/AnimGraph/Classes/AnimGraphNode_Root.h
+++ b/Engine/Source/Editor/AnimGraph/Classes/AnimGraphNode_Root.h
@@ -14,19 +14,19 @@ class UAnimGraphNode_Root : public UAnimGraphNode_Base
 	FAnimNode_Root Node;
 
 	// Begin UEdGraphNode interface.
-	virtual FLinearColor GetNodeTitleColor() const OVERRIDE;
-	virtual FString GetTooltip() const OVERRIDE;
-	virtual FText GetNodeTitle(ENodeTitleType::Type TitleType) const OVERRIDE;
-	virtual bool CanUserDeleteNode() const OVERRIDE { return false; }
-	virtual bool CanDuplicateNode() const OVERRIDE { return false; }
-	virtual void GetMenuEntries(FGraphContextMenuBuilder& ContextMenuBuilder) const OVERRIDE;
+	virtual FLinearColor GetNodeTitleColor() const override;
+	virtual FString GetTooltip() const override;
+	virtual FText GetNodeTitle(ENodeTitleType::Type TitleType) const override;
+	virtual bool CanUserDeleteNode() const override { return false; }
+	virtual bool CanDuplicateNode() const override { return false; }
+	virtual void GetMenuEntries(FGraphContextMenuBuilder& ContextMenuBuilder) const override;
 	// End UEdGraphNode interface.
 
 	// UAnimGraphNode_Base interface
-	virtual bool IsSinkNode() const OVERRIDE;
+	virtual bool IsSinkNode() const override;
 
 	// Get the link to the documentation
-	virtual FString GetDocumentationLink() const OVERRIDE;
+	virtual FString GetDocumentationLink() const override;
 
 	// End of UAnimGraphNode_Base interface
 };
diff --git a/Engine/Source/Editor/AnimGraph/Classes/AnimGraphNode_RotateRootBone.h b/Engine/Source/Editor/AnimGraph/Classes/AnimGraphNode_RotateRootBone.h
--- a/Engine/Source/Editor/AnimGraph/Classes/AnimGraphNode_RotateRootBone.h
+++ b/Engine/Source/Editor/AnimGraph/Classes/AnimGraphNode_RotateRootBone.h
@@ -15,12 +15,12 @@ class UAnimGraphNode_RotateRootBone : public UAnimGraphNode_Base
 	FAnimNode_RotateRootBone Node;
 
 	// Begin UEdGraphNode interface.
-	virtual FLinearColor GetNodeTitleColor() const OVERRIDE;
-	virtual FString GetTooltip() const OVERRIDE;
-	virtual FText GetNodeTitle(ENodeTitleType::Type TitleType) const OVERRIDE;
+	virtual FLinearColor GetNodeTitleColor() const override;
+	virtual FString GetTooltip() const override;
+	virtual FText GetNodeTitle(ENodeTitleType::Type TitleType) const override;
 	// End UEdGraphNode interface.
 
 	// UAnimGraphNode_Base interface
-	virtual FString GetNodeCategory() const OVERRIDE;
+	virtual FString GetNodeCategory() const override;
 	// End of UAnimGraphNode_Base interface
 };
diff --git a/Engine/Source/Editor/AnimGraph/Classes/AnimGraphNode_RotationOffsetBlendSpace.h b/Engine/Source/Editor/AnimGraph/Classes/AnimGraphNode_RotationOffsetBlendSpace.h
--- a/Engine/Source/Editor/AnimGraph/Classes/AnimGraphNode_RotationOffsetBlendSpace.h
+++ b/Engine/Source/Editor/AnimGraph/Classes/AnimGraphNode_RotationOffsetBlendSpace.h
@@ -14,18 +14,18 @@ class UAnimGraphNode_RotationOffsetBlendSpace: public UAnimGraphNode_BlendSpaceB
 	FAnimNode_RotationOffsetBlendSpace Node;
 
 	// UEdGraphNode interface
-	virtual FString GetTooltip() const OVERRIDE;
-	virtual FText GetNodeTitle(ENodeTitleType::Type TitleType) const OVERRIDE;
-	virtual FString GetNodeNativeTitle(ENodeTitleType::Type TitleType) const OVERRIDE;
+	virtual FString GetTooltip() const override;
+	virtual FText GetNodeTitle(ENodeTitleType::Type TitleType) const override;
+	virtual FString GetNodeNativeTitle(ENodeTitleType::Type TitleType) const override;
 	// End of UEdGraphNode interface
 
 	// UAnimGraphNode_Base interface
-	virtual void GetMenuEntries(FGraphContextMenuBuilder& ContextMenuBuilder) const OVERRIDE;
-	virtual void ValidateAnimNodeDuringCompilation(class USkeleton* ForSkeleton, class FCompilerResultsLog& MessageLog) OVERRIDE;
-	virtual UAnimationAsset* GetAnimationAsset() const OVERRIDE { return Node.BlendSpace; }
+	virtual void GetMenuEntries(FGraphContextMenuBuilder& ContextMenuBuilder) const override;
+	virtual void ValidateAnimNodeDuringCompilation(class USkeleton* ForSkeleton, class FCompilerResultsLog& MessageLog) override;
+	virtual UAnimationAsset* GetAnimationAsset() const override { return Node.BlendSpace; }
 	// End of UAnimGraphNode_Base interface
 
 	// UK2Node interface
-	virtual void GetContextMenuActions(const FGraphNodeContextMenuBuilder& Context) const OVERRIDE;
+	virtual void GetContextMenuActions(const FGraphNodeContextMenuBuilder& Context) const override;
 	// End of UK2Node interface
 };
diff --git a/Engine/Source/Editor/AnimGraph/Classes/AnimGraphNode_SaveCachedPose.h b/Engine/Source/Editor/AnimGraph/Classes/AnimGraphNode_SaveCachedPose.h
--- a/Engine/Source/Editor/AnimGraph/Classes/AnimGraphNode_SaveCachedPose.h
+++ b/Engine/Source/Editor/AnimGraph/Classes/AnimGraphNode_SaveCachedPose.h
@@ -17,16 +17,16 @@ class UAnimGraphNode_SaveCachedPose : public UAnimGraphNode_Base
 	FString CacheName;
 
 	// UEdGraphNode interface
-	virtual FString GetTooltip() const OVERRIDE;
-	virtual FText GetNodeTitle(ENodeTitleType::Type TitleType) const OVERRIDE;
-	virtual FString GetNodeNativeTitle(ENodeTitleType::Type TitleType) const OVERRIDE;
-	virtual void OnRenameNode(const FString& NewName) OVERRIDE;
-	virtual TSharedPtr<class INameValidatorInterface> MakeNameValidator() const OVERRIDE;
+	virtual FString GetTooltip() const override;
+	virtual FText GetNodeTitle(ENodeTitleType::Type TitleType) const override;
+	virtual FString GetNodeNativeTitle(ENodeTitleType::Type TitleType) const override;
+	virtual void OnRenameNode(const FString& NewName) override;
+	virtual TSharedPtr<class INameValidatorInterface> MakeNameValidator() const override;
 	// End of UEdGraphNode interface
 
 	// UAnimGraphNode_Base interface
-	virtual FString GetNodeCategory() const OVERRIDE;
-	virtual void GetMenuEntries(FGraphContextMenuBuilder& ContextMenuBuilder) const OVERRIDE;
-	virtual bool IsSinkNode() const OVERRIDE { return true; }
+	virtual FString GetNodeCategory() const override;
+	virtual void GetMenuEntries(FGraphContextMenuBuilder& ContextMenuBuilder) const override;
+	virtual bool IsSinkNode() const override { return true; }
 	// End of UAnimGraphNode_Base interface
 };
diff --git a/Engine/Source/Editor/AnimGraph/Classes/AnimGraphNode_SequenceEvaluator.h b/Engine/Source/Editor/AnimGraph/Classes/AnimGraphNode_SequenceEvaluator.h
--- a/Engine/Source/Editor/AnimGraph/Classes/AnimGraphNode_SequenceEvaluator.h
+++ b/Engine/Source/Editor/AnimGraph/Classes/AnimGraphNode_SequenceEvaluator.h
@@ -14,25 +14,25 @@ class UAnimGraphNode_SequenceEvaluator : public UAnimGraphNode_Base
 	FAnimNode_SequenceEvaluator Node;
 
 	// UEdGraphNode interface
-	virtual FText GetNodeTitle(ENodeTitleType::Type TitleType) const OVERRIDE;
-	virtual FString GetNodeNativeTitle(ENodeTitleType::Type TitleType) const OVERRIDE;
-	virtual FString GetTooltip() const OVERRIDE;
+	virtual FText GetNodeTitle(ENodeTitleType::Type TitleType) const override;
+	virtual FString GetNodeNativeTitle(ENodeTitleType::Type TitleType) const override;
+	virtual FString GetTooltip() const override;
 	// End of UEdGraphNode
 
 	// UAnimGraphNode_Base interface
-	virtual void GetMenuEntries(FGraphContextMenuBuilder& ContextMenuBuilder) const OVERRIDE;
-	virtual void ValidateAnimNodeDuringCompilation(class USkeleton* ForSkeleton, class FCompilerResultsLog& MessageLog) OVERRIDE;
-	virtual void PreloadRequiredAssets() OVERRIDE;	
+	virtual void GetMenuEntries(FGraphContextMenuBuilder& ContextMenuBuilder) const override;
+	virtual void ValidateAnimNodeDuringCompilation(class USkeleton* ForSkeleton, class FCompilerResultsLog& MessageLog) override;
+	virtual void PreloadRequiredAssets() override;	
 	// Interface to support transition getter
-	virtual bool DoesSupportTimeForTransitionGetter() const OVERRIDE;
-	virtual UAnimationAsset* GetAnimationAsset() const OVERRIDE;
-	virtual const TCHAR* GetTimePropertyName() const OVERRIDE;
-	virtual UScriptStruct* GetTimePropertyStruct() const OVERRIDE;
-	virtual void GetAllAnimationSequencesReferred(TArray<UAnimationAsset*>& ComplexAnims, TArray<UAnimSequence*>& AnimationSequences) const OVERRIDE;
-	virtual void ReplaceReferredAnimations(const TMap<UAnimationAsset*, UAnimationAsset*>& ComplexAnimsMap, const TMap<UAnimSequence*, UAnimSequence*>& AnimSequenceMap) OVERRIDE;
+	virtual bool DoesSupportTimeForTransitionGetter() const override;
+	virtual UAnimationAsset* GetAnimationAsset() const override;
+	virtual const TCHAR* GetTimePropertyName() const override;
+	virtual UScriptStruct* GetTimePropertyStruct() const override;
+	virtual void GetAllAnimationSequencesReferred(TArray<UAnimationAsset*>& ComplexAnims, TArray<UAnimSequence*>& AnimationSequences) const override;
+	virtual void ReplaceReferredAnimations(const TMap<UAnimationAsset*, UAnimationAsset*>& ComplexAnimsMap, const TMap<UAnimSequence*, UAnimSequence*>& AnimSequenceMap) override;
 	// End of UAnimGraphNode_Base
 
 	// UK2Node interface
-	virtual void GetContextMenuActions(const FGraphNodeContextMenuBuilder& Context) const OVERRIDE;
+	virtual void GetContextMenuActions(const FGraphNodeContextMenuBuilder& Context) const override;
 	// End of UK2Node interface
 };
diff --git a/Engine/Source/Editor/AnimGraph/Classes/AnimGraphNode_SequencePlayer.h b/Engine/Source/Editor/AnimGraph/Classes/AnimGraphNode_SequencePlayer.h
--- a/Engine/Source/Editor/AnimGraph/Classes/AnimGraphNode_SequencePlayer.h
+++ b/Engine/Source/Editor/AnimGraph/Classes/AnimGraphNode_SequencePlayer.h
@@ -18,27 +18,27 @@ class UAnimGraphNode_SequencePlayer : public UAnimGraphNode_Base
 	FAnimationGroupReference SyncGroup;
 
 	// UEdGraphNode interface
-	virtual FLinearColor GetNodeTitleColor() const OVERRIDE;
-	virtual FString GetTooltip() const OVERRIDE;
-	virtual FText GetNodeTitle(ENodeTitleType::Type TitleType) const OVERRIDE;
-	virtual FString GetNodeNativeTitle(ENodeTitleType::Type TitleType) const OVERRIDE;
+	virtual FLinearColor GetNodeTitleColor() const override;
+	virtual FString GetTooltip() const override;
+	virtual FText GetNodeTitle(ENodeTitleType::Type TitleType) const override;
+	virtual FString GetNodeNativeTitle(ENodeTitleType::Type TitleType) const override;
 	// End of UEdGraphNode interface
 
 	// UAnimGraphNode_Base interface
-	virtual void GetMenuEntries(FGraphContextMenuBuilder& ContextMenuBuilder) const OVERRIDE;
-	virtual void ValidateAnimNodeDuringCompilation(class USkeleton* ForSkeleton, class FCompilerResultsLog& MessageLog) OVERRIDE;
-	virtual void PreloadRequiredAssets() OVERRIDE;		
-	virtual void BakeDataDuringCompilation(class FCompilerResultsLog& MessageLog) OVERRIDE;
-	virtual bool DoesSupportTimeForTransitionGetter() const OVERRIDE;
-	virtual UAnimationAsset* GetAnimationAsset() const OVERRIDE;
-	virtual const TCHAR* GetTimePropertyName() const OVERRIDE;
-	virtual UScriptStruct* GetTimePropertyStruct() const OVERRIDE;
-	virtual void GetAllAnimationSequencesReferred(TArray<UAnimationAsset*>& ComplexAnims, TArray<UAnimSequence*>& AnimationSequences) const OVERRIDE;
-	virtual void ReplaceReferredAnimations(const TMap<UAnimationAsset*, UAnimationAsset*>& ComplexAnimsMap, const TMap<UAnimSequence*, UAnimSequence*>& AnimSequenceMap) OVERRIDE;
+	virtual void GetMenuEntries(FGraphContextMenuBuilder& ContextMenuBuilder) const override;
+	virtual void ValidateAnimNodeDuringCompilation(class USkeleton* ForSkeleton, class FCompilerResultsLog& MessageLog) override;
+	virtual void PreloadRequiredAssets() override;		
+	virtual void BakeDataDuringCompilation(class FCompilerResultsLog& MessageLog) override;
+	virtual bool DoesSupportTimeForTransitionGetter() const override;
+	virtual UAnimationAsset* GetAnimationAsset() const override;
+	virtual const TCHAR* GetTimePropertyName() const override;
+	virtual UScriptStruct* GetTimePropertyStruct() const override;
+	virtual void GetAllAnimationSequencesReferred(TArray<UAnimationAsset*>& ComplexAnims, TArray<UAnimSequence*>& AnimationSequences) const override;
+	virtual void ReplaceReferredAnimations(const TMap<UAnimationAsset*, UAnimationAsset*>& ComplexAnimsMap, const TMap<UAnimSequence*, UAnimSequence*>& AnimSequenceMap) override;
 	// End of UAnimGraphNode_Base interface
 
 	// UK2Node interface
-	virtual void GetContextMenuActions(const FGraphNodeContextMenuBuilder& Context) const OVERRIDE;
+	virtual void GetContextMenuActions(const FGraphNodeContextMenuBuilder& Context) const override;
 	// End of UK2Node interface
 
 private:
diff --git a/Engine/Source/Editor/AnimGraph/Classes/AnimGraphNode_SkeletalControlBase.h b/Engine/Source/Editor/AnimGraph/Classes/AnimGraphNode_SkeletalControlBase.h
--- a/Engine/Source/Editor/AnimGraph/Classes/AnimGraphNode_SkeletalControlBase.h
+++ b/Engine/Source/Editor/AnimGraph/Classes/AnimGraphNode_SkeletalControlBase.h
@@ -12,14 +12,14 @@ class UAnimGraphNode_SkeletalControlBase : public UAnimGraphNode_Base
 
 public:
 	// UEdGraphNode interface
-	ANIMGRAPH_API virtual FLinearColor GetNodeTitleColor() const OVERRIDE;
-	ANIMGRAPH_API virtual FString GetTooltip() const OVERRIDE;
+	ANIMGRAPH_API virtual FLinearColor GetNodeTitleColor() const override;
+	ANIMGRAPH_API virtual FString GetTooltip() const override;
 	// End of UEdGraphNode interface
 
 	// UAnimGraphNode_Base interface
-	ANIMGRAPH_API virtual FString GetNodeCategory() const OVERRIDE;
-	ANIMGRAPH_API virtual void CreateOutputPins() OVERRIDE;
-	ANIMGRAPH_API virtual void GetMenuEntries(FGraphContextMenuBuilder& ContextMenuBuilder) const OVERRIDE;
+	ANIMGRAPH_API virtual FString GetNodeCategory() const override;
+	ANIMGRAPH_API virtual void CreateOutputPins() override;
+	ANIMGRAPH_API virtual void GetMenuEntries(FGraphContextMenuBuilder& ContextMenuBuilder) const override;
 	// End of UAnimGraphNode_Base interface
 
 	// Draw function for supporting visualization
diff --git a/Engine/Source/Editor/AnimGraph/Classes/AnimGraphNode_Slot.h b/Engine/Source/Editor/AnimGraph/Classes/AnimGraphNode_Slot.h
--- a/Engine/Source/Editor/AnimGraph/Classes/AnimGraphNode_Slot.h
+++ b/Engine/Source/Editor/AnimGraph/Classes/AnimGraphNode_Slot.h
@@ -14,14 +14,14 @@ class UAnimGraphNode_Slot : public UAnimGraphNode_Base
 	FAnimNode_Slot Node;
 
 	// Begin UEdGraphNode interface.
-	virtual FLinearColor GetNodeTitleColor() const OVERRIDE;
-	virtual FString GetTooltip() const OVERRIDE;
-	virtual FText GetNodeTitle(ENodeTitleType::Type TitleType) const OVERRIDE;
-	virtual FString GetNodeNativeTitle(ENodeTitleType::Type TitleType) const OVERRIDE;
+	virtual FLinearColor GetNodeTitleColor() const override;
+	virtual FString GetTooltip() const override;
+	virtual FText GetNodeTitle(ENodeTitleType::Type TitleType) const override;
+	virtual FString GetNodeNativeTitle(ENodeTitleType::Type TitleType) const override;
 	// End UEdGraphNode interface.
 
 	// UAnimGraphNode_Base interface
-	virtual FString GetNodeCategory() const OVERRIDE;
-	virtual void BakeDataDuringCompilation(class FCompilerResultsLog& MessageLog) OVERRIDE;
+	virtual FString GetNodeCategory() const override;
+	virtual void BakeDataDuringCompilation(class FCompilerResultsLog& MessageLog) override;
 	// End of UAnimGraphNode_Base interface
 };
diff --git a/Engine/Source/Editor/AnimGraph/Classes/AnimGraphNode_StateMachine.h b/Engine/Source/Editor/AnimGraph/Classes/AnimGraphNode_StateMachine.h
--- a/Engine/Source/Editor/AnimGraph/Classes/AnimGraphNode_StateMachine.h
+++ b/Engine/Source/Editor/AnimGraph/Classes/AnimGraphNode_StateMachine.h
@@ -14,6 +14,6 @@ class UAnimGraphNode_StateMachine : public UAnimGraphNode_StateMachineBase
 	FAnimNode_StateMachine Node;
 
 	// UAnimGraphNode_StateMachineBase interface
-	virtual FAnimNode_StateMachine& GetNode() OVERRIDE { return Node; }
+	virtual FAnimNode_StateMachine& GetNode() override { return Node; }
 	// End of UAnimGraphNode_StateMachineBase interface
 };
diff --git a/Engine/Source/Editor/AnimGraph/Classes/AnimGraphNode_StateMachineBase.h b/Engine/Source/Editor/AnimGraph/Classes/AnimGraphNode_StateMachineBase.h
--- a/Engine/Source/Editor/AnimGraph/Classes/AnimGraphNode_StateMachineBase.h
+++ b/Engine/Source/Editor/AnimGraph/Classes/AnimGraphNode_StateMachineBase.h
@@ -15,22 +15,22 @@ class ANIMGRAPH_API UAnimGraphNode_StateMachineBase : public UAnimGraphNode_Base
 	class UAnimationStateMachineGraph* EditorStateMachineGraph;
 
 	// UEdGraphNode interface
-	virtual FLinearColor GetNodeTitleColor() const OVERRIDE;
-	virtual FString GetTooltip() const OVERRIDE;
-	virtual FText GetNodeTitle(ENodeTitleType::Type TitleType) const OVERRIDE;
-	virtual FString GetNodeNativeTitle(ENodeTitleType::Type TitleType) const OVERRIDE;
-	virtual void PostPlacedNewNode() OVERRIDE;
-	virtual UObject* GetJumpTargetForDoubleClick() const OVERRIDE;
-	virtual void DestroyNode() OVERRIDE;
-	virtual void PostPasteNode() OVERRIDE;
-	virtual TSharedPtr<class INameValidatorInterface> MakeNameValidator() const OVERRIDE;
-	virtual FString GetDocumentationLink() const OVERRIDE;
-	virtual void OnRenameNode(const FString& NewName) OVERRIDE;
+	virtual FLinearColor GetNodeTitleColor() const override;
+	virtual FString GetTooltip() const override;
+	virtual FText GetNodeTitle(ENodeTitleType::Type TitleType) const override;
+	virtual FString GetNodeNativeTitle(ENodeTitleType::Type TitleType) const override;
+	virtual void PostPlacedNewNode() override;
+	virtual UObject* GetJumpTargetForDoubleClick() const override;
+	virtual void DestroyNode() override;
+	virtual void PostPasteNode() override;
+	virtual TSharedPtr<class INameValidatorInterface> MakeNameValidator() const override;
+	virtual FString GetDocumentationLink() const override;
+	virtual void OnRenameNode(const FString& NewName) override;
 	// End of UEdGraphNode interface
 
 	// UAnimGraphNode_Base interface
-	virtual FString GetNodeCategory() const OVERRIDE;
-	virtual void GetMenuEntries(FGraphContextMenuBuilder& ContextMenuBuilder) const OVERRIDE;
+	virtual FString GetNodeCategory() const override;
+	virtual void GetMenuEntries(FGraphContextMenuBuilder& ContextMenuBuilder) const override;
 	// End of UAnimGraphNode_Base interface
 
 	//  @return the name of this state machine
diff --git a/Engine/Source/Editor/AnimGraph/Classes/AnimGraphNode_TransitionPoseEvaluator.h b/Engine/Source/Editor/AnimGraph/Classes/AnimGraphNode_TransitionPoseEvaluator.h
--- a/Engine/Source/Editor/AnimGraph/Classes/AnimGraphNode_TransitionPoseEvaluator.h
+++ b/Engine/Source/Editor/AnimGraph/Classes/AnimGraphNode_TransitionPoseEvaluator.h
@@ -14,16 +14,16 @@ class UAnimGraphNode_TransitionPoseEvaluator : public UAnimGraphNode_Base
 	FAnimNode_TransitionPoseEvaluator Node;
 	
 	// UEdGraphNode interface
-	virtual FLinearColor GetNodeTitleColor() const OVERRIDE;
-	virtual FString GetTooltip() const OVERRIDE;
-	virtual FText GetNodeTitle(ENodeTitleType::Type TitleType) const OVERRIDE;
-	virtual void ValidateAnimNodeDuringCompilation(class USkeleton* ForSkeleton, class FCompilerResultsLog& MessageLog) OVERRIDE;
-	virtual bool CanDuplicateNode() const OVERRIDE { return false; }
-	virtual bool CanUserDeleteNode() const OVERRIDE;
+	virtual FLinearColor GetNodeTitleColor() const override;
+	virtual FString GetTooltip() const override;
+	virtual FText GetNodeTitle(ENodeTitleType::Type TitleType) const override;
+	virtual void ValidateAnimNodeDuringCompilation(class USkeleton* ForSkeleton, class FCompilerResultsLog& MessageLog) override;
+	virtual bool CanDuplicateNode() const override { return false; }
+	virtual bool CanUserDeleteNode() const override;
 	// End of UEdGraphNode interface
 
 	// UAnimGraphNode_Base interface
-	virtual FString GetNodeCategory() const OVERRIDE;
-	virtual void GetMenuEntries(FGraphContextMenuBuilder& ContextMenuBuilder) const OVERRIDE;
+	virtual FString GetNodeCategory() const override;
+	virtual void GetMenuEntries(FGraphContextMenuBuilder& ContextMenuBuilder) const override;
 	// End of UAnimGraphNode_Base interface
 };
\ No newline at end of file
diff --git a/Engine/Source/Editor/AnimGraph/Classes/AnimGraphNode_TransitionResult.h b/Engine/Source/Editor/AnimGraph/Classes/AnimGraphNode_TransitionResult.h
--- a/Engine/Source/Editor/AnimGraph/Classes/AnimGraphNode_TransitionResult.h
+++ b/Engine/Source/Editor/AnimGraph/Classes/AnimGraphNode_TransitionResult.h
@@ -14,15 +14,15 @@ class UAnimGraphNode_TransitionResult : public UAnimGraphNode_Base
 	FAnimNode_TransitionResult Node;
 
 	// UEdGraphNode interface
-	virtual bool CanUserDeleteNode() const OVERRIDE { return false; }
-	virtual bool CanDuplicateNode() const OVERRIDE { return false; }
-	virtual FLinearColor GetNodeTitleColor() const OVERRIDE;
-	virtual FString GetTooltip() const OVERRIDE;
-	virtual FText GetNodeTitle(ENodeTitleType::Type TitleType) const OVERRIDE;
-	virtual void GetMenuEntries(FGraphContextMenuBuilder& ContextMenuBuilder) const OVERRIDE;
+	virtual bool CanUserDeleteNode() const override { return false; }
+	virtual bool CanDuplicateNode() const override { return false; }
+	virtual FLinearColor GetNodeTitleColor() const override;
+	virtual FString GetTooltip() const override;
+	virtual FText GetNodeTitle(ENodeTitleType::Type TitleType) const override;
+	virtual void GetMenuEntries(FGraphContextMenuBuilder& ContextMenuBuilder) const override;
 	// End of UEdGraphNode interface
 
 	// UAnimGraphNode_Base interface
-	virtual bool IsSinkNode() const OVERRIDE { return true; }
+	virtual bool IsSinkNode() const override { return true; }
 	// End of UAnimGraphNode_Base interface
 };
diff --git a/Engine/Source/Editor/AnimGraph/Classes/AnimGraphNode_TwoBoneIK.h b/Engine/Source/Editor/AnimGraph/Classes/AnimGraphNode_TwoBoneIK.h
--- a/Engine/Source/Editor/AnimGraph/Classes/AnimGraphNode_TwoBoneIK.h
+++ b/Engine/Source/Editor/AnimGraph/Classes/AnimGraphNode_TwoBoneIK.h
@@ -15,18 +15,18 @@ class UAnimGraphNode_TwoBoneIK : public UAnimGraphNode_SkeletalControlBase
 
 public:
 	// UEdGraphNode interface
-	virtual FText GetNodeTitle(ENodeTitleType::Type TitleType) const OVERRIDE;
-	virtual FString GetNodeNativeTitle(ENodeTitleType::Type TitleType) const OVERRIDE;
-	virtual FString GetTooltip() const OVERRIDE;
+	virtual FText GetNodeTitle(ENodeTitleType::Type TitleType) const override;
+	virtual FString GetNodeNativeTitle(ENodeTitleType::Type TitleType) const override;
+	virtual FString GetTooltip() const override;
 	// End of UEdGraphNode interface
 
 	// UAnimGraphNode_SkeletalControlBase interface
-	ANIMGRAPH_API virtual void Draw( FPrimitiveDrawInterface* PDI, USkeletalMeshComponent* SkelMeshComp) const OVERRIDE;
+	ANIMGRAPH_API virtual void Draw( FPrimitiveDrawInterface* PDI, USkeletalMeshComponent* SkelMeshComp) const override;
 	// End of UAnimGraphNode_SkeletalControlBase interface
 
 protected:
 	// UAnimGraphNode_SkeletalControlBase interface
-	virtual FText GetControllerDescription() const OVERRIDE;
+	virtual FText GetControllerDescription() const override;
 	// End of UAnimGraphNode_SkeletalControlBase interface
 
 	// local conversion function for drawing
diff --git a/Engine/Source/Editor/AnimGraph/Classes/AnimGraphNode_TwoWayBlend.h b/Engine/Source/Editor/AnimGraph/Classes/AnimGraphNode_TwoWayBlend.h
--- a/Engine/Source/Editor/AnimGraph/Classes/AnimGraphNode_TwoWayBlend.h
+++ b/Engine/Source/Editor/AnimGraph/Classes/AnimGraphNode_TwoWayBlend.h
@@ -14,12 +14,12 @@ class UAnimGraphNode_TwoWayBlend : public UAnimGraphNode_Base
 	FAnimationNode_TwoWayBlend BlendNode;
 
 	// Begin UEdGraphNode interface.
-	virtual FLinearColor GetNodeTitleColor() const OVERRIDE;
-	virtual FString GetTooltip() const OVERRIDE;
-	virtual FText GetNodeTitle(ENodeTitleType::Type TitleType) const OVERRIDE;
+	virtual FLinearColor GetNodeTitleColor() const override;
+	virtual FString GetTooltip() const override;
+	virtual FText GetNodeTitle(ENodeTitleType::Type TitleType) const override;
 	// End UEdGraphNode interface.
 
 	// UAnimGraphNode_Base interface
-	virtual FString GetNodeCategory() const OVERRIDE;
+	virtual FString GetNodeCategory() const override;
 	// End of UAnimGraphNode_Base interface
 };
diff --git a/Engine/Source/Editor/AnimGraph/Classes/AnimGraphNode_UseCachedPose.h b/Engine/Source/Editor/AnimGraph/Classes/AnimGraphNode_UseCachedPose.h
--- a/Engine/Source/Editor/AnimGraph/Classes/AnimGraphNode_UseCachedPose.h
+++ b/Engine/Source/Editor/AnimGraph/Classes/AnimGraphNode_UseCachedPose.h
@@ -19,13 +19,13 @@ public:
 
 public:
 	// UEdGraphNode interface
-	virtual FString GetTooltip() const OVERRIDE;
-	virtual FText GetNodeTitle(ENodeTitleType::Type TitleType) const OVERRIDE;
-	virtual FString GetNodeNativeTitle(ENodeTitleType::Type TitleType) const OVERRIDE;
+	virtual FString GetTooltip() const override;
+	virtual FText GetNodeTitle(ENodeTitleType::Type TitleType) const override;
+	virtual FString GetNodeNativeTitle(ENodeTitleType::Type TitleType) const override;
 	// End of UEdGraphNode interface
 
 	// UAnimGraphNode_Base interface
-	virtual FString GetNodeCategory() const OVERRIDE;
-	virtual void GetMenuEntries(FGraphContextMenuBuilder& ContextMenuBuilder) const OVERRIDE;
+	virtual FString GetNodeCategory() const override;
+	virtual void GetMenuEntries(FGraphContextMenuBuilder& ContextMenuBuilder) const override;
 	// End of UAnimGraphNode_Base interface
 };
diff --git a/Engine/Source/Editor/AnimGraph/Classes/AnimGraphNode_WheelHandler.h b/Engine/Source/Editor/AnimGraph/Classes/AnimGraphNode_WheelHandler.h
--- a/Engine/Source/Editor/AnimGraph/Classes/AnimGraphNode_WheelHandler.h
+++ b/Engine/Source/Editor/AnimGraph/Classes/AnimGraphNode_WheelHandler.h
@@ -16,16 +16,16 @@ class UAnimGraphNode_WheelHandler : public UAnimGraphNode_SkeletalControlBase
 
 public:
 	// UEdGraphNode interface
-	virtual FText GetNodeTitle(ENodeTitleType::Type TitleType) const OVERRIDE;
-	virtual FString GetTooltip() const OVERRIDE;
-	virtual FString GetKeywords() const OVERRIDE;
+	virtual FText GetNodeTitle(ENodeTitleType::Type TitleType) const override;
+	virtual FString GetTooltip() const override;
+	virtual FString GetKeywords() const override;
 	// validate if this is within VehicleAnimInstance
-	virtual void ValidateAnimNodePostCompile(class FCompilerResultsLog& MessageLog, class UAnimBlueprintGeneratedClass* CompiledClass, int32 CompiledNodeIndex) OVERRIDE;
-	virtual void GetMenuEntries(FGraphContextMenuBuilder& ContextMenuBuilder) const OVERRIDE;
+	virtual void ValidateAnimNodePostCompile(class FCompilerResultsLog& MessageLog, class UAnimBlueprintGeneratedClass* CompiledClass, int32 CompiledNodeIndex) override;
+	virtual void GetMenuEntries(FGraphContextMenuBuilder& ContextMenuBuilder) const override;
 	// End of UEdGraphNode interface
 
 protected:
 	// UAnimGraphNode_SkeletalControlBase interface
-	virtual FText GetControllerDescription() const OVERRIDE;
+	virtual FText GetControllerDescription() const override;
 	// End of UAnimGraphNode_SkeletalControlBase interface
 };
diff --git a/Engine/Source/Editor/AnimGraph/Classes/AnimPreviewInstance.h b/Engine/Source/Editor/AnimGraph/Classes/AnimPreviewInstance.h
--- a/Engine/Source/Editor/AnimGraph/Classes/AnimPreviewInstance.h
+++ b/Engine/Source/Editor/AnimGraph/Classes/AnimPreviewInstance.h
@@ -45,8 +45,8 @@ class ANIMGRAPH_API UAnimPreviewInstance : public UAnimSingleNodeInstance
 	int32 MontagePreviewStartSectionIdx;
 
 	// Begin UAnimInstance interface
-	virtual void NativeInitializeAnimation() OVERRIDE;
-	virtual bool NativeEvaluateAnimation(FPoseContext& Output) OVERRIDE;
+	virtual void NativeInitializeAnimation() override;
+	virtual bool NativeEvaluateAnimation(FPoseContext& Output) override;
 	// End UAnimInstance interface
 
 	/** Set SkeletalControl Alpha**/
@@ -55,7 +55,7 @@ class ANIMGRAPH_API UAnimPreviewInstance : public UAnimSingleNodeInstance
 	UAnimSequence* GetAnimSequence();
 
 	// Begin UAnimSingleNodeInstance interface
-	virtual void RestartMontage(UAnimMontage* Montage, FName FromSection = FName()) OVERRIDE;
+	virtual void RestartMontage(UAnimMontage* Montage, FName FromSection = FName()) override;
 	// End UAnimSingleNodeInstance interface
 
 	/** Montage preview functions */
diff --git a/Engine/Source/Editor/AnimGraph/Classes/AnimStateConduitNode.h b/Engine/Source/Editor/AnimGraph/Classes/AnimStateConduitNode.h
--- a/Engine/Source/Editor/AnimGraph/Classes/AnimStateConduitNode.h
+++ b/Engine/Source/Editor/AnimGraph/Classes/AnimStateConduitNode.h
@@ -15,23 +15,23 @@ public:
 	class UEdGraph* BoundGraph;
 
 	// Begin UEdGraphNode interface
-	virtual void AllocateDefaultPins() OVERRIDE;
-	virtual void AutowireNewNode(UEdGraphPin* FromPin) OVERRIDE;
-	virtual FText GetNodeTitle(ENodeTitleType::Type TitleType) const OVERRIDE;
-	virtual FString GetTooltip() const OVERRIDE;
-	virtual bool CanDuplicateNode() const OVERRIDE { return true; }
-	virtual void PostPasteNode() OVERRIDE;
-	virtual void PostPlacedNewNode() OVERRIDE;
-	virtual void DestroyNode() OVERRIDE;
-	virtual void ValidateNodeDuringCompilation(class FCompilerResultsLog& MessageLog) const OVERRIDE;
+	virtual void AllocateDefaultPins() override;
+	virtual void AutowireNewNode(UEdGraphPin* FromPin) override;
+	virtual FText GetNodeTitle(ENodeTitleType::Type TitleType) const override;
+	virtual FString GetTooltip() const override;
+	virtual bool CanDuplicateNode() const override { return true; }
+	virtual void PostPasteNode() override;
+	virtual void PostPlacedNewNode() override;
+	virtual void DestroyNode() override;
+	virtual void ValidateNodeDuringCompilation(class FCompilerResultsLog& MessageLog) const override;
 	// End UEdGraphNode interface
 
 	// Begin UAnimStateNodeBase interface
-	virtual UEdGraphPin* GetInputPin() const OVERRIDE;
-	virtual UEdGraphPin* GetOutputPin() const OVERRIDE;
-	virtual FString GetStateName() const OVERRIDE;
-	virtual void GetTransitionList(TArray<class UAnimStateTransitionNode*>& OutTransitions, bool bWantSortedList = true) OVERRIDE;
-	virtual FString GetDesiredNewNodeName() const OVERRIDE;
-	virtual UEdGraph* GetBoundGraph() const OVERRIDE { return BoundGraph; }
+	virtual UEdGraphPin* GetInputPin() const override;
+	virtual UEdGraphPin* GetOutputPin() const override;
+	virtual FString GetStateName() const override;
+	virtual void GetTransitionList(TArray<class UAnimStateTransitionNode*>& OutTransitions, bool bWantSortedList = true) override;
+	virtual FString GetDesiredNewNodeName() const override;
+	virtual UEdGraph* GetBoundGraph() const override { return BoundGraph; }
 	// End of UAnimStateNodeBase interface
 };
diff --git a/Engine/Source/Editor/AnimGraph/Classes/AnimStateEntryNode.h b/Engine/Source/Editor/AnimGraph/Classes/AnimStateEntryNode.h
--- a/Engine/Source/Editor/AnimGraph/Classes/AnimStateEntryNode.h
+++ b/Engine/Source/Editor/AnimGraph/Classes/AnimStateEntryNode.h
@@ -11,9 +11,9 @@ class UAnimStateEntryNode : public UEdGraphNode
 
 
 	// Begin UEdGraphNode interface
-	virtual void AllocateDefaultPins() OVERRIDE;
-	virtual FText GetNodeTitle(ENodeTitleType::Type TitleType) const OVERRIDE;
-	virtual FString GetTooltip() const OVERRIDE;
+	virtual void AllocateDefaultPins() override;
+	virtual FText GetNodeTitle(ENodeTitleType::Type TitleType) const override;
+	virtual FString GetTooltip() const override;
 	// End UEdGraphNode interface
 	
 	ANIMGRAPH_API UEdGraphNode* GetOutputNode() const;
diff --git a/Engine/Source/Editor/AnimGraph/Classes/AnimStateNode.h b/Engine/Source/Editor/AnimGraph/Classes/AnimStateNode.h
--- a/Engine/Source/Editor/AnimGraph/Classes/AnimStateNode.h
+++ b/Engine/Source/Editor/AnimGraph/Classes/AnimStateNode.h
@@ -36,30 +36,30 @@ public:
 	FAnimNotifyEvent StateFullyBlended;
 
 	// Begin UObject interface
-	virtual void PostEditChangeProperty(struct FPropertyChangedEvent& PropertyChangedEvent) OVERRIDE;
+	virtual void PostEditChangeProperty(struct FPropertyChangedEvent& PropertyChangedEvent) override;
 	// End UObject interface
 
 	// Begin UEdGraphNode interface
-	virtual void AllocateDefaultPins() OVERRIDE;
-	virtual void AutowireNewNode(UEdGraphPin* FromPin) OVERRIDE;
-	virtual FText GetNodeTitle(ENodeTitleType::Type TitleType) const OVERRIDE;
-	virtual FString GetTooltip() const OVERRIDE;
+	virtual void AllocateDefaultPins() override;
+	virtual void AutowireNewNode(UEdGraphPin* FromPin) override;
+	virtual FText GetNodeTitle(ENodeTitleType::Type TitleType) const override;
+	virtual FString GetTooltip() const override;
 	virtual bool CanDuplicateNode() const { return true; }
-	virtual void PostPasteNode() OVERRIDE;
-	virtual void PostPlacedNewNode() OVERRIDE;
-	virtual void DestroyNode() OVERRIDE;
+	virtual void PostPasteNode() override;
+	virtual void PostPlacedNewNode() override;
+	virtual void DestroyNode() override;
 	// End UEdGraphNode interface
 	
 	// Begin UAnimStateNodeBase interface
-	virtual UEdGraphPin* GetInputPin() const OVERRIDE;
-	virtual UEdGraphPin* GetOutputPin() const OVERRIDE;
-	virtual FString GetStateName() const OVERRIDE;
-	virtual void GetTransitionList(TArray<class UAnimStateTransitionNode*>& OutTransitions, bool bWantSortedList = false) OVERRIDE;
+	virtual UEdGraphPin* GetInputPin() const override;
+	virtual UEdGraphPin* GetOutputPin() const override;
+	virtual FString GetStateName() const override;
+	virtual void GetTransitionList(TArray<class UAnimStateTransitionNode*>& OutTransitions, bool bWantSortedList = false) override;
 	// End of UAnimStateNodeBase interface
 
 	// @return the pose pin of the state sink node within the anim graph of this state
 	ANIMGRAPH_API UEdGraphPin* GetPoseSinkPinInsideState() const;
 
 public:
-	virtual UEdGraph* GetBoundGraph() const OVERRIDE { return BoundGraph; }
+	virtual UEdGraph* GetBoundGraph() const override { return BoundGraph; }
 };
diff --git a/Engine/Source/Editor/AnimGraph/Classes/AnimStateNodeBase.h b/Engine/Source/Editor/AnimGraph/Classes/AnimStateNodeBase.h
--- a/Engine/Source/Editor/AnimGraph/Classes/AnimStateNodeBase.h
+++ b/Engine/Source/Editor/AnimGraph/Classes/AnimStateNodeBase.h
@@ -10,11 +10,11 @@ class UAnimStateNodeBase : public UEdGraphNode
 	GENERATED_UCLASS_BODY()
 
 	// UEdGraphNode interface
-	virtual void PostPasteNode() OVERRIDE;
-	virtual UObject* GetJumpTargetForDoubleClick() const OVERRIDE;
-	virtual bool CanCreateUnderSpecifiedSchema(const UEdGraphSchema* Schema) const OVERRIDE;
-	virtual void OnRenameNode(const FString& NewName) OVERRIDE;
-	virtual TSharedPtr<class INameValidatorInterface> MakeNameValidator() const OVERRIDE;
+	virtual void PostPasteNode() override;
+	virtual UObject* GetJumpTargetForDoubleClick() const override;
+	virtual bool CanCreateUnderSpecifiedSchema(const UEdGraphSchema* Schema) const override;
+	virtual void OnRenameNode(const FString& NewName) override;
+	virtual TSharedPtr<class INameValidatorInterface> MakeNameValidator() const override;
 	// End of UEdGraphNode interface
 
 	// @return the input pin for this state
@@ -38,6 +38,6 @@ protected:
 	
 public:
 	// Gets the animation state node documentation link
-	virtual FString GetDocumentationLink() const OVERRIDE;
+	virtual FString GetDocumentationLink() const override;
 
 };
diff --git a/Engine/Source/Editor/AnimGraph/Classes/AnimStateTransitionNode.h b/Engine/Source/Editor/AnimGraph/Classes/AnimStateTransitionNode.h
--- a/Engine/Source/Editor/AnimGraph/Classes/AnimStateTransitionNode.h
+++ b/Engine/Source/Editor/AnimGraph/Classes/AnimStateTransitionNode.h
@@ -77,29 +77,29 @@ class UAnimStateTransitionNode : public UAnimStateNodeBase
 	int32 SharedCrossfadeIdx;
 
 	// Begin UObject interface
-	virtual void PostEditChangeProperty(struct FPropertyChangedEvent& PropertyChangedEvent) OVERRIDE;
-	virtual void PostLoad() OVERRIDE;
+	virtual void PostEditChangeProperty(struct FPropertyChangedEvent& PropertyChangedEvent) override;
+	virtual void PostLoad() override;
 	// End UObject interface
 
 	// Begin UEdGraphNode interface
-	virtual void AllocateDefaultPins() OVERRIDE;
-	virtual FText GetNodeTitle(ENodeTitleType::Type TitleType) const OVERRIDE;
-	virtual FString GetNodeNativeTitle(ENodeTitleType::Type TitleType) const OVERRIDE;
-	virtual FString GetTooltip() const OVERRIDE;
-	virtual FLinearColor GetNodeTitleColor() const OVERRIDE;
-	virtual void PinConnectionListChanged(UEdGraphPin* Pin) OVERRIDE;
+	virtual void AllocateDefaultPins() override;
+	virtual FText GetNodeTitle(ENodeTitleType::Type TitleType) const override;
+	virtual FString GetNodeNativeTitle(ENodeTitleType::Type TitleType) const override;
+	virtual FString GetTooltip() const override;
+	virtual FLinearColor GetNodeTitleColor() const override;
+	virtual void PinConnectionListChanged(UEdGraphPin* Pin) override;
 	virtual bool CanDuplicateNode() const { return true; }
-	virtual void PrepareForCopying() OVERRIDE;
-	virtual void PostPasteNode() OVERRIDE;
-	virtual void PostPlacedNewNode() OVERRIDE;
-	virtual void DestroyNode() OVERRIDE;
-	virtual void ValidateNodeDuringCompilation(class FCompilerResultsLog& MessageLog) const OVERRIDE;
+	virtual void PrepareForCopying() override;
+	virtual void PostPasteNode() override;
+	virtual void PostPlacedNewNode() override;
+	virtual void DestroyNode() override;
+	virtual void ValidateNodeDuringCompilation(class FCompilerResultsLog& MessageLog) const override;
 	// End UEdGraphNode interface
 
 	// Begin UAnimStateNodeBase interface
-	virtual UEdGraph* GetBoundGraph() const OVERRIDE { return BoundGraph; }
-	virtual UEdGraphPin* GetInputPin() const OVERRIDE { return Pins[0]; }
-	virtual UEdGraphPin* GetOutputPin() const OVERRIDE { return Pins[1]; }
+	virtual UEdGraph* GetBoundGraph() const override { return BoundGraph; }
+	virtual UEdGraphPin* GetInputPin() const override { return Pins[0]; }
+	virtual UEdGraphPin* GetOutputPin() const override { return Pins[1]; }
 	// End UAnimStateNodeBase interface
 
 	// @return the name of this state
diff --git a/Engine/Source/Editor/AnimGraph/Classes/AnimationConduitGraphSchema.h b/Engine/Source/Editor/AnimGraph/Classes/AnimationConduitGraphSchema.h
--- a/Engine/Source/Editor/AnimGraph/Classes/AnimationConduitGraphSchema.h
+++ b/Engine/Source/Editor/AnimGraph/Classes/AnimationConduitGraphSchema.h
@@ -11,8 +11,8 @@ class UAnimationConduitGraphSchema : public UEdGraphSchema_K2
 	GENERATED_UCLASS_BODY()
 
 	// UEdGraphSchema interface
-	virtual void CreateDefaultNodesForGraph(UEdGraph& Graph) const OVERRIDE;
-	virtual bool CanDuplicateGraph(UEdGraph* InSourceGraph) const OVERRIDE { return false; }
-	virtual void GetGraphDisplayInformation(const UEdGraph& Graph, /*out*/ FGraphDisplayInfo& DisplayInfo) const OVERRIDE;
+	virtual void CreateDefaultNodesForGraph(UEdGraph& Graph) const override;
+	virtual bool CanDuplicateGraph(UEdGraph* InSourceGraph) const override { return false; }
+	virtual void GetGraphDisplayInformation(const UEdGraph& Graph, /*out*/ FGraphDisplayInfo& DisplayInfo) const override;
 	// End of UEdGraphSchema interface
 };
diff --git a/Engine/Source/Editor/AnimGraph/Classes/AnimationCustomTransitionSchema.h b/Engine/Source/Editor/AnimGraph/Classes/AnimationCustomTransitionSchema.h
--- a/Engine/Source/Editor/AnimGraph/Classes/AnimationCustomTransitionSchema.h
+++ b/Engine/Source/Editor/AnimGraph/Classes/AnimationCustomTransitionSchema.h
@@ -10,7 +10,7 @@ class UAnimationCustomTransitionSchema : public UAnimationGraphSchema
 	GENERATED_UCLASS_BODY()
 
 	// Begin UEdGraphSchema interface.
-	virtual void CreateDefaultNodesForGraph(UEdGraph& Graph) const OVERRIDE;
-	void GetGraphDisplayInformation(const UEdGraph& Graph, /*out*/ FGraphDisplayInfo& DisplayInfo) const OVERRIDE;
+	virtual void CreateDefaultNodesForGraph(UEdGraph& Graph) const override;
+	void GetGraphDisplayInformation(const UEdGraph& Graph, /*out*/ FGraphDisplayInfo& DisplayInfo) const override;
 	// End UEdGraphSchema interface.
 };
diff --git a/Engine/Source/Editor/AnimGraph/Classes/AnimationGraphSchema.h b/Engine/Source/Editor/AnimGraph/Classes/AnimationGraphSchema.h
--- a/Engine/Source/Editor/AnimGraph/Classes/AnimationGraphSchema.h
+++ b/Engine/Source/Editor/AnimGraph/Classes/AnimationGraphSchema.h
@@ -36,28 +36,28 @@ class UAnimationGraphSchema : public UEdGraphSchema_K2
 	FName DefaultEvaluationHandlerName;
 
 	// Begin UEdGraphSchema interface.
-	virtual void GetGraphContextActions(FGraphContextMenuBuilder& ContextMenuBuilder) const OVERRIDE;
-	virtual FLinearColor GetPinTypeColor(const FEdGraphPinType& PinType) const OVERRIDE;
-	virtual EGraphType GetGraphType(const UEdGraph* TestEdGraph) const OVERRIDE;
-	virtual void CreateDefaultNodesForGraph(UEdGraph& Graph) const OVERRIDE;
-	virtual void HandleGraphBeingDeleted(UEdGraph& GraphBeingRemoved) const OVERRIDE;
-	virtual bool CreateAutomaticConversionNodeAndConnections(UEdGraphPin* PinA, UEdGraphPin* PinB) const OVERRIDE;
-	virtual void DroppedAssetsOnGraph(const TArray<FAssetData>& Assets, const FVector2D& GraphPosition, UEdGraph* Graph) const OVERRIDE;
-	virtual void DroppedAssetsOnNode(const TArray<FAssetData>& Assets, const FVector2D& GraphPosition, UEdGraphNode* Node) const OVERRIDE;
-	virtual void DroppedAssetsOnPin(const TArray<FAssetData>& Assets, const FVector2D& GraphPosition, UEdGraphPin* Pin) const OVERRIDE;
-	virtual void GetAssetsNodeHoverMessage(const TArray<FAssetData>& Assets, const UEdGraphNode* HoverNode, FString& OutTooltipText, bool& OutOkIcon) const OVERRIDE;
-	virtual void GetAssetsPinHoverMessage(const TArray<FAssetData>& Assets, const UEdGraphPin* HoverPin, FString& OutTooltipText, bool& OutOkIcon) const OVERRIDE;
-	virtual void GetAssetsGraphHoverMessage(const TArray<FAssetData>& Assets, const UEdGraph* HoverGraph, FString& OutTooltipText, bool& OutOkIcon) const OVERRIDE;
-	virtual void GetContextMenuActions(const UEdGraph* CurrentGraph, const UEdGraphNode* InGraphNode, const UEdGraphPin* InGraphPin, FMenuBuilder* MenuBuilder, bool bIsDebugging) const OVERRIDE;
-	virtual FString GetPinDisplayName(const UEdGraphPin* Pin) const OVERRIDE;
-	virtual bool CanDuplicateGraph(UEdGraph* InSourceGraph) const OVERRIDE {	return false; }
+	virtual void GetGraphContextActions(FGraphContextMenuBuilder& ContextMenuBuilder) const override;
+	virtual FLinearColor GetPinTypeColor(const FEdGraphPinType& PinType) const override;
+	virtual EGraphType GetGraphType(const UEdGraph* TestEdGraph) const override;
+	virtual void CreateDefaultNodesForGraph(UEdGraph& Graph) const override;
+	virtual void HandleGraphBeingDeleted(UEdGraph& GraphBeingRemoved) const override;
+	virtual bool CreateAutomaticConversionNodeAndConnections(UEdGraphPin* PinA, UEdGraphPin* PinB) const override;
+	virtual void DroppedAssetsOnGraph(const TArray<FAssetData>& Assets, const FVector2D& GraphPosition, UEdGraph* Graph) const override;
+	virtual void DroppedAssetsOnNode(const TArray<FAssetData>& Assets, const FVector2D& GraphPosition, UEdGraphNode* Node) const override;
+	virtual void DroppedAssetsOnPin(const TArray<FAssetData>& Assets, const FVector2D& GraphPosition, UEdGraphPin* Pin) const override;
+	virtual void GetAssetsNodeHoverMessage(const TArray<FAssetData>& Assets, const UEdGraphNode* HoverNode, FString& OutTooltipText, bool& OutOkIcon) const override;
+	virtual void GetAssetsPinHoverMessage(const TArray<FAssetData>& Assets, const UEdGraphPin* HoverPin, FString& OutTooltipText, bool& OutOkIcon) const override;
+	virtual void GetAssetsGraphHoverMessage(const TArray<FAssetData>& Assets, const UEdGraph* HoverGraph, FString& OutTooltipText, bool& OutOkIcon) const override;
+	virtual void GetContextMenuActions(const UEdGraph* CurrentGraph, const UEdGraphNode* InGraphNode, const UEdGraphPin* InGraphPin, FMenuBuilder* MenuBuilder, bool bIsDebugging) const override;
+	virtual FString GetPinDisplayName(const UEdGraphPin* Pin) const override;
+	virtual bool CanDuplicateGraph(UEdGraph* InSourceGraph) const override {	return false; }
 	// End UEdGraphSchema interface.
 
 	// Begin UEdGraphSchema_K2 interface
-	virtual const FPinConnectionResponse DetermineConnectionResponseOfCompatibleTypedPins(const UEdGraphPin* PinA, const UEdGraphPin* PinB, const UEdGraphPin* InputPin, const UEdGraphPin* OutputPin) const OVERRIDE;
-	virtual bool SearchForAutocastFunction(const UEdGraphPin* OutputPin, const UEdGraphPin* InputPin, /*out*/ FName& TargetFunction) const OVERRIDE;
-	virtual bool ArePinsCompatible(const UEdGraphPin* PinA, const UEdGraphPin* PinB, const UClass* CallingContext = NULL, bool bIgnoreArray = false) const OVERRIDE;
-	virtual bool DoesSupportAnimNotifyActions() const OVERRIDE;
+	virtual const FPinConnectionResponse DetermineConnectionResponseOfCompatibleTypedPins(const UEdGraphPin* PinA, const UEdGraphPin* PinB, const UEdGraphPin* InputPin, const UEdGraphPin* OutputPin) const override;
+	virtual bool SearchForAutocastFunction(const UEdGraphPin* OutputPin, const UEdGraphPin* InputPin, /*out*/ FName& TargetFunction) const override;
+	virtual bool ArePinsCompatible(const UEdGraphPin* PinA, const UEdGraphPin* PinB, const UClass* CallingContext = NULL, bool bIgnoreArray = false) const override;
+	virtual bool DoesSupportAnimNotifyActions() const override;
 	// End UEdGraphSchema_K2 interface
 
 	/** Spawn the correct node in the Animation Graph using the given AnimationAsset at the supplied location */
diff --git a/Engine/Source/Editor/AnimGraph/Classes/AnimationStateGraphSchema.h b/Engine/Source/Editor/AnimGraph/Classes/AnimationStateGraphSchema.h
--- a/Engine/Source/Editor/AnimGraph/Classes/AnimationStateGraphSchema.h
+++ b/Engine/Source/Editor/AnimGraph/Classes/AnimationStateGraphSchema.h
@@ -10,7 +10,7 @@ class UAnimationStateGraphSchema : public UAnimationGraphSchema
 	GENERATED_UCLASS_BODY()
 
 	// Begin UEdGraphSchema interface.
-	virtual void CreateDefaultNodesForGraph(UEdGraph& Graph) const OVERRIDE;
-	virtual void GetGraphDisplayInformation(const UEdGraph& Graph, /*out*/ FGraphDisplayInfo& DisplayInfo) const OVERRIDE;
+	virtual void CreateDefaultNodesForGraph(UEdGraph& Graph) const override;
+	virtual void GetGraphDisplayInformation(const UEdGraph& Graph, /*out*/ FGraphDisplayInfo& DisplayInfo) const override;
 	// End UEdGraphSchema interface.
 };
