Fix for state machine transition created in backward.

Since I need triple eyes, I'll add more people for code review

diff --git a/Engine/Source/Editor/AnimGraph/Private/AnimationStateMachineSchema.cpp b/Engine/Source/Editor/AnimGraph/Private/AnimationStateMachineSchema.cpp
--- a/Engine/Source/Editor/AnimGraph/Private/AnimationStateMachineSchema.cpp
+++ b/Engine/Source/Editor/AnimGraph/Private/AnimationStateMachineSchema.cpp
@@ -207,15 +207,15 @@ bool UAnimationStateMachineSchema::TryCreateConnection(UEdGraphPin* PinA, UEdGra
 {
 	if (PinB->Direction == PinA->Direction)
 	{
-		if (UAnimStateNodeBase* Node = Cast<UAnimStateNodeBase>(PinA->GetOwningNode()))
+		if (UAnimStateNodeBase* Node = Cast<UAnimStateNodeBase>(PinB->GetOwningNode()))
 		{
-			if (PinB->Direction == EGPD_Input)
+			if (PinA->Direction == EGPD_Input)
 			{
-				PinA = Node->GetOutputPin();
+				PinB = Node->GetOutputPin();
 			}
 			else
 			{
-				PinA = Node->GetInputPin();
+				PinB = Node->GetInputPin();
 			}
 		}
 	}
