341921	EDITOR: PERSONA: Updating parameters of a Blend Space 1D does not update throughout the graph

diff --git a/Engine/Source/Editor/AnimGraph/Classes/AnimGraphNode_BlendSpaceBase.h b/Engine/Source/Editor/AnimGraph/Classes/AnimGraphNode_BlendSpaceBase.h
--- a/Engine/Source/Editor/AnimGraph/Classes/AnimGraphNode_BlendSpaceBase.h
+++ b/Engine/Source/Editor/AnimGraph/Classes/AnimGraphNode_BlendSpaceBase.h
@@ -16,6 +16,7 @@ class UAnimGraphNode_BlendSpaceBase : public UAnimGraphNode_Base
 	// UAnimGraphNode_Base interface
 	virtual void CustomizePinData(UEdGraphPin* Pin, FName SourcePropertyName, int32 ArrayIndex) const override;
 	virtual void PreloadRequiredAssets() override;
+	virtual void PostProcessPinName(const UEdGraphPin* Pin, FString& DisplayName) const override;
 	// End of UAnimGraphNode_Base interface
 
 protected:
diff --git a/Engine/Source/Editor/AnimGraph/Private/AnimGraphNode_BlendSpaceBase.cpp b/Engine/Source/Editor/AnimGraph/Private/AnimGraphNode_BlendSpaceBase.cpp
--- a/Engine/Source/Editor/AnimGraph/Private/AnimGraphNode_BlendSpaceBase.cpp
+++ b/Engine/Source/Editor/AnimGraph/Private/AnimGraphNode_BlendSpaceBase.cpp
@@ -109,4 +109,30 @@ void UAnimGraphNode_BlendSpaceBase::PreloadRequiredAssets()
 	PreloadObject(GetBlendSpace());
 
 	Super::PreloadRequiredAssets();
-}
\ No newline at end of file
+}
+
+void UAnimGraphNode_BlendSpaceBase::PostProcessPinName(const UEdGraphPin* Pin, FString& DisplayName) const
+{
+	if(Pin->Direction == EGPD_Input)
+	{
+		UBlendSpaceBase * BlendSpace = GetBlendSpace();
+
+		if(BlendSpace != NULL)
+		{
+			if(Pin->PinName == TEXT("X"))
+			{
+				DisplayName = BlendSpace->GetBlendParameter(0).DisplayName;
+			}
+			else if(Pin->PinName == TEXT("Y"))
+			{
+				DisplayName = BlendSpace->GetBlendParameter(1).DisplayName;
+			}
+			else if(Pin->PinName == TEXT("Z"))
+			{
+				DisplayName = BlendSpace->GetBlendParameter(2).DisplayName;
+			}
+		}
+	}
+
+	Super::PostProcessPinName(Pin, DisplayName);
+}
