dllexport on the GetControllerDescription - requested by licensee.

diff --git a/Engine/Source/Editor/AnimGraph/Classes/AnimGraphNode_SkeletalControlBase.h b/Engine/Source/Editor/AnimGraph/Classes/AnimGraphNode_SkeletalControlBase.h
--- a/Engine/Source/Editor/AnimGraph/Classes/AnimGraphNode_SkeletalControlBase.h
+++ b/Engine/Source/Editor/AnimGraph/Classes/AnimGraphNode_SkeletalControlBase.h
@@ -26,5 +26,5 @@ public:
 	ANIMGRAPH_API virtual void Draw( FPrimitiveDrawInterface* PDI, USkeletalMeshComponent * PreviewSkelMeshComp ) const {};
 protected:
 	// Returns the short descriptive name of the controller
-	virtual FText GetControllerDescription() const;
+	ANIMGRAPH_API virtual FText GetControllerDescription() const;
 };
