(WIP) For the new blueprint menu system, ensuring the proper animation nodes are available.

diff --git a/Engine/Source/Editor/AnimGraph/Classes/AnimGraphNode_Base.h b/Engine/Source/Editor/AnimGraph/Classes/AnimGraphNode_Base.h
--- a/Engine/Source/Editor/AnimGraph/Classes/AnimGraphNode_Base.h
+++ b/Engine/Source/Editor/AnimGraph/Classes/AnimGraphNode_Base.h
@@ -113,6 +113,8 @@ class UAnimGraphNode_Base : public UK2Node
 	ANIMGRAPH_API virtual bool CanCreateUnderSpecifiedSchema(const UEdGraphSchema* DesiredSchema) const override;
 	ANIMGRAPH_API virtual void GetMenuEntries(FGraphContextMenuBuilder& ContextMenuBuilder) const override;
 	ANIMGRAPH_API virtual void GetNodeAttributes(TArray<TKeyValuePair<FString, FString>>& OutNodeAttributes) const override;
+	ANIMGRAPH_API virtual void GetMenuActions(TArray<UBlueprintNodeSpawner*>& ActionListOut) const override;
+	ANIMGRAPH_API virtual FText GetMenuCategory() const override;
 
 	// By default return any animation assets we have
 	virtual UObject* GetJumpTargetForDoubleClick() const override { return GetAnimationAsset(); }
diff --git a/Engine/Source/Editor/AnimGraph/Classes/AnimGraphNode_BlendSpaceEvaluator.h b/Engine/Source/Editor/AnimGraph/Classes/AnimGraphNode_BlendSpaceEvaluator.h
--- a/Engine/Source/Editor/AnimGraph/Classes/AnimGraphNode_BlendSpaceEvaluator.h
+++ b/Engine/Source/Editor/AnimGraph/Classes/AnimGraphNode_BlendSpaceEvaluator.h
@@ -19,6 +19,7 @@ class UAnimGraphNode_BlendSpaceEvaluator : public UAnimGraphNode_BlendSpaceBase
 	// UEdGraphNode interface
 	virtual FString GetTooltip() const override;
 	virtual FText GetNodeTitle(ENodeTitleType::Type TitleType) const override;
+	virtual void GetMenuActions(TArray<UBlueprintNodeSpawner*>& ActionListOut) const override;
 	// End of UEdGraphNode interface
 
 	// UAnimGraphNode_Base interface
diff --git a/Engine/Source/Editor/AnimGraph/Classes/AnimGraphNode_BlendSpacePlayer.h b/Engine/Source/Editor/AnimGraph/Classes/AnimGraphNode_BlendSpacePlayer.h
--- a/Engine/Source/Editor/AnimGraph/Classes/AnimGraphNode_BlendSpacePlayer.h
+++ b/Engine/Source/Editor/AnimGraph/Classes/AnimGraphNode_BlendSpacePlayer.h
@@ -20,6 +20,7 @@ class UAnimGraphNode_BlendSpacePlayer : public UAnimGraphNode_BlendSpaceBase
 	// UEdGraphNode interface
 	virtual FString GetTooltip() const override;
 	virtual FText GetNodeTitle(ENodeTitleType::Type TitleType) const override;
+	virtual void GetMenuActions(TArray<UBlueprintNodeSpawner*>& ActionListOut) const override;
 	// End of UEdGraphNode interface
 
 	// UAnimGraphNode_Base interface
diff --git a/Engine/Source/Editor/AnimGraph/Classes/AnimGraphNode_Root.h b/Engine/Source/Editor/AnimGraph/Classes/AnimGraphNode_Root.h
--- a/Engine/Source/Editor/AnimGraph/Classes/AnimGraphNode_Root.h
+++ b/Engine/Source/Editor/AnimGraph/Classes/AnimGraphNode_Root.h
@@ -20,6 +20,7 @@ class UAnimGraphNode_Root : public UAnimGraphNode_Base
 	virtual bool CanUserDeleteNode() const override { return false; }
 	virtual bool CanDuplicateNode() const override { return false; }
 	virtual void GetMenuEntries(FGraphContextMenuBuilder& ContextMenuBuilder) const override;
+	virtual void GetMenuActions(TArray<UBlueprintNodeSpawner*>& ActionListOut) const override;
 	// End UEdGraphNode interface.
 
 	// UAnimGraphNode_Base interface
diff --git a/Engine/Source/Editor/AnimGraph/Classes/AnimGraphNode_RotationOffsetBlendSpace.h b/Engine/Source/Editor/AnimGraph/Classes/AnimGraphNode_RotationOffsetBlendSpace.h
--- a/Engine/Source/Editor/AnimGraph/Classes/AnimGraphNode_RotationOffsetBlendSpace.h
+++ b/Engine/Source/Editor/AnimGraph/Classes/AnimGraphNode_RotationOffsetBlendSpace.h
@@ -26,5 +26,6 @@ class UAnimGraphNode_RotationOffsetBlendSpace: public UAnimGraphNode_BlendSpaceB
 
 	// UK2Node interface
 	virtual void GetContextMenuActions(const FGraphNodeContextMenuBuilder& Context) const override;
+	virtual void GetMenuActions(TArray<UBlueprintNodeSpawner*>& ActionListOut) const override;
 	// End of UK2Node interface
 };
diff --git a/Engine/Source/Editor/AnimGraph/Classes/AnimGraphNode_SaveCachedPose.h b/Engine/Source/Editor/AnimGraph/Classes/AnimGraphNode_SaveCachedPose.h
--- a/Engine/Source/Editor/AnimGraph/Classes/AnimGraphNode_SaveCachedPose.h
+++ b/Engine/Source/Editor/AnimGraph/Classes/AnimGraphNode_SaveCachedPose.h
@@ -21,8 +21,13 @@ class UAnimGraphNode_SaveCachedPose : public UAnimGraphNode_Base
 	virtual FText GetNodeTitle(ENodeTitleType::Type TitleType) const override;
 	virtual void OnRenameNode(const FString& NewName) override;
 	virtual TSharedPtr<class INameValidatorInterface> MakeNameValidator() const override;
+	virtual bool IsCompatibleWithGraph(const UEdGraph* TargetGraph) const override;
 	// End of UEdGraphNode interface
 
+	// UK2Node interface.
+	virtual void GetMenuActions(TArray<UBlueprintNodeSpawner*>& ActionListOut) const override;
+	// End of UK2Node interface
+
 	// UAnimGraphNode_Base interface
 	virtual FString GetNodeCategory() const override;
 	virtual void GetMenuEntries(FGraphContextMenuBuilder& ContextMenuBuilder) const override;
diff --git a/Engine/Source/Editor/AnimGraph/Classes/AnimGraphNode_SequenceEvaluator.h b/Engine/Source/Editor/AnimGraph/Classes/AnimGraphNode_SequenceEvaluator.h
--- a/Engine/Source/Editor/AnimGraph/Classes/AnimGraphNode_SequenceEvaluator.h
+++ b/Engine/Source/Editor/AnimGraph/Classes/AnimGraphNode_SequenceEvaluator.h
@@ -33,5 +33,6 @@ class UAnimGraphNode_SequenceEvaluator : public UAnimGraphNode_Base
 
 	// UK2Node interface
 	virtual void GetContextMenuActions(const FGraphNodeContextMenuBuilder& Context) const override;
+	virtual void GetMenuActions(TArray<UBlueprintNodeSpawner*>& ActionListOut) const override;
 	// End of UK2Node interface
 };
diff --git a/Engine/Source/Editor/AnimGraph/Classes/AnimGraphNode_TransitionPoseEvaluator.h b/Engine/Source/Editor/AnimGraph/Classes/AnimGraphNode_TransitionPoseEvaluator.h
--- a/Engine/Source/Editor/AnimGraph/Classes/AnimGraphNode_TransitionPoseEvaluator.h
+++ b/Engine/Source/Editor/AnimGraph/Classes/AnimGraphNode_TransitionPoseEvaluator.h
@@ -26,4 +26,8 @@ class UAnimGraphNode_TransitionPoseEvaluator : public UAnimGraphNode_Base
 	virtual FString GetNodeCategory() const override;
 	virtual void GetMenuEntries(FGraphContextMenuBuilder& ContextMenuBuilder) const override;
 	// End of UAnimGraphNode_Base interface
+
+	// UK2Node interface.
+	virtual void GetMenuActions(TArray<UBlueprintNodeSpawner*>& ActionListOut) const override;
+	// End of UK2Node interface
 };
\ No newline at end of file
diff --git a/Engine/Source/Editor/AnimGraph/Classes/AnimGraphNode_TransitionResult.h b/Engine/Source/Editor/AnimGraph/Classes/AnimGraphNode_TransitionResult.h
--- a/Engine/Source/Editor/AnimGraph/Classes/AnimGraphNode_TransitionResult.h
+++ b/Engine/Source/Editor/AnimGraph/Classes/AnimGraphNode_TransitionResult.h
@@ -22,6 +22,10 @@ class UAnimGraphNode_TransitionResult : public UAnimGraphNode_Base
 	virtual void GetMenuEntries(FGraphContextMenuBuilder& ContextMenuBuilder) const override;
 	// End of UEdGraphNode interface
 
+	// UK2Node interface.
+	virtual void GetMenuActions(TArray<UBlueprintNodeSpawner*>& ActionListOut) const override;
+	// End of UK2Node interface
+
 	// UAnimGraphNode_Base interface
 	virtual bool IsSinkNode() const override { return true; }
 	// End of UAnimGraphNode_Base interface
diff --git a/Engine/Source/Editor/AnimGraph/Classes/AnimGraphNode_WheelHandler.h b/Engine/Source/Editor/AnimGraph/Classes/AnimGraphNode_WheelHandler.h
--- a/Engine/Source/Editor/AnimGraph/Classes/AnimGraphNode_WheelHandler.h
+++ b/Engine/Source/Editor/AnimGraph/Classes/AnimGraphNode_WheelHandler.h
@@ -22,6 +22,7 @@ public:
 	// validate if this is within VehicleAnimInstance
 	virtual void ValidateAnimNodePostCompile(class FCompilerResultsLog& MessageLog, class UAnimBlueprintGeneratedClass* CompiledClass, int32 CompiledNodeIndex) override;
 	virtual void GetMenuEntries(FGraphContextMenuBuilder& ContextMenuBuilder) const override;
+	virtual bool IsCompatibleWithGraph(const UEdGraph* TargetGraph) const override;
 	// End of UEdGraphNode interface
 
 protected:
diff --git a/Engine/Source/Editor/AnimGraph/Private/AnimGraphNode_Base.cpp b/Engine/Source/Editor/AnimGraph/Private/AnimGraphNode_Base.cpp
--- a/Engine/Source/Editor/AnimGraph/Private/AnimGraphNode_Base.cpp
+++ b/Engine/Source/Editor/AnimGraph/Private/AnimGraphNode_Base.cpp
@@ -5,6 +5,7 @@
 #include "AnimGraphNode_Base.h"
 #include "K2ActionMenuBuilder.h" // for FK2ActionMenuBuilder::AddNewNodeAction()
 #include "AnimationGraphSchema.h"
+#include "BlueprintNodeSpawner.h"
 
 /////////////////////////////////////////////////////
 // FA3NodeOptionalPinManager
@@ -226,6 +227,18 @@ void UAnimGraphNode_Base::GetNodeAttributes( TArray<TKeyValuePair<FString, FStri
 	OutNodeAttributes.Add( TKeyValuePair<FString, FString>( TEXT( "Name" ), GetName() ));
 }
 
+void UAnimGraphNode_Base::GetMenuActions(TArray<UBlueprintNodeSpawner*>& ActionListOut) const
+{
+	UBlueprintNodeSpawner* NodeSpawner = UBlueprintNodeSpawner::Create(GetClass());
+	check(NodeSpawner != nullptr);
+	ActionListOut.Add(NodeSpawner);
+}
+
+FText UAnimGraphNode_Base::GetMenuCategory() const
+{
+	return FText::FromString(GetNodeCategory());
+}
+
 TSharedPtr<FEdGraphSchemaAction_K2NewNode> UAnimGraphNode_Base::CreateDefaultMenuEntry(FGraphContextMenuBuilder& ContextMenuBuilder) const
 {
 	UAnimGraphNode_Base* TemplateNode = NewObject<UAnimGraphNode_Base>(GetTransientPackage(), GetClass());
diff --git a/Engine/Source/Editor/AnimGraph/Private/AnimGraphNode_BlendSpaceEvaluator.cpp b/Engine/Source/Editor/AnimGraph/Private/AnimGraphNode_BlendSpaceEvaluator.cpp
--- a/Engine/Source/Editor/AnimGraph/Private/AnimGraphNode_BlendSpaceEvaluator.cpp
+++ b/Engine/Source/Editor/AnimGraph/Private/AnimGraphNode_BlendSpaceEvaluator.cpp
@@ -54,6 +54,12 @@ void UAnimGraphNode_BlendSpaceEvaluator::GetMenuEntries(FGraphContextMenuBuilder
 	// You can convert a regular blend space player to an evaluator via the right click context menu
 }
 
+void UAnimGraphNode_BlendSpaceEvaluator::GetMenuActions(TArray<UBlueprintNodeSpawner*>& ActionListOut) const
+{
+	// Intentionally empty so that we don't get duplicate blend space entries.
+	// You can convert a regular blend space player to an evaluator via the right click context menu
+}
+
 void UAnimGraphNode_BlendSpaceEvaluator::ValidateAnimNodeDuringCompilation(class USkeleton* ForSkeleton, class FCompilerResultsLog& MessageLog)
 {
 	if (Node.BlendSpace == NULL)
diff --git a/Engine/Source/Editor/AnimGraph/Private/AnimGraphNode_BlendSpacePlayer.cpp b/Engine/Source/Editor/AnimGraph/Private/AnimGraphNode_BlendSpacePlayer.cpp
--- a/Engine/Source/Editor/AnimGraph/Private/AnimGraphNode_BlendSpacePlayer.cpp
+++ b/Engine/Source/Editor/AnimGraph/Private/AnimGraphNode_BlendSpacePlayer.cpp
@@ -4,6 +4,7 @@
 #include "GraphEditorActions.h"
 #include "CompilerResultsLog.h"
 #include "AnimGraphNode_BlendSpacePlayer.h"
+#include "BlueprintNodeSpawner.h"
 
 /////////////////////////////////////////////////////
 // UAnimGraphNode_BlendSpacePlayer
@@ -54,6 +55,33 @@ void UAnimGraphNode_BlendSpacePlayer::GetMenuEntries(FGraphContextMenuBuilder& C
 	GetBlendSpaceEntries(bWantAimOffsets, ContextMenuBuilder);
 }
 
+void UAnimGraphNode_BlendSpacePlayer::GetMenuActions(TArray<UBlueprintNodeSpawner*>& ActionListOut) const
+{
+	auto PostSpawnSetupLambda = [](UEdGraphNode* NewNode, bool /*bIsTemplateNode*/, TWeakObjectPtr<UBlendSpaceBase> BlendSpace)
+	{
+		UAnimGraphNode_BlendSpacePlayer* BlendSpaceNode = CastChecked<UAnimGraphNode_BlendSpacePlayer>(NewNode);
+		BlendSpaceNode->Node.BlendSpace = BlendSpace.Get();		
+	};
+
+	for (TObjectIterator<UBlendSpaceBase> BlendSpaceIt; BlendSpaceIt; ++BlendSpaceIt)
+	{
+		UBlendSpaceBase* BlendSpace = *BlendSpaceIt;
+
+		bool const bIsAimOffset = BlendSpace->IsA(UAimOffsetBlendSpace::StaticClass()) || 
+			BlendSpace->IsA(UAimOffsetBlendSpace1D::StaticClass());
+		if (!bIsAimOffset)
+		{
+			UBlueprintNodeSpawner* NodeSpawner = UBlueprintNodeSpawner::Create(GetClass());
+			check(NodeSpawner != nullptr);
+			// @TODO: Need to file this one under the BlendSpace asset (so it can be cleared/updated when the asset is)
+			ActionListOut.Add(NodeSpawner);
+
+			TWeakObjectPtr<UBlendSpaceBase> BlendSpacePtr = BlendSpace;
+			NodeSpawner->CustomizeNodeDelegate = UBlueprintNodeSpawner::FCustomizeNodeDelegate::CreateStatic(PostSpawnSetupLambda, BlendSpacePtr);
+		}
+	}
+}
+
 void UAnimGraphNode_BlendSpacePlayer::ValidateAnimNodeDuringCompilation(class USkeleton* ForSkeleton, class FCompilerResultsLog& MessageLog)
 {
 	if (Node.BlendSpace == NULL)
diff --git a/Engine/Source/Editor/AnimGraph/Private/AnimGraphNode_BoneDrivenController.cpp b/Engine/Source/Editor/AnimGraph/Private/AnimGraphNode_BoneDrivenController.cpp
--- a/Engine/Source/Editor/AnimGraph/Private/AnimGraphNode_BoneDrivenController.cpp
+++ b/Engine/Source/Editor/AnimGraph/Private/AnimGraphNode_BoneDrivenController.cpp
@@ -20,20 +20,30 @@ FText UAnimGraphNode_BoneDrivenController::GetNodeTitle(ENodeTitleType::Type Tit
 {
 	FFormatNamedArguments Args;
 	Args.Add(TEXT("ControllerDesc"), GetControllerDescription());
+	Args.Add(TEXT("SourceBone"), FText::FromName(Node.SourceBone.BoneName));
+	Args.Add(TEXT("TargetBone"), FText::FromName(Node.TargetBone.BoneName));
 
-	if(TitleType == ENodeTitleType::ListView)
+	FText NodeTitle;
+	if (TitleType == ENodeTitleType::ListView)
 	{
 		Args.Add(TEXT("Delim"), FText::FromString(TEXT(" - ")));
+
+		if ((Node.SourceBone.BoneName == NAME_None) && (Node.TargetBone.BoneName == NAME_None))
+		{
+			NodeTitle = FText::Format(LOCTEXT("AnimGraphNode_BoneDrivenController_MenuTitle", "{ControllerDesc}"), Args);
+		}
 	}
 	else
 	{
 		Args.Add(TEXT("Delim"), FText::FromString(TEXT("\n")));
 	}
 
-	Args.Add(TEXT("SourceBone"), FText::FromName(Node.SourceBone.BoneName));
-	Args.Add(TEXT("TargetBone"), FText::FromName(Node.TargetBone.BoneName));
-
-	return FText::Format(LOCTEXT("AnimGraphNode_BoneDrivenController_Title", "{ControllerDesc}{Delim}Driving Bone: {SourceBone}{Delim}Driven Bone: {TargetBone}"), Args);
+	if (NodeTitle.IsEmpty())
+	{
+		NodeTitle = FText::Format(LOCTEXT("AnimGraphNode_BoneDrivenController_Title", "{ControllerDesc}{Delim}Driving Bone: {SourceBone}{Delim}Driven Bone: {TargetBone}"), Args);
+	}
+	
+	return NodeTitle;
 }
 
 FText UAnimGraphNode_BoneDrivenController::GetControllerDescription() const
diff --git a/Engine/Source/Editor/AnimGraph/Private/AnimGraphNode_CopyBone.cpp b/Engine/Source/Editor/AnimGraph/Private/AnimGraphNode_CopyBone.cpp
--- a/Engine/Source/Editor/AnimGraph/Private/AnimGraphNode_CopyBone.cpp
+++ b/Engine/Source/Editor/AnimGraph/Private/AnimGraphNode_CopyBone.cpp
@@ -29,7 +29,24 @@ FText UAnimGraphNode_CopyBone::GetNodeTitle(ENodeTitleType::Type TitleType) cons
 	Args.Add(TEXT("ControllerDescription"), GetControllerDescription());
 	Args.Add(TEXT("SourceBoneName"), FText::FromName(Node.SourceBone.BoneName));
 	Args.Add(TEXT("TargetBoneName"), FText::FromName(Node.TargetBone.BoneName));
-	return FText::Format(LOCTEXT("AnimGraphNode_CopyBone_Title", "{ControllerDescription}\nSource Bone: {SourceBoneName}\nTarget Bone: {TargetBoneName}"), Args);
+
+	FText NodeTitle;
+	if (TitleType == ENodeTitleType::ListView)
+	{
+		if ((Node.SourceBone.BoneName == NAME_None) && (Node.TargetBone.BoneName == NAME_None))
+		{
+			NodeTitle = FText::Format(LOCTEXT("AnimGraphNode_CopyBone_MenuTitle", "{ControllerDescription}"), Args);
+		}
+		else
+		{
+			NodeTitle = FText::Format(LOCTEXT("AnimGraphNode_CopyBone_ListTitle", "{ControllerDescription} - Source Bone: {SourceBoneName} - Target Bone: {TargetBoneName}"), Args);
+		}
+	}
+	else
+	{
+		NodeTitle = FText::Format(LOCTEXT("AnimGraphNode_CopyBone_Title", "{ControllerDescription}\nSource Bone: {SourceBoneName}\nTarget Bone: {TargetBoneName}"), Args);
+	}
+	return NodeTitle;
 }
 
 #undef LOCTEXT_NAMESPACE
diff --git a/Engine/Source/Editor/AnimGraph/Private/AnimGraphNode_HandIKRetargeting.cpp b/Engine/Source/Editor/AnimGraph/Private/AnimGraphNode_HandIKRetargeting.cpp
--- a/Engine/Source/Editor/AnimGraph/Private/AnimGraphNode_HandIKRetargeting.cpp
+++ b/Engine/Source/Editor/AnimGraph/Private/AnimGraphNode_HandIKRetargeting.cpp
@@ -25,9 +25,7 @@ FString UAnimGraphNode_HandIKRetargeting::GetTooltip() const
 
 FText UAnimGraphNode_HandIKRetargeting::GetNodeTitle(ENodeTitleType::Type TitleType) const
 {
-	FFormatNamedArguments Args;
-	Args.Add(TEXT("ControllerDescription"), GetControllerDescription());
-	return FText::Format(LOCTEXT("AnimGraphNode_CopyBone_Title", "{ControllerDescription}"), Args);
+	return GetControllerDescription();
 }
 
 #undef LOCTEXT_NAMESPACE
diff --git a/Engine/Source/Editor/AnimGraph/Private/AnimGraphNode_LookAt.cpp b/Engine/Source/Editor/AnimGraph/Private/AnimGraphNode_LookAt.cpp
--- a/Engine/Source/Editor/AnimGraph/Private/AnimGraphNode_LookAt.cpp
+++ b/Engine/Source/Editor/AnimGraph/Private/AnimGraphNode_LookAt.cpp
@@ -36,7 +36,11 @@ FText UAnimGraphNode_LookAt::GetNodeTitle(ENodeTitleType::Type TitleType) const
 
 	if(TitleType == ENodeTitleType::ListView)
 	{
-		return FText::Format(LOCTEXT("AnimGraphNode_LookAt_Title", "{ControllerDescription} - Bone: {BoneName}"), Args);
+		if (Node.BoneToModify.BoneName == NAME_None)
+		{
+			return FText::Format(LOCTEXT("AnimGraphNode_LookAt_MenuTitle", "{ControllerDescription}"), Args);
+		}
+		return FText::Format(LOCTEXT("AnimGraphNode_LookAt_ListTitle", "{ControllerDescription} - Bone: {BoneName}"), Args);
 	}
 	else
 	{
diff --git a/Engine/Source/Editor/AnimGraph/Private/AnimGraphNode_ModifyBone.cpp b/Engine/Source/Editor/AnimGraph/Private/AnimGraphNode_ModifyBone.cpp
--- a/Engine/Source/Editor/AnimGraph/Private/AnimGraphNode_ModifyBone.cpp
+++ b/Engine/Source/Editor/AnimGraph/Private/AnimGraphNode_ModifyBone.cpp
@@ -36,7 +36,11 @@ FText UAnimGraphNode_ModifyBone::GetNodeTitle(ENodeTitleType::Type TitleType) co
 
 	if(TitleType == ENodeTitleType::ListView)
 	{
-		return FText::Format(LOCTEXT("AnimGraphNode_ModifyBone_Title", "{ControllerDescription} - Bone: {BoneName}"), Args);
+		if (Node.BoneToModify.BoneName == NAME_None)
+		{
+			return FText::Format(LOCTEXT("AnimGraphNode_ModifyBone_MenuTitle", "{ControllerDescription}"), Args);
+		}
+		return FText::Format(LOCTEXT("AnimGraphNode_ModifyBone_ListTitle", "{ControllerDescription} - Bone: {BoneName}"), Args);
 	}
 	else
 	{
diff --git a/Engine/Source/Editor/AnimGraph/Private/AnimGraphNode_Root.cpp b/Engine/Source/Editor/AnimGraph/Private/AnimGraphNode_Root.cpp
--- a/Engine/Source/Editor/AnimGraph/Private/AnimGraphNode_Root.cpp
+++ b/Engine/Source/Editor/AnimGraph/Private/AnimGraphNode_Root.cpp
@@ -68,6 +68,11 @@ void UAnimGraphNode_Root::GetMenuEntries(FGraphContextMenuBuilder& ContextMenuBu
 	// Intentionally empty.  This node is autogenerated when a new graph is created.
 }
 
+void UAnimGraphNode_Root::GetMenuActions(TArray<UBlueprintNodeSpawner*>& ActionListOut) const
+{
+	// Intentionally empty. This node is auto-generated when a new graph is created.
+}
+
 FString UAnimGraphNode_Root::GetDocumentationLink() const
 {
 	return TEXT("Shared/GraphNodes/AnimationStateMachine");
diff --git a/Engine/Source/Editor/AnimGraph/Private/AnimGraphNode_RotationMultiplier.cpp b/Engine/Source/Editor/AnimGraph/Private/AnimGraphNode_RotationMultiplier.cpp
--- a/Engine/Source/Editor/AnimGraph/Private/AnimGraphNode_RotationMultiplier.cpp
+++ b/Engine/Source/Editor/AnimGraph/Private/AnimGraphNode_RotationMultiplier.cpp
@@ -30,7 +30,11 @@ FText UAnimGraphNode_RotationMultiplier::GetNodeTitle(ENodeTitleType::Type Title
 
 	if(TitleType == ENodeTitleType::ListView)
 	{
-		return FText::Format(LOCTEXT("AnimGraphNode_ModifyBone_Title", "{ControllerDescription} - Bone: {BoneName}"), Args);
+		if (Node.TargetBone.BoneName == NAME_None)
+		{
+			return FText::Format(LOCTEXT("AnimGraphNode_ModifyBone_MenuTitle", "{ControllerDescription}"), Args);
+		}
+		return FText::Format(LOCTEXT("AnimGraphNode_ModifyBone_ListTitle", "{ControllerDescription} - Bone: {BoneName}"), Args);
 	}
 	else
 	{
diff --git a/Engine/Source/Editor/AnimGraph/Private/AnimGraphNode_RotationOffsetBlendSpace.cpp b/Engine/Source/Editor/AnimGraph/Private/AnimGraphNode_RotationOffsetBlendSpace.cpp
--- a/Engine/Source/Editor/AnimGraph/Private/AnimGraphNode_RotationOffsetBlendSpace.cpp
+++ b/Engine/Source/Editor/AnimGraph/Private/AnimGraphNode_RotationOffsetBlendSpace.cpp
@@ -4,6 +4,7 @@
 #include "AnimGraphNode_RotationOffsetBlendSpace.h"
 #include "GraphEditorActions.h"
 #include "CompilerResultsLog.h"
+#include "BlueprintNodeSpawner.h"
 
 /////////////////////////////////////////////////////
 // UAnimGraphNode_RotationOffsetBlendSpace
@@ -43,6 +44,33 @@ void UAnimGraphNode_RotationOffsetBlendSpace::GetMenuEntries(FGraphContextMenuBu
 	GetBlendSpaceEntries(bWantAimOffsets, ContextMenuBuilder);
 }
 
+void UAnimGraphNode_RotationOffsetBlendSpace::GetMenuActions(TArray<UBlueprintNodeSpawner*>& ActionListOut) const
+{
+	auto PostSpawnSetupLambda = [](UEdGraphNode* NewNode, bool /*bIsTemplateNode*/, TWeakObjectPtr<UBlendSpaceBase> BlendSpace)
+	{
+		UAnimGraphNode_RotationOffsetBlendSpace* BlendSpaceNode = CastChecked<UAnimGraphNode_RotationOffsetBlendSpace>(NewNode);
+		BlendSpaceNode->Node.BlendSpace = BlendSpace.Get();
+	};
+
+	for (TObjectIterator<UBlendSpaceBase> BlendSpaceIt; BlendSpaceIt; ++BlendSpaceIt)
+	{
+		UBlendSpaceBase* BlendSpace = *BlendSpaceIt;
+
+		bool const bIsAimOffset = BlendSpace->IsA(UAimOffsetBlendSpace::StaticClass()) ||
+			BlendSpace->IsA(UAimOffsetBlendSpace1D::StaticClass());
+		if (bIsAimOffset)
+		{
+			UBlueprintNodeSpawner* NodeSpawner = UBlueprintNodeSpawner::Create(GetClass());
+			check(NodeSpawner != nullptr);
+			// @TODO: Need to file this one under the BlendSpace asset (so it can be cleared/updated when the asset is)
+			ActionListOut.Add(NodeSpawner);
+
+			TWeakObjectPtr<UBlendSpaceBase> BlendSpacePtr = BlendSpace;
+			NodeSpawner->CustomizeNodeDelegate = UBlueprintNodeSpawner::FCustomizeNodeDelegate::CreateStatic(PostSpawnSetupLambda, BlendSpacePtr);
+		}
+	}
+}
+
 void UAnimGraphNode_RotationOffsetBlendSpace::ValidateAnimNodeDuringCompilation(class USkeleton* ForSkeleton, class FCompilerResultsLog& MessageLog)
 {
 	if (Node.BlendSpace == NULL)
diff --git a/Engine/Source/Editor/AnimGraph/Private/AnimGraphNode_SaveCachedPose.cpp b/Engine/Source/Editor/AnimGraph/Private/AnimGraphNode_SaveCachedPose.cpp
--- a/Engine/Source/Editor/AnimGraph/Private/AnimGraphNode_SaveCachedPose.cpp
+++ b/Engine/Source/Editor/AnimGraph/Private/AnimGraphNode_SaveCachedPose.cpp
@@ -7,6 +7,7 @@
 #include "Kismet2NameValidators.h"
 #include "K2ActionMenuBuilder.h" // for FK2ActionMenuBuilder::AddNewNodeAction()
 #include "AnimGraphNode_SaveCachedPose.h"
+#include "BlueprintNodeSpawner.h"
 
 /////////////////////////////////////////////////////
 // FCachedPoseNameValidator
@@ -36,7 +37,6 @@ UAnimGraphNode_SaveCachedPose::UAnimGraphNode_SaveCachedPose(const FPostConstruc
 	: Super(PCIP)
 {
 	bCanRenameNode = true;
-	CacheName = TEXT("SavedPose");
 }
 
 FString UAnimGraphNode_SaveCachedPose::GetTooltip() const
@@ -50,6 +50,10 @@ FText UAnimGraphNode_SaveCachedPose::GetNodeTitle(ENodeTitleType::Type TitleType
 	{
 		return FText::FromString(CacheName);
 	}
+	else if ((TitleType == ENodeTitleType::ListView) && CacheName.IsEmpty())
+	{
+		return LOCTEXT("NewSaveCachedPose", "New Save cached pose...");
+	}
 	else
 	{
 		FFormatNamedArguments Args;
@@ -75,7 +79,7 @@ void UAnimGraphNode_SaveCachedPose::GetMenuEntries(FGraphContextMenuBuilder& Con
 		{
 			// Offer save cached pose
 			UAnimGraphNode_SaveCachedPose* SaveCachedPose = NewObject<UAnimGraphNode_SaveCachedPose>();
-			SaveCachedPose->CacheName += FString::FromInt(FMath::Rand());
+			SaveCachedPose->CacheName = TEXT("SavedPose") + FString::FromInt(FMath::Rand());
 
 			TSharedPtr<FEdGraphSchemaAction_K2NewNode> SaveCachedPoseAction = FK2ActionMenuBuilder::AddNewNodeAction(ContextMenuBuilder, GetNodeCategory(), LOCTEXT("NewSaveCachedPose", "New Save cached pose..."), SaveCachedPose->GetTooltip(), 0, SaveCachedPose->GetKeywords());
 			SaveCachedPoseAction->NodeTemplate = SaveCachedPose;
@@ -83,6 +87,34 @@ void UAnimGraphNode_SaveCachedPose::GetMenuEntries(FGraphContextMenuBuilder& Con
 	}
 }
 
+void UAnimGraphNode_SaveCachedPose::GetMenuActions(TArray<UBlueprintNodeSpawner*>& ActionListOut) const
+{
+	auto PostSpawnSetupLambda = [](UEdGraphNode* NewNode, bool bIsTemplateNode)
+	{
+		UAnimGraphNode_SaveCachedPose* CachedPoseNode = CastChecked<UAnimGraphNode_SaveCachedPose>(NewNode);
+		// we use an empty CacheName in GetNodeTitle() to relay the proper menu title
+		if (!bIsTemplateNode)
+		{
+			// @TODO: is the idea that this name is unique? what if Rand() hit twice? why not MakeUniqueObjectName()?			
+			CachedPoseNode->CacheName = TEXT("SavedPose") + FString::FromInt(FMath::Rand());
+		}
+	};
+
+	UBlueprintNodeSpawner* NodeSpawner = UBlueprintNodeSpawner::Create(GetClass());
+	NodeSpawner->CustomizeNodeDelegate = UBlueprintNodeSpawner::FCustomizeNodeDelegate::CreateStatic(PostSpawnSetupLambda);
+
+	ActionListOut.Add(NodeSpawner);
+}
+
+bool UAnimGraphNode_SaveCachedPose::IsCompatibleWithGraph(const UEdGraph* TargetGraph) const
+{
+	//EGraphType GraphType = TargetGraph->GetSchema()->GetGraphType(TargetGraph);
+	//bool const bIsNotStateMachine = (GraphType != GT_StateMachine);
+
+	bool const bIsNotStateMachine = TargetGraph->GetOuter()->IsA(UAnimBlueprint::StaticClass());
+	return bIsNotStateMachine && Super::IsCompatibleWithGraph(TargetGraph);
+}
+
 void UAnimGraphNode_SaveCachedPose::OnRenameNode(const FString& NewName)
 {
 	CacheName = NewName;
diff --git a/Engine/Source/Editor/AnimGraph/Private/AnimGraphNode_SequenceEvaluator.cpp b/Engine/Source/Editor/AnimGraph/Private/AnimGraphNode_SequenceEvaluator.cpp
--- a/Engine/Source/Editor/AnimGraph/Private/AnimGraphNode_SequenceEvaluator.cpp
+++ b/Engine/Source/Editor/AnimGraph/Private/AnimGraphNode_SequenceEvaluator.cpp
@@ -70,6 +70,11 @@ void UAnimGraphNode_SequenceEvaluator::GetMenuEntries(FGraphContextMenuBuilder&
 	// Intentionally empty; you can drop down a regular sequence player and convert into a sequence evaluator in the right-click menu.
 }
 
+void UAnimGraphNode_SequenceEvaluator::GetMenuActions(TArray<UBlueprintNodeSpawner*>& ActionListOut) const
+{
+	// Intentionally empty; you can drop down a regular sequence player and convert into a sequence evaluator in the right-click menu.
+}
+
 void UAnimGraphNode_SequenceEvaluator::ValidateAnimNodeDuringCompilation(class USkeleton* ForSkeleton, class FCompilerResultsLog& MessageLog)
 {
 	if (Node.Sequence == NULL)
diff --git a/Engine/Source/Editor/AnimGraph/Private/AnimGraphNode_SpringBone.cpp b/Engine/Source/Editor/AnimGraph/Private/AnimGraphNode_SpringBone.cpp
--- a/Engine/Source/Editor/AnimGraph/Private/AnimGraphNode_SpringBone.cpp
+++ b/Engine/Source/Editor/AnimGraph/Private/AnimGraphNode_SpringBone.cpp
@@ -31,7 +31,11 @@ FText UAnimGraphNode_SpringBone::GetNodeTitle(ENodeTitleType::Type TitleType) co
 
 	if(TitleType == ENodeTitleType::ListView)
 	{
-		return FText::Format(LOCTEXT("AnimGraphNode_SpringBone_Title", "{ControllerDescription} - Bone: {BoneName}"), Args);
+		if (Node.SpringBone.BoneName == NAME_None)
+		{
+			return FText::Format(LOCTEXT("AnimGraphNode_SpringBone_MenuTitle", "{ControllerDescription}"), Args);
+		}
+		return FText::Format(LOCTEXT("AnimGraphNode_SpringBone_ListTitle", "{ControllerDescription} - Bone: {BoneName}"), Args);
 	}
 	else
 	{
diff --git a/Engine/Source/Editor/AnimGraph/Private/AnimGraphNode_StateMachineBase.cpp b/Engine/Source/Editor/AnimGraph/Private/AnimGraphNode_StateMachineBase.cpp
--- a/Engine/Source/Editor/AnimGraph/Private/AnimGraphNode_StateMachineBase.cpp
+++ b/Engine/Source/Editor/AnimGraph/Private/AnimGraphNode_StateMachineBase.cpp
@@ -60,7 +60,11 @@ FString UAnimGraphNode_StateMachineBase::GetTooltip() const
 FText UAnimGraphNode_StateMachineBase::GetNodeTitle(ENodeTitleType::Type TitleType) const
 {
 	const FText FirstLine = (EditorStateMachineGraph != NULL) ? FText::FromString(EditorStateMachineGraph->GetName()) : LOCTEXT("ErrorNoGraph", "Error: No Graph");
-	if(TitleType == ENodeTitleType::FullTitle)
+	if ((TitleType == ENodeTitleType::ListView) && (EditorStateMachineGraph == nullptr))
+	{
+		return LOCTEXT("AddNewStateMachine", "Add New State Machine...");
+	}
+	else if(TitleType == ENodeTitleType::FullTitle)
 	{
 		FFormatNamedArguments Args;
 		Args.Add(TEXT("Title"), FirstLine);
diff --git a/Engine/Source/Editor/AnimGraph/Private/AnimGraphNode_Trail.cpp b/Engine/Source/Editor/AnimGraph/Private/AnimGraphNode_Trail.cpp
--- a/Engine/Source/Editor/AnimGraph/Private/AnimGraphNode_Trail.cpp
+++ b/Engine/Source/Editor/AnimGraph/Private/AnimGraphNode_Trail.cpp
@@ -31,7 +31,11 @@ FText UAnimGraphNode_Trail::GetNodeTitle(ENodeTitleType::Type TitleType) const
 
 	if (TitleType == ENodeTitleType::ListView)
 	{
-		return FText::Format(LOCTEXT("AnimGraphNode_Trail_Title", "{ControllerDescription} - Bone: {BoneName}"), Args);
+		if (Node.TrailBone.BoneName == NAME_None)
+		{
+			return FText::Format(LOCTEXT("AnimGraphNode_Trail_MenuTitle", "{ControllerDescription}"), Args);
+		}
+		return FText::Format(LOCTEXT("AnimGraphNode_Trail_ListTitle", "{ControllerDescription} - Bone: {BoneName}"), Args);
 	}
 	else
 	{
diff --git a/Engine/Source/Editor/AnimGraph/Private/AnimGraphNode_TransitionPoseEvaluator.cpp b/Engine/Source/Editor/AnimGraph/Private/AnimGraphNode_TransitionPoseEvaluator.cpp
--- a/Engine/Source/Editor/AnimGraph/Private/AnimGraphNode_TransitionPoseEvaluator.cpp
+++ b/Engine/Source/Editor/AnimGraph/Private/AnimGraphNode_TransitionPoseEvaluator.cpp
@@ -63,4 +63,9 @@ void UAnimGraphNode_TransitionPoseEvaluator::GetMenuEntries(FGraphContextMenuBui
 	// Intentionally empty: Don't allow an option to create them, as they're auto-created when custom blend graphs are made
 }
 
+void UAnimGraphNode_TransitionPoseEvaluator::GetMenuActions(TArray<UBlueprintNodeSpawner*>& ActionListOut) const
+{
+	// Intentionally empty: Don't allow an option to create them, as they're auto-created when custom blend graphs are made
+}
+
 #undef LOCTEXT_NAMESPACE
diff --git a/Engine/Source/Editor/AnimGraph/Private/AnimGraphNode_TransitionResult.cpp b/Engine/Source/Editor/AnimGraph/Private/AnimGraphNode_TransitionResult.cpp
--- a/Engine/Source/Editor/AnimGraph/Private/AnimGraphNode_TransitionResult.cpp
+++ b/Engine/Source/Editor/AnimGraph/Private/AnimGraphNode_TransitionResult.cpp
@@ -37,4 +37,10 @@ void UAnimGraphNode_TransitionResult::GetMenuEntries(FGraphContextMenuBuilder& C
 	// Intentionally empty.  This node is autogenerated when a transition graph is created.
 }
 
+
+void UAnimGraphNode_TransitionResult::GetMenuActions(TArray<UBlueprintNodeSpawner*>& ActionListOut) const
+{
+	// Intentionally empty. This node is auto-generated when a transition graph is created.
+}
+
 #undef LOCTEXT_NAMESPACE
diff --git a/Engine/Source/Editor/AnimGraph/Private/AnimGraphNode_TwoBoneIK.cpp b/Engine/Source/Editor/AnimGraph/Private/AnimGraphNode_TwoBoneIK.cpp
--- a/Engine/Source/Editor/AnimGraph/Private/AnimGraphNode_TwoBoneIK.cpp
+++ b/Engine/Source/Editor/AnimGraph/Private/AnimGraphNode_TwoBoneIK.cpp
@@ -31,7 +31,11 @@ FText UAnimGraphNode_TwoBoneIK::GetNodeTitle(ENodeTitleType::Type TitleType) con
 
 	if(TitleType == ENodeTitleType::ListView)
 	{
-		return FText::Format(LOCTEXT("AnimGraphNode_IKBone_Title", "{ControllerDescription} - Bone: {BoneName}"), Args);
+		if (Node.IKBone.BoneName == NAME_None)
+		{
+			return FText::Format(LOCTEXT("AnimGraphNode_IKBone_MenuTitle", "{ControllerDescription}"), Args);
+		}
+		return FText::Format(LOCTEXT("AnimGraphNode_IKBone_ListTitle", "{ControllerDescription} - Bone: {BoneName}"), Args);
 	}
 	else
 	{
diff --git a/Engine/Source/Editor/AnimGraph/Private/AnimGraphNode_WheelHandler.cpp b/Engine/Source/Editor/AnimGraph/Private/AnimGraphNode_WheelHandler.cpp
--- a/Engine/Source/Editor/AnimGraph/Private/AnimGraphNode_WheelHandler.cpp
+++ b/Engine/Source/Editor/AnimGraph/Private/AnimGraphNode_WheelHandler.cpp
@@ -31,9 +31,19 @@ FString UAnimGraphNode_WheelHandler::GetTooltip() const
 
 FText UAnimGraphNode_WheelHandler::GetNodeTitle(ENodeTitleType::Type TitleType) const
 {
-	// we don't have any run-time information, so it's limited to print anymore than whaat it is
-	// it would be nice to print more data such as name of bones for wheels, but it's not available in Persona
-	return FText(LOCTEXT("AnimGraphNode_WheelHandler_Title", "Wheel Handler"));
+	FText NodeTitle;
+	if (TitleType == ENodeTitleType::ListView)
+	{
+		NodeTitle = GetControllerDescription();
+	}
+	else
+	{
+		// we don't have any run-time information, so it's limited to print  
+		// anymore than what it is it would be nice to print more data such as 
+		// name of bones for wheels, but it's not available in Persona
+		NodeTitle = FText(LOCTEXT("AnimGraphNode_WheelHandler_Title", "Wheel Handler"));
+	}	
+	return NodeTitle;
 }
 
 void UAnimGraphNode_WheelHandler::ValidateAnimNodePostCompile(class FCompilerResultsLog& MessageLog, class UAnimBlueprintGeneratedClass* CompiledClass, int32 CompiledNodeIndex)
@@ -57,4 +67,10 @@ void UAnimGraphNode_WheelHandler::GetMenuEntries(FGraphContextMenuBuilder& Conte
 	// else we don't show
 }
 
+bool UAnimGraphNode_WheelHandler::IsCompatibleWithGraph(const UEdGraph* TargetGraph) const
+{
+	UBlueprint* Blueprint = FBlueprintEditorUtils::FindBlueprintForGraph(TargetGraph);
+	return (Blueprint != nullptr) && Blueprint->ParentClass->IsChildOf<UVehicleAnimInstance>() && Super::IsCompatibleWithGraph(TargetGraph);
+}
+
 #undef LOCTEXT_NAMESPACE
diff --git a/Engine/Source/Editor/BlueprintGraph/Classes/K2Node_AddComponent.h b/Engine/Source/Editor/BlueprintGraph/Classes/K2Node_AddComponent.h
--- a/Engine/Source/Editor/BlueprintGraph/Classes/K2Node_AddComponent.h
+++ b/Engine/Source/Editor/BlueprintGraph/Classes/K2Node_AddComponent.h
@@ -25,6 +25,7 @@ class UK2Node_AddComponent : public UK2Node_CallFunction
 	virtual FText GetNodeTitle(ENodeTitleType::Type TitleType) const override;
 	virtual FString GetDocumentationLink() const override;
 	virtual FString GetDocumentationExcerptName() const override;
+	virtual bool IsCompatibleWithGraph(UEdGraph const* Graph) const override;
 	// End UEdGraphNode interface
 
 	// Begin K2Node interface
diff --git a/Engine/Source/Editor/BlueprintGraph/Classes/K2Node_BaseMCDelegate.h b/Engine/Source/Editor/BlueprintGraph/Classes/K2Node_BaseMCDelegate.h
--- a/Engine/Source/Editor/BlueprintGraph/Classes/K2Node_BaseMCDelegate.h
+++ b/Engine/Source/Editor/BlueprintGraph/Classes/K2Node_BaseMCDelegate.h
@@ -29,6 +29,7 @@ public:
 	// UEdGraphNode interface
 	virtual void AllocateDefaultPins() override;
 	virtual void ValidateNodeDuringCompilation(class FCompilerResultsLog& MessageLog) const override;
+	virtual bool IsCompatibleWithGraph(const UEdGraph* TargetGraph) const override;
 	// End of UEdGraphNode interface
 
 	BLUEPRINTGRAPH_API void SetFromProperty(const UProperty* Property, bool bSelfContext)
diff --git a/Engine/Source/Editor/BlueprintGraph/Classes/K2Node_GetInputAxisKeyValue.h b/Engine/Source/Editor/BlueprintGraph/Classes/K2Node_GetInputAxisKeyValue.h
--- a/Engine/Source/Editor/BlueprintGraph/Classes/K2Node_GetInputAxisKeyValue.h
+++ b/Engine/Source/Editor/BlueprintGraph/Classes/K2Node_GetInputAxisKeyValue.h
@@ -25,6 +25,7 @@ class UK2Node_GetInputAxisKeyValue : public UK2Node_CallFunction
 	virtual FText GetNodeTitle(ENodeTitleType::Type TitleType) const override;
 	virtual FString GetTooltip() const override;
 	virtual FName GetPaletteIcon(FLinearColor& OutColor) const override;
+	virtual bool IsCompatibleWithGraph(UEdGraph const* Graph) const override;
 	// End EdGraphNode interface
 
 	// Begin UK2Node interface
diff --git a/Engine/Source/Editor/BlueprintGraph/Classes/K2Node_GetInputAxisValue.h b/Engine/Source/Editor/BlueprintGraph/Classes/K2Node_GetInputAxisValue.h
--- a/Engine/Source/Editor/BlueprintGraph/Classes/K2Node_GetInputAxisValue.h
+++ b/Engine/Source/Editor/BlueprintGraph/Classes/K2Node_GetInputAxisValue.h
@@ -24,6 +24,7 @@ class UK2Node_GetInputAxisValue : public UK2Node_CallFunction
 	virtual void AllocateDefaultPins() override;
 	virtual FText GetNodeTitle(ENodeTitleType::Type TitleType) const override;
 	virtual FString GetTooltip() const override;
+	virtual bool IsCompatibleWithGraph(UEdGraph const* Graph) const override;
 	// End EdGraphNode interface
 
 	// Begin UK2Node interface
diff --git a/Engine/Source/Editor/BlueprintGraph/Classes/K2Node_InputAction.h b/Engine/Source/Editor/BlueprintGraph/Classes/K2Node_InputAction.h
--- a/Engine/Source/Editor/BlueprintGraph/Classes/K2Node_InputAction.h
+++ b/Engine/Source/Editor/BlueprintGraph/Classes/K2Node_InputAction.h
@@ -35,6 +35,7 @@ class UK2Node_InputAction : public UK2Node
 	virtual FText GetNodeTitle(ENodeTitleType::Type TitleType) const override;
 	virtual FString GetTooltip() const override;
 	virtual FName GetPaletteIcon(FLinearColor& OutColor) const override{ return TEXT("GraphEditor.Event_16x"); }
+	virtual bool IsCompatibleWithGraph(UEdGraph const* Graph) const override;
 	// End UEdGraphNode interface.
 
 	// Begin UK2Node interface
diff --git a/Engine/Source/Editor/BlueprintGraph/Classes/K2Node_InputKey.h b/Engine/Source/Editor/BlueprintGraph/Classes/K2Node_InputKey.h
--- a/Engine/Source/Editor/BlueprintGraph/Classes/K2Node_InputKey.h
+++ b/Engine/Source/Editor/BlueprintGraph/Classes/K2Node_InputKey.h
@@ -60,6 +60,7 @@ class UK2Node_InputKey : public UK2Node
 	virtual FText GetNodeTitle(ENodeTitleType::Type TitleType) const override;
 	virtual FString GetTooltip() const override;
 	virtual FName GetPaletteIcon(FLinearColor& OutColor) const override;
+	virtual bool IsCompatibleWithGraph(UEdGraph const* Graph) const override;
 	// End UEdGraphNode interface.
 
 	BLUEPRINTGRAPH_API FText GetModifierText() const;
diff --git a/Engine/Source/Editor/BlueprintGraph/Classes/K2Node_InputTouch.h b/Engine/Source/Editor/BlueprintGraph/Classes/K2Node_InputTouch.h
--- a/Engine/Source/Editor/BlueprintGraph/Classes/K2Node_InputTouch.h
+++ b/Engine/Source/Editor/BlueprintGraph/Classes/K2Node_InputTouch.h
@@ -38,6 +38,7 @@ class UK2Node_InputTouch : public UK2Node
 	virtual FName GetPaletteIcon(FLinearColor& OutColor) const override{ return TEXT("GraphEditor.TouchEvent_16x"); }
 	virtual void GetMenuActions(TArray<UBlueprintNodeSpawner*>& ActionListOut) const override;
 	virtual FText GetMenuCategory() const override;
+	virtual bool IsCompatibleWithGraph(const UEdGraph* TargetGraph) const override;
 	// End UEdGraphNode interface.
 
 	BLUEPRINTGRAPH_API static UEnum* GetTouchIndexEnum();
diff --git a/Engine/Source/Editor/BlueprintGraph/Private/BlueprintNodeTemplateCache.cpp b/Engine/Source/Editor/BlueprintGraph/Private/BlueprintNodeTemplateCache.cpp
--- a/Engine/Source/Editor/BlueprintGraph/Private/BlueprintNodeTemplateCache.cpp
+++ b/Engine/Source/Editor/BlueprintGraph/Private/BlueprintNodeTemplateCache.cpp
@@ -124,7 +124,7 @@ FBlueprintNodeTemplateCache::FBlueprintNodeTemplateCache()
 	using namespace BlueprintNodeTemplateCacheImpl; // for MakeCompatibleBlueprint()
 
 	TemplateOuters.Add(MakeCompatibleBlueprint(UBlueprint::StaticClass(), AActor::StaticClass(), UBlueprintGeneratedClass::StaticClass()));
-	//TemplateOuters.Add(MakeCompatibleBlueprint(UAnimBlueprint::StaticClass(), UAnimInstance::StaticClass(), UAnimBlueprintGeneratedClass::StaticClass()));
+	TemplateOuters.Add(MakeCompatibleBlueprint(UAnimBlueprint::StaticClass(), UAnimInstance::StaticClass(), UAnimBlueprintGeneratedClass::StaticClass()));
 }
 
 //------------------------------------------------------------------------------
@@ -150,11 +150,14 @@ UEdGraphNode* FBlueprintNodeTemplateCache::GetNodeTemplate(UBlueprintNodeSpawner
 		{
 			// by the time we're asking for a prototype for this spawner, we should
 			// be sure that it is compatible with the TargetGraph
-			check(IsCompatible(NodeCDO, TargetGraph));
+			//check(IsCompatible(NodeCDO, TargetGraph));
 
-			UBlueprint* TargetBlueprint = FBlueprintEditorUtils::FindBlueprintForGraph(TargetGraph);
+			TargetBlueprint = FBlueprintEditorUtils::FindBlueprintForGraph(TargetGraph);
 			check(TargetBlueprint != nullptr);
-			TSubclassOf<UBlueprint> BlueprintClass = TargetBlueprint->GetClass();
+			BlueprintClass = TargetBlueprint->GetClass();
+
+			TargetGraph = FindCompatibleGraph(TargetBlueprint, NodeCDO);
+			check(TargetGraph != nullptr);
 		}		
 
 		UBlueprint* CompatibleBlueprint = nullptr;
diff --git a/Engine/Source/Editor/BlueprintGraph/Private/K2Node_AddComponent.cpp b/Engine/Source/Editor/BlueprintGraph/Private/K2Node_AddComponent.cpp
--- a/Engine/Source/Editor/BlueprintGraph/Private/K2Node_AddComponent.cpp
+++ b/Engine/Source/Editor/BlueprintGraph/Private/K2Node_AddComponent.cpp
@@ -7,6 +7,7 @@
 #include "CompilerResultsLog.h"
 #include "CallFunctionHandler.h"
 #include "Particles/ParticleSystemComponent.h"
+#include "BlueprintEditorUtils.h"
 
 #define LOCTEXT_NAMESPACE "K2Node_AddComponent"
 
@@ -395,6 +396,12 @@ FString UK2Node_AddComponent::GetDocumentationExcerptName() const
 	return Super::GetDocumentationExcerptName();
 }
 
+bool UK2Node_AddComponent::IsCompatibleWithGraph(UEdGraph const* Graph) const
+{
+	UBlueprint* Blueprint = FBlueprintEditorUtils::FindBlueprintForGraph(Graph);
+	return (Blueprint != nullptr) && FBlueprintEditorUtils::IsActorBased(Blueprint) && Super::IsCompatibleWithGraph(Graph);
+}
+
 FNodeHandlingFunctor* UK2Node_AddComponent::CreateNodeHandler(FKismetCompilerContext& CompilerContext) const
 {
 	return new FKCHandler_AddComponent(CompilerContext);
diff --git a/Engine/Source/Editor/BlueprintGraph/Private/K2Node_GetInputAxisKeyValue.cpp b/Engine/Source/Editor/BlueprintGraph/Private/K2Node_GetInputAxisKeyValue.cpp
--- a/Engine/Source/Editor/BlueprintGraph/Private/K2Node_GetInputAxisKeyValue.cpp
+++ b/Engine/Source/Editor/BlueprintGraph/Private/K2Node_GetInputAxisKeyValue.cpp
@@ -6,6 +6,8 @@
 #include "BlueprintNodeSpawner.h"
 #include "EditorCategoryUtils.h"
 #include "Engine/InputAxisKeyDelegateBinding.h"
+#include "BlueprintEditorUtils.h"
+#include "EdGraphSchema_K2.h"
 
 #define LOCTEXT_NAMESPACE "K2Node_GetInputAxisKeyValue"
 
@@ -49,6 +51,16 @@ FString UK2Node_GetInputAxisKeyValue::GetTooltip() const
 	return FText::Format(NSLOCTEXT("K2Node", "GetInputAxisKey_Tooltip", "Returns the current value of input axis key {AxisKey}.  If input is disabled for the actor the value will be 0."), Args).ToString();
 }
 
+bool UK2Node_GetInputAxisKeyValue::IsCompatibleWithGraph(UEdGraph const* Graph) const
+{
+	UBlueprint* Blueprint = FBlueprintEditorUtils::FindBlueprintForGraph(Graph);
+
+	UEdGraphSchema_K2 const* K2Schema = Cast<UEdGraphSchema_K2>(Graph->GetSchema());
+	bool const bIsConstructionScript = (K2Schema != nullptr) ? K2Schema->IsConstructionScript(Graph) : false;
+
+	return (Blueprint != nullptr) && Blueprint->SupportsInputEvents() && !bIsConstructionScript && Super::IsCompatibleWithGraph(Graph);
+}
+
 void UK2Node_GetInputAxisKeyValue::ValidateNodeDuringCompilation(class FCompilerResultsLog& MessageLog) const
 {
 	Super::ValidateNodeDuringCompilation(MessageLog);
diff --git a/Engine/Source/Editor/BlueprintGraph/Private/K2Node_GetInputAxisValue.cpp b/Engine/Source/Editor/BlueprintGraph/Private/K2Node_GetInputAxisValue.cpp
--- a/Engine/Source/Editor/BlueprintGraph/Private/K2Node_GetInputAxisValue.cpp
+++ b/Engine/Source/Editor/BlueprintGraph/Private/K2Node_GetInputAxisValue.cpp
@@ -6,6 +6,8 @@
 #include "BlueprintNodeSpawner.h"
 #include "EditorCategoryUtils.h"
 #include "Engine/InputAxisDelegateBinding.h"
+#include "BlueprintEditorUtils.h"
+#include "EdGraphSchema_K2.h"
 
 #define LOCTEXT_NAMESPACE "K2Node_GetInputAxisValue"
 
@@ -48,6 +50,16 @@ FString UK2Node_GetInputAxisValue::GetTooltip() const
 	return FString::Printf(*NSLOCTEXT("K2Node", "GetInputAxis_Tooltip", "Returns the current value of input axis %s.  If input is disabled for the actor the value will be 0.").ToString(), *InputAxisName.ToString());
 }
 
+bool UK2Node_GetInputAxisValue::IsCompatibleWithGraph(UEdGraph const* Graph) const
+{
+	UBlueprint* Blueprint = FBlueprintEditorUtils::FindBlueprintForGraph(Graph);
+
+	UEdGraphSchema_K2 const* K2Schema = Cast<UEdGraphSchema_K2>(Graph->GetSchema());
+	bool const bIsConstructionScript = (K2Schema != nullptr) ? K2Schema->IsConstructionScript(Graph) : false;
+
+	return (Blueprint != nullptr) && Blueprint->SupportsInputEvents() && !bIsConstructionScript && Super::IsCompatibleWithGraph(Graph);
+}
+
 void UK2Node_GetInputAxisValue::ValidateNodeDuringCompilation(class FCompilerResultsLog& MessageLog) const
 {
 	Super::ValidateNodeDuringCompilation(MessageLog);
diff --git a/Engine/Source/Editor/BlueprintGraph/Private/K2Node_InputAction.cpp b/Engine/Source/Editor/BlueprintGraph/Private/K2Node_InputAction.cpp
--- a/Engine/Source/Editor/BlueprintGraph/Private/K2Node_InputAction.cpp
+++ b/Engine/Source/Editor/BlueprintGraph/Private/K2Node_InputAction.cpp
@@ -6,6 +6,8 @@
 #include "KismetCompiler.h"
 #include "BlueprintNodeSpawner.h"
 #include "EditorCategoryUtils.h"
+#include "BlueprintEditorUtils.h"
+#include "EdGraphSchema_K2.h"
 
 #define LOCTEXT_NAMESPACE "K2Node_InputAction"
 
@@ -61,6 +63,16 @@ FString UK2Node_InputAction::GetTooltip() const
 	return FString::Printf(*NSLOCTEXT("K2Node", "InputAction_Tooltip", "Event for when the keys bound to input action %s are pressed or released.").ToString(), *InputActionName.ToString());
 }
 
+bool UK2Node_InputAction::IsCompatibleWithGraph(UEdGraph const* Graph) const
+{
+	UBlueprint* Blueprint = FBlueprintEditorUtils::FindBlueprintForGraph(Graph);
+
+	UEdGraphSchema_K2 const* K2Schema = Cast<UEdGraphSchema_K2>(Graph->GetSchema());
+	bool const bIsConstructionScript = (K2Schema != nullptr) ? K2Schema->IsConstructionScript(Graph) : false;
+
+	return (Blueprint != nullptr) && Blueprint->SupportsInputEvents() && !bIsConstructionScript && Super::IsCompatibleWithGraph(Graph);
+}
+
 UEdGraphPin* UK2Node_InputAction::GetPressedPin() const
 {
 	return FindPin(TEXT("Pressed"));
diff --git a/Engine/Source/Editor/BlueprintGraph/Private/K2Node_InputAxisKeyEvent.cpp b/Engine/Source/Editor/BlueprintGraph/Private/K2Node_InputAxisKeyEvent.cpp
--- a/Engine/Source/Editor/BlueprintGraph/Private/K2Node_InputAxisKeyEvent.cpp
+++ b/Engine/Source/Editor/BlueprintGraph/Private/K2Node_InputAxisKeyEvent.cpp
@@ -6,6 +6,8 @@
 #include "BlueprintNodeSpawner.h"
 #include "EditorCategoryUtils.h"
 #include "Engine/InputAxisKeyDelegateBinding.h"
+#include "BlueprintEditorUtils.h"
+#include "EdGraphSchema_K2.h"
 
 #define LOCTEXT_NAMESPACE "UK2Node_InputAxisKeyEvent"
 
@@ -96,11 +98,15 @@ bool UK2Node_InputAxisKeyEvent::IsCompatibleWithGraph(const UEdGraph* TargetGrap
 	
 	// Find the Blueprint that owns the target graph
 	UBlueprint* Blueprint = FBlueprintEditorUtils::FindBlueprintForGraph(TargetGraph);
-	if (Blueprint && Blueprint->SkeletonGeneratedClass)
+	if (Blueprint != nullptr)
 	{
-		bIsCompatible = Blueprint->ParentClass->IsChildOf(AActor::StaticClass());
+		bIsCompatible = FBlueprintEditorUtils::IsActorBased(Blueprint) && Blueprint->SupportsInputEvents();
 	}
 
+	UEdGraphSchema_K2 const* K2Schema = Cast<UEdGraphSchema_K2>(TargetGraph->GetSchema());
+	bool const bIsConstructionScript = (K2Schema != nullptr) ? K2Schema->IsConstructionScript(TargetGraph) : false;
+	bIsCompatible &= !bIsConstructionScript;
+
 	return bIsCompatible && Super::IsCompatibleWithGraph(TargetGraph);
 }
 
diff --git a/Engine/Source/Editor/BlueprintGraph/Private/K2Node_InputKey.cpp b/Engine/Source/Editor/BlueprintGraph/Private/K2Node_InputKey.cpp
--- a/Engine/Source/Editor/BlueprintGraph/Private/K2Node_InputKey.cpp
+++ b/Engine/Source/Editor/BlueprintGraph/Private/K2Node_InputKey.cpp
@@ -6,6 +6,8 @@
 #include "KismetCompiler.h"
 #include "BlueprintNodeSpawner.h"
 #include "EditorCategoryUtils.h"
+#include "BlueprintEditorUtils.h"
+#include "EdGraphSchema_K2.h"
 
 #define LOCTEXT_NAMESPACE "UK2Node_InputKey"
 
@@ -268,6 +270,16 @@ FName UK2Node_InputKey::GetPaletteIcon(FLinearColor& OutColor) const
 	}
 }
 
+bool UK2Node_InputKey::IsCompatibleWithGraph(UEdGraph const* Graph) const
+{
+	UBlueprint* Blueprint = FBlueprintEditorUtils::FindBlueprintForGraph(Graph);
+
+	UEdGraphSchema_K2 const* K2Schema = Cast<UEdGraphSchema_K2>(Graph->GetSchema());
+	bool const bIsConstructionScript = (K2Schema != nullptr) ? K2Schema->IsConstructionScript(Graph) : false;
+
+	return (Blueprint != nullptr) && Blueprint->SupportsInputEvents() && !bIsConstructionScript && Super::IsCompatibleWithGraph(Graph);
+}
+
 UEdGraphPin* UK2Node_InputKey::GetPressedPin() const
 {
 	return FindPin(TEXT("Pressed"));
diff --git a/Engine/Source/Editor/BlueprintGraph/Private/K2Node_InputTouch.cpp b/Engine/Source/Editor/BlueprintGraph/Private/K2Node_InputTouch.cpp
--- a/Engine/Source/Editor/BlueprintGraph/Private/K2Node_InputTouch.cpp
+++ b/Engine/Source/Editor/BlueprintGraph/Private/K2Node_InputTouch.cpp
@@ -80,6 +80,16 @@ FText UK2Node_InputTouch::GetMenuCategory() const
 	return FEditorCategoryUtils::GetCommonCategory(FCommonEditorCategory::Input);
 }
 
+bool UK2Node_InputTouch::IsCompatibleWithGraph(const UEdGraph* TargetGraph) const
+{
+	UBlueprint* Blueprint = FBlueprintEditorUtils::FindBlueprintForGraph(TargetGraph);
+
+	UEdGraphSchema_K2 const* K2Schema = Cast<UEdGraphSchema_K2>(TargetGraph->GetSchema());
+	bool const bIsConstructionScript = (K2Schema != nullptr) ? K2Schema->IsConstructionScript(TargetGraph) : false;
+
+	return (Blueprint != nullptr) && Blueprint->SupportsInputEvents() && !bIsConstructionScript && Super::IsCompatibleWithGraph(TargetGraph);
+}
+
 UEdGraphPin* UK2Node_InputTouch::GetPressedPin() const
 {
 	return FindPin(TEXT("Pressed"));
diff --git a/Engine/Source/Editor/BlueprintGraph/Private/K2Node_MCDelegate.cpp b/Engine/Source/Editor/BlueprintGraph/Private/K2Node_MCDelegate.cpp
--- a/Engine/Source/Editor/BlueprintGraph/Private/K2Node_MCDelegate.cpp
+++ b/Engine/Source/Editor/BlueprintGraph/Private/K2Node_MCDelegate.cpp
@@ -29,6 +29,15 @@ void UK2Node_BaseMCDelegate::ValidateNodeDuringCompilation(class FCompilerResult
 	}
 }
 
+bool UK2Node_BaseMCDelegate::IsCompatibleWithGraph(const UEdGraph* TargetGraph) const
+{
+	UEdGraphSchema const* Schema = TargetGraph->GetSchema();
+	EGraphType GraphType = Schema->GetGraphType(TargetGraph);
+
+	bool const bIsCompatible = (GraphType == GT_Ubergraph) || (GraphType == GT_Function);
+	return bIsCompatible&& Super::IsCompatibleWithGraph(TargetGraph);
+}
+
 UK2Node::ERedirectType UK2Node_BaseMCDelegate::DoPinsMatchForReconstruction(const UEdGraphPin* NewPin, int32 NewPinIndex, const UEdGraphPin* OldPin, int32 OldPinIndex) const
 {
 	ERedirectType OrginalResult = Super::DoPinsMatchForReconstruction(NewPin, NewPinIndex, OldPin, OldPinIndex);
diff --git a/Engine/Source/Editor/UnrealEd/Private/Commandlets/DumpBlueprintsInfoCommandlet.cpp b/Engine/Source/Editor/UnrealEd/Private/Commandlets/DumpBlueprintsInfoCommandlet.cpp
--- a/Engine/Source/Editor/UnrealEd/Private/Commandlets/DumpBlueprintsInfoCommandlet.cpp
+++ b/Engine/Source/Editor/UnrealEd/Private/Commandlets/DumpBlueprintsInfoCommandlet.cpp
@@ -24,6 +24,7 @@
 #include "ModuleManager.h"
 #include "AssetToolsModule.h"
 #include "BlueprintNodeSpawner.h"
+#include "AnimationGraphSchema.h"
 
 DEFINE_LOG_CATEGORY_STATIC(LogBlueprintInfoDump, Log, All);
 
@@ -689,8 +690,10 @@ static UBlueprint* DumpBlueprintInfoUtils::MakeTempBlueprint(UClass* ParentClass
 			GeneratedClass = UAnimBlueprintGeneratedClass::StaticClass();
 
 			UAnimBlueprintFactory* BlueprintFactory = ConstructObject<UAnimBlueprintFactory>(UAnimBlueprintFactory::StaticClass());
-			BlueprintFactory->ParentClass   = ParentClass;
-			BlueprintFactory->BlueprintType = BlueprintType;
+			BlueprintFactory->ParentClass    = ParentClass;
+			BlueprintFactory->BlueprintType  = BlueprintType;
+			BlueprintFactory->TargetSkeleton = (USkeleton*)StaticLoadObject(USkeleton::StaticClass(), /*Outer =*/nullptr, TEXT("/Engine/NotForLicensees/Automation/QAAutomationtest_Assets/TEST_SkeletalMesh_Skeleton.TEST_SkeletalMesh_Skeleton"));
+
 			AssetFactory = BlueprintFactory;
 		}
 		else if (bIsLevelBlueprint)
@@ -1529,11 +1532,14 @@ static bool DumpBlueprintInfoUtils::DumpPinContextActions(uint32 Indent, UEdGrap
 		}
 		else if (UScriptStruct* StructType = FindObject<UScriptStruct>(ANY_PACKAGE, *CommandOptions.PinType))
 		{
+			PinType.PinCategory = K2Schema->PC_Struct;
+			PinType.PinSubCategoryObject = StructType;
 			bIsValidPinType = UEdGraphSchema_K2::IsAllowableBlueprintVariableType(StructType);
-			if (bIsValidPinType)
+
+			UEdGraphSchema const* Schema = Graph->GetSchema();
+			if (UAnimationGraphSchema const* AnimSchema = Cast<UAnimationGraphSchema>(Schema))
 			{
-				PinType.PinCategory = K2Schema->PC_Struct;
-				PinType.PinSubCategoryObject = StructType;
+				bIsValidPinType |= AnimSchema->IsPosePin(PinType);
 			}
 		}
 		else if (!CommandOptions.PinType.Compare("self", ESearchCase::IgnoreCase))
