Caching tooltips that use FText::Format() (optimizing for the new BP menu system).

diff --git a/Engine/Source/Editor/AnimGraph/Classes/AnimGraphNode_BlendListByEnum.h b/Engine/Source/Editor/AnimGraph/Classes/AnimGraphNode_BlendListByEnum.h
--- a/Engine/Source/Editor/AnimGraph/Classes/AnimGraphNode_BlendListByEnum.h
+++ b/Engine/Source/Editor/AnimGraph/Classes/AnimGraphNode_BlendListByEnum.h
@@ -3,6 +3,7 @@
 #pragma once
 #include "AnimGraphNode_BlendListBase.h"
 #include "Animation/AnimNode_BlendListByEnum.h"
+#include "EdGraph/EdGraphNodeUtils.h" // for FNodeTextCache
 #include "AnimGraphNode_BlendListByEnum.generated.h"
 
 UCLASS(MinimalAPI)
@@ -53,4 +54,8 @@ protected:
 
 	// Gets information about the specified pin.  If both bIsPosePin and bIsTimePin are false, the index is meaningless
 	static void GetPinInformation(const FString& InPinName, int32& Out_PinIndex, bool& Out_bIsPosePin, bool& Out_bIsTimePin);
+
+private:
+	/** Constructing FText strings can be costly, so we cache the node's title */
+	FNodeTextCache CachedNodeTitle;
 };
diff --git a/Engine/Source/Editor/AnimGraph/Classes/AnimGraphNode_BlendSpaceEvaluator.h b/Engine/Source/Editor/AnimGraph/Classes/AnimGraphNode_BlendSpaceEvaluator.h
--- a/Engine/Source/Editor/AnimGraph/Classes/AnimGraphNode_BlendSpaceEvaluator.h
+++ b/Engine/Source/Editor/AnimGraph/Classes/AnimGraphNode_BlendSpaceEvaluator.h
@@ -3,6 +3,7 @@
 #pragma once
 #include "AnimGraphNode_BlendSpaceBase.h"
 #include "Animation/AnimNode_BlendSpaceEvaluator.h"
+#include "EdGraph/EdGraphNodeUtils.h" // for FNodeTitleTextTable
 #include "AnimGraphNode_BlendSpaceEvaluator.generated.h"
 
 UCLASS(MinimalAPI)
@@ -36,4 +37,8 @@ class UAnimGraphNode_BlendSpaceEvaluator : public UAnimGraphNode_BlendSpaceBase
 	// UK2Node interface
 	virtual void GetContextMenuActions(const FGraphNodeContextMenuBuilder& Context) const override;
 	// End of UK2Node interface
+
+private:
+	/** Constructing FText strings can be costly, so we cache the node's title */
+	FNodeTitleTextTable CachedNodeTitles;
 };
diff --git a/Engine/Source/Editor/AnimGraph/Classes/AnimGraphNode_BlendSpacePlayer.h b/Engine/Source/Editor/AnimGraph/Classes/AnimGraphNode_BlendSpacePlayer.h
--- a/Engine/Source/Editor/AnimGraph/Classes/AnimGraphNode_BlendSpacePlayer.h
+++ b/Engine/Source/Editor/AnimGraph/Classes/AnimGraphNode_BlendSpacePlayer.h
@@ -3,6 +3,7 @@
 #pragma once
 #include "AnimGraphNode_BlendSpaceBase.h"
 #include "Animation/AnimNode_BlendSpacePlayer.h"
+#include "EdGraph/EdGraphNodeUtils.h" // for FNodeTitleTextTable
 #include "AnimGraphNode_BlendSpacePlayer.generated.h"
 
 UCLASS(MinimalAPI)
@@ -38,4 +39,8 @@ class UAnimGraphNode_BlendSpacePlayer : public UAnimGraphNode_BlendSpaceBase
 	// UK2Node interface
 	virtual void GetContextMenuActions(const FGraphNodeContextMenuBuilder& Context) const override;
 	// End of UK2Node interface
+
+private:
+	/** Constructing FText strings can be costly, so we cache the node's title */
+	FNodeTitleTextTable CachedNodeTitles;
 };
diff --git a/Engine/Source/Editor/AnimGraph/Classes/AnimGraphNode_RotationOffsetBlendSpace.h b/Engine/Source/Editor/AnimGraph/Classes/AnimGraphNode_RotationOffsetBlendSpace.h
--- a/Engine/Source/Editor/AnimGraph/Classes/AnimGraphNode_RotationOffsetBlendSpace.h
+++ b/Engine/Source/Editor/AnimGraph/Classes/AnimGraphNode_RotationOffsetBlendSpace.h
@@ -3,6 +3,7 @@
 #pragma once
 #include "AnimGraphNode_BlendSpaceBase.h"
 #include "Animation/AnimNode_RotationOffsetBlendSpace.h"
+#include "EdGraph/EdGraphNodeUtils.h" // for FNodeTitleTextTable
 #include "AnimGraphNode_RotationOffsetBlendSpace.generated.h"
 
 UCLASS(MinimalAPI)
@@ -28,4 +29,8 @@ class UAnimGraphNode_RotationOffsetBlendSpace: public UAnimGraphNode_BlendSpaceB
 	virtual void GetContextMenuActions(const FGraphNodeContextMenuBuilder& Context) const override;
 	virtual void GetMenuActions(FBlueprintActionDatabaseRegistrar& ActionRegistrar) const override;
 	// End of UK2Node interface
+
+private:
+	/** Constructing FText strings can be costly, so we cache the node's title */
+	FNodeTitleTextTable CachedNodeTitles;
 };
diff --git a/Engine/Source/Editor/AnimGraph/Classes/AnimGraphNode_SequenceEvaluator.h b/Engine/Source/Editor/AnimGraph/Classes/AnimGraphNode_SequenceEvaluator.h
--- a/Engine/Source/Editor/AnimGraph/Classes/AnimGraphNode_SequenceEvaluator.h
+++ b/Engine/Source/Editor/AnimGraph/Classes/AnimGraphNode_SequenceEvaluator.h
@@ -3,6 +3,7 @@
 #pragma once
 #include "AnimGraphNode_Base.h"
 #include "Animation/AnimNode_SequenceEvaluator.h"
+#include "EdGraph/EdGraphNodeUtils.h" // for FNodeTextCache
 #include "AnimGraphNode_SequenceEvaluator.generated.h"
 
 UCLASS(MinimalAPI)
@@ -35,4 +36,8 @@ class UAnimGraphNode_SequenceEvaluator : public UAnimGraphNode_Base
 	virtual void GetContextMenuActions(const FGraphNodeContextMenuBuilder& Context) const override;
 	virtual void GetMenuActions(FBlueprintActionDatabaseRegistrar& ActionRegistrar) const override;
 	// End of UK2Node interface
+
+private:
+	/** Constructing FText strings can be costly, so we cache the node's title */
+	FNodeTextCache CachedNodeTitle;
 };
diff --git a/Engine/Source/Editor/AnimGraph/Private/AnimGraphNode_BlendListByEnum.cpp b/Engine/Source/Editor/AnimGraph/Private/AnimGraphNode_BlendListByEnum.cpp
--- a/Engine/Source/Editor/AnimGraph/Private/AnimGraphNode_BlendListByEnum.cpp
+++ b/Engine/Source/Editor/AnimGraph/Private/AnimGraphNode_BlendListByEnum.cpp
@@ -26,30 +26,23 @@ FString UAnimGraphNode_BlendListByEnum::GetNodeCategory() const
 
 FText UAnimGraphNode_BlendListByEnum::GetTooltipText() const
 {
-	if (BoundEnum != NULL)
-	{
-		FFormatNamedArguments Args;
-		Args.Add(TEXT("EnumName"), FText::FromString(BoundEnum->GetName()));
-		return FText::Format(LOCTEXT("AnimGraphNode_BlendListByEnum_Tooltip", "Blend Poses based on enum ({EnumName}). Use context menu to add new entries."), Args);
-	}
-	else
-	{
-		return LOCTEXT("AnimGraphNode_BlendListByEnum_TooltipError", "ERROR: Blend Poses (by missing enum)");
-	}
+	// FText::Format() is slow, so we utilize the cached list title
+	return GetNodeTitle(ENodeTitleType::ListView);
 }
 
 FText UAnimGraphNode_BlendListByEnum::GetNodeTitle(ENodeTitleType::Type TitleType) const
 {
-	if (BoundEnum != NULL)
+	if (BoundEnum == nullptr)
 	{
-		FFormatNamedArguments Args;
-		Args.Add(TEXT("EnumName"), FText::FromString(BoundEnum->GetName()));
-		return FText::Format(LOCTEXT("AnimGraphNode_BlendListByEnum_Title", "Blend Poses ({EnumName})"), Args);
+		return LOCTEXT("AnimGraphNode_BlendListByEnum_TitleError", "ERROR: Blend Poses (by missing enum)");
 	}
-	else
+	else if (CachedNodeTitle.IsOutOfDate())
 	{
-		return LOCTEXT("AnimGraphNode_BlendListByEnum_TitleError", "ERROR: Blend Poses (by missing enum)");
+		FFormatNamedArguments Args;
+		Args.Add(TEXT("EnumName"), FText::FromString(BoundEnum->GetName()));
+		CachedNodeTitle = FText::Format(LOCTEXT("AnimGraphNode_BlendListByEnum_Title", "Blend Poses ({EnumName})"), Args);
 	}
+	return CachedNodeTitle;
 }
 
 void UAnimGraphNode_BlendListByEnum::PostPlacedNewNode()
diff --git a/Engine/Source/Editor/AnimGraph/Private/AnimGraphNode_BlendSpaceEvaluator.cpp b/Engine/Source/Editor/AnimGraph/Private/AnimGraphNode_BlendSpaceEvaluator.cpp
--- a/Engine/Source/Editor/AnimGraph/Private/AnimGraphNode_BlendSpaceEvaluator.cpp
+++ b/Engine/Source/Editor/AnimGraph/Private/AnimGraphNode_BlendSpaceEvaluator.cpp
@@ -17,36 +17,53 @@ UAnimGraphNode_BlendSpaceEvaluator::UAnimGraphNode_BlendSpaceEvaluator(const FPo
 
 FText UAnimGraphNode_BlendSpaceEvaluator::GetTooltipText() const
 {
-	FText const BlendspaceName = (Node.BlendSpace != NULL) ? FText::FromString(Node.BlendSpace->GetPathName()) : LOCTEXT("None", "(None)");
-	return FText::Format(LOCTEXT("BlendspaceEvaluator_Tooltip", "Blendspace Evaluator '{0}'"), BlendspaceName);
+	// FText::Format() is slow, so we utilize the cached list title
+	return GetNodeTitle(ENodeTitleType::ListView);
 }
 
 FText UAnimGraphNode_BlendSpaceEvaluator::GetNodeTitle(ENodeTitleType::Type TitleType) const
 {
-	const FText BlendSpaceName((Node.BlendSpace != NULL) ? FText::FromString(Node.BlendSpace->GetName()) : LOCTEXT("None", "(None)"));
-	
-	if (TitleType == ENodeTitleType::ListView)
+	if (Node.BlendSpace == nullptr)
 	{
-		FFormatNamedArguments Args;
-		Args.Add(TEXT("BlendSpaceName"), BlendSpaceName);
-		return FText::Format(LOCTEXT("BlendSpaceEvaluatorListTitle", "Blendspace Evaluator '{BlendSpaceName}'"), Args);
+		if (TitleType == ENodeTitleType::ListView)
+		{
+			return LOCTEXT("BlendSpaceEvaluator_NONE_ListTitle", "Blendspace Evaluator '(None)'");
+		}
+		else
+		{
+			return LOCTEXT("BlendSpaceEvaluator_NONE_Title", "(None)\nBlendspace Evaluator");
+		}
 	}
-	else
+	else if (!CachedNodeTitles.IsTitleCached(TitleType))
 	{
-		FFormatNamedArguments TitleArgs;
-		TitleArgs.Add(TEXT("BlendSpaceName"), BlendSpaceName);
-		FText Title = FText::Format(LOCTEXT("BlendSpaceEvaluatorFullTitle", "{BlendSpaceName}\nBlendspace Evaluator"), TitleArgs);
+		const FText BlendSpaceName = FText::FromString(Node.BlendSpace->GetName());
 
-		if ((TitleType == ENodeTitleType::FullTitle) && (SyncGroup.GroupName != NAME_None))
+		if (TitleType == ENodeTitleType::ListView)
 		{
 			FFormatNamedArguments Args;
-			Args.Add(TEXT("Title"), Title);
-			Args.Add(TEXT("SyncGroupName"), FText::FromName(SyncGroup.GroupName));
-			Title = FText::Format(LOCTEXT("BlendSpaceNodeGroupSubtitle", "{Title}\nSync group {SyncGroupName}"), Args);
-		}
+			Args.Add(TEXT("BlendSpaceName"), BlendSpaceName);
 
-		return Title;
+			// FText::Format() is slow, so we cache this to save on performance
+			CachedNodeTitles.SetCachedTitle(TitleType, FText::Format(LOCTEXT("BlendSpaceEvaluatorListTitle", "Blendspace Evaluator '{BlendSpaceName}'"), Args));
+		}
+		else
+		{
+			FFormatNamedArguments TitleArgs;
+			TitleArgs.Add(TEXT("BlendSpaceName"), BlendSpaceName);
+			FText Title = FText::Format(LOCTEXT("BlendSpaceEvaluatorFullTitle", "{BlendSpaceName}\nBlendspace Evaluator"), TitleArgs);
+
+			if ((TitleType == ENodeTitleType::FullTitle) && (SyncGroup.GroupName != NAME_None))
+			{
+				FFormatNamedArguments Args;
+				Args.Add(TEXT("Title"), Title);
+				Args.Add(TEXT("SyncGroupName"), FText::FromName(SyncGroup.GroupName));
+				Title = FText::Format(LOCTEXT("BlendSpaceNodeGroupSubtitle", "{Title}\nSync group {SyncGroupName}"), Args);
+			}
+			// FText::Format() is slow, so we cache this to save on performance
+			CachedNodeTitles.SetCachedTitle(TitleType, Title);
+		}
 	}
+	return CachedNodeTitles[TitleType];
 }
 
 void UAnimGraphNode_BlendSpaceEvaluator::GetMenuEntries(FGraphContextMenuBuilder& ContextMenuBuilder) const
diff --git a/Engine/Source/Editor/AnimGraph/Private/AnimGraphNode_BlendSpacePlayer.cpp b/Engine/Source/Editor/AnimGraph/Private/AnimGraphNode_BlendSpacePlayer.cpp
--- a/Engine/Source/Editor/AnimGraph/Private/AnimGraphNode_BlendSpacePlayer.cpp
+++ b/Engine/Source/Editor/AnimGraph/Private/AnimGraphNode_BlendSpacePlayer.cpp
@@ -19,36 +19,52 @@ UAnimGraphNode_BlendSpacePlayer::UAnimGraphNode_BlendSpacePlayer(const FPostCons
 
 FText UAnimGraphNode_BlendSpacePlayer::GetTooltipText() const
 {
-	const FText BlendSpaceName((Node.BlendSpace != NULL) ? FText::FromString(Node.BlendSpace->GetName()) : LOCTEXT("None", "(None)"));
-	return FText::Format(LOCTEXT("BlendspacePlayer", "Blendspace Player '{0}'"), BlendSpaceName);
+	// FText::Format() is slow, so we utilize the cached list title
+	return GetNodeTitle(ENodeTitleType::ListView);
 }
 
 FText UAnimGraphNode_BlendSpacePlayer::GetNodeTitle(ENodeTitleType::Type TitleType) const
-{
-	const FText BlendSpaceName((Node.BlendSpace != NULL) ? FText::FromString(Node.BlendSpace->GetName()) : LOCTEXT("None", "(None)"));
-	
-	if (TitleType == ENodeTitleType::ListView)
+{	
+	if (Node.BlendSpace == nullptr)
 	{
-		FFormatNamedArguments Args;
-		Args.Add(TEXT("BlendSpaceName"), BlendSpaceName);
-		return FText::Format(LOCTEXT("BlendspacePlayer", "Blendspace Player '{BlendSpaceName}'"), Args);
+		if (TitleType == ENodeTitleType::ListView)
+		{
+			return LOCTEXT("BlendspacePlayer_NONE_ListTitle", "Blendspace Player '(None)'");
+		}
+		else
+		{
+			return LOCTEXT("BlendspacePlayer_NONE_Title", "(None)\nBlendspace Player");
+		}
 	}
-	else
+	else if (!CachedNodeTitles.IsTitleCached(TitleType))
 	{
-		FFormatNamedArguments TitleArgs;
-		TitleArgs.Add(TEXT("BlendSpaceName"), BlendSpaceName);
-		FText Title = FText::Format(LOCTEXT("BlendSpacePlayerFullTitle", "{BlendSpaceName}\nBlendspace Player"), TitleArgs);
+		const FText BlendSpaceName = FText::FromString(Node.BlendSpace->GetName());
 
-		if ((TitleType == ENodeTitleType::FullTitle) && (SyncGroup.GroupName != NAME_None))
+		if (TitleType == ENodeTitleType::ListView)
 		{
 			FFormatNamedArguments Args;
-			Args.Add(TEXT("Title"), Title);
-			Args.Add(TEXT("SyncGroupName"), FText::FromName(SyncGroup.GroupName));
-			Title = FText::Format(LOCTEXT("BlendSpaceNodeGroupSubtitle", "{Title}\nSync group {SyncGroupName}"), Args);
+			Args.Add(TEXT("BlendSpaceName"), BlendSpaceName);
+			// FText::Format() is slow, so we cache this to save on performance
+			CachedNodeTitles.SetCachedTitle(TitleType, FText::Format(LOCTEXT("BlendspacePlayer", "Blendspace Player '{BlendSpaceName}'"), Args));
+		}
+		else
+		{
+			FFormatNamedArguments TitleArgs;
+			TitleArgs.Add(TEXT("BlendSpaceName"), BlendSpaceName);
+			FText Title = FText::Format(LOCTEXT("BlendSpacePlayerFullTitle", "{BlendSpaceName}\nBlendspace Player"), TitleArgs);
+
+			if ((TitleType == ENodeTitleType::FullTitle) && (SyncGroup.GroupName != NAME_None))
+			{
+				FFormatNamedArguments Args;
+				Args.Add(TEXT("Title"), Title);
+				Args.Add(TEXT("SyncGroupName"), FText::FromName(SyncGroup.GroupName));
+				Title = FText::Format(LOCTEXT("BlendSpaceNodeGroupSubtitle", "{Title}\nSync group {SyncGroupName}"), Args);
+			}
+			// FText::Format() is slow, so we cache this to save on performance
+			CachedNodeTitles.SetCachedTitle(TitleType, Title);
 		}
-
-		return Title;
 	}
+	return CachedNodeTitles[TitleType];
 }
 
 void UAnimGraphNode_BlendSpacePlayer::GetMenuEntries(FGraphContextMenuBuilder& ContextMenuBuilder) const
diff --git a/Engine/Source/Editor/AnimGraph/Private/AnimGraphNode_RotationOffsetBlendSpace.cpp b/Engine/Source/Editor/AnimGraph/Private/AnimGraphNode_RotationOffsetBlendSpace.cpp
--- a/Engine/Source/Editor/AnimGraph/Private/AnimGraphNode_RotationOffsetBlendSpace.cpp
+++ b/Engine/Source/Editor/AnimGraph/Private/AnimGraphNode_RotationOffsetBlendSpace.cpp
@@ -19,25 +19,41 @@ UAnimGraphNode_RotationOffsetBlendSpace::UAnimGraphNode_RotationOffsetBlendSpace
 
 FText UAnimGraphNode_RotationOffsetBlendSpace::GetTooltipText() const
 {
-	const FText BlendSpaceName((Node.BlendSpace != NULL) ? FText::FromString(Node.BlendSpace->GetName()) : LOCTEXT("None", "(None)"));
-	return FText::Format(LOCTEXT("AimOffsetTooltip", "AimOffset '{0}'"), BlendSpaceName);
+	// FText::Format() is slow, so we utilize the cached list title
+	return GetNodeTitle(ENodeTitleType::ListView);
 }
 
 FText UAnimGraphNode_RotationOffsetBlendSpace::GetNodeTitle(ENodeTitleType::Type TitleType) const
 {
-	const FText BlendSpaceName((Node.BlendSpace != NULL) ? FText::FromString(Node.BlendSpace->GetName()) : LOCTEXT("None", "(None)"));
-
-	FFormatNamedArguments Args;
-	Args.Add(TEXT("BlendSpaceName"), BlendSpaceName);
-
-	if (TitleType == ENodeTitleType::ListView)
+	if (Node.BlendSpace == nullptr)
 	{
-		return FText::Format(LOCTEXT("AimOffsetListTitle", "AimOffset '{BlendSpaceName}'"), Args);
+		if (TitleType == ENodeTitleType::ListView)
+		{
+			return LOCTEXT("RotationOffsetBlend_NONE_ListTitle", "AimOffset '(None)'");
+		}
+		else
+		{
+			return LOCTEXT("RotationOffsetBlend_NONE_Title", "(None)\nAimOffset");
+		}
 	}
-	else
+	else if (!CachedNodeTitles.IsTitleCached(TitleType))
 	{
-		return FText::Format(LOCTEXT("AimOffsetFullTitle", "{BlendSpaceName}\nAimOffset"), Args);
+		const FText BlendSpaceName = FText::FromString(Node.BlendSpace->GetName());
+
+		FFormatNamedArguments Args;
+		Args.Add(TEXT("BlendSpaceName"), BlendSpaceName);
+
+		// FText::Format() is slow, so we cache this to save on performance
+		if (TitleType == ENodeTitleType::ListView)
+		{
+			CachedNodeTitles.SetCachedTitle(TitleType, FText::Format(LOCTEXT("AimOffsetListTitle", "AimOffset '{BlendSpaceName}'"), Args));
+		}
+		else
+		{
+			CachedNodeTitles.SetCachedTitle(TitleType, FText::Format(LOCTEXT("AimOffsetFullTitle", "{BlendSpaceName}\nAimOffset"), Args));
+		}
 	}
+	return CachedNodeTitles[TitleType];
 }
 
 void UAnimGraphNode_RotationOffsetBlendSpace::GetMenuEntries(FGraphContextMenuBuilder& ContextMenuBuilder) const
diff --git a/Engine/Source/Editor/AnimGraph/Private/AnimGraphNode_SequenceEvaluator.cpp b/Engine/Source/Editor/AnimGraph/Private/AnimGraphNode_SequenceEvaluator.cpp
--- a/Engine/Source/Editor/AnimGraph/Private/AnimGraphNode_SequenceEvaluator.cpp
+++ b/Engine/Source/Editor/AnimGraph/Private/AnimGraphNode_SequenceEvaluator.cpp
@@ -38,31 +38,35 @@ void UAnimGraphNode_SequenceEvaluator::ReplaceReferredAnimations(const TMap<UAni
 
 FText UAnimGraphNode_SequenceEvaluator::GetTooltipText() const
 {
-	if ((Node.Sequence != NULL) && Node.Sequence->IsValidAdditive())
-	{
-		return FText::Format(LOCTEXT("SequenceEvaluator_AdditiveTooltip", "Evaluate {0} (additive)"), FText::FromString(Node.Sequence->GetPathName()));
-	}
-	else
-	{
-		return FText::Format(LOCTEXT("SequenceEvaluator_Tooltip", "Evaluate {0}"), FText::FromString(Node.Sequence->GetPathName()));
-	}
+	// FText::Format() is slow, so we utilize the cached list title
+	return GetNodeTitle(ENodeTitleType::ListView);
 }
 
 FText UAnimGraphNode_SequenceEvaluator::GetNodeTitle(ENodeTitleType::Type TitleType) const
 {
-	const FText SequenceName((Node.Sequence != NULL) ? FText::FromString(Node.Sequence->GetName()) : LOCTEXT("None", "(None)"));
-
-	FFormatNamedArguments Args;
-	Args.Add(TEXT("SequenceName"), SequenceName);
-
-	if (Node.Sequence && Node.Sequence->IsValidAdditive())
+	if (Node.Sequence == nullptr)
 	{
-		return FText::Format(LOCTEXT("EvaluateSequence_Additive", "Evaluate {SequenceName} (additive)"), Args);
+		return LOCTEXT("EvaluateSequence_TitleNONE", "Evaluate (None)");
 	}
-	else
+	else if (CachedNodeTitle.IsOutOfDate())
 	{
-		return FText::Format(LOCTEXT("EvaluateSequence", "Evaluate {SequenceName}"), Args);
+		const FText SequenceName = FText::FromString(Node.Sequence->GetName());
+
+		FFormatNamedArguments Args;
+		Args.Add(TEXT("SequenceName"), SequenceName);
+
+		// FText::Format() is slow, so we cache this to save on performance
+		if (Node.Sequence->IsValidAdditive())
+		{
+			CachedNodeTitle = FText::Format(LOCTEXT("EvaluateSequence_Additive", "Evaluate {SequenceName} (additive)"), Args);
+		}
+		else
+		{
+			CachedNodeTitle = FText::Format(LOCTEXT("EvaluateSequence", "Evaluate {SequenceName}"), Args);
+		}
 	}
+
+	return CachedNodeTitle;
 }
 
 void UAnimGraphNode_SequenceEvaluator::GetMenuEntries(FGraphContextMenuBuilder& ContextMenuBuilder) const
diff --git a/Engine/Source/Editor/BehaviorTreeEditor/Private/BehaviorTreeGraphNode.cpp b/Engine/Source/Editor/BehaviorTreeEditor/Private/BehaviorTreeGraphNode.cpp
--- a/Engine/Source/Editor/BehaviorTreeEditor/Private/BehaviorTreeGraphNode.cpp
+++ b/Engine/Source/Editor/BehaviorTreeEditor/Private/BehaviorTreeGraphNode.cpp
@@ -173,6 +173,9 @@ FText UBehaviorTreeGraphNode::GetTooltipText() const
 	if (bInjectedNode)
 	{
 		FText const InjectedDesc = !TooltipDesc.IsEmpty() ? TooltipDesc : FText::FromString(GetDescription());
+		// @TODO: FText::Format() is slow... consider caching this tooltip like 
+		//        we do for a lot of the BP nodes now (unfamiliar with this 
+		//        node's purpose, so hesitant to muck with this at this time).
 		TooltipDesc = FText::Format(LOCTEXT("InjectedTooltip", "Injected: {0}"), InjectedDesc);
 	}
 
diff --git a/Engine/Source/Editor/BlueprintGraph/Classes/K2Node_BreakStruct.h b/Engine/Source/Editor/BlueprintGraph/Classes/K2Node_BreakStruct.h
--- a/Engine/Source/Editor/BlueprintGraph/Classes/K2Node_BreakStruct.h
+++ b/Engine/Source/Editor/BlueprintGraph/Classes/K2Node_BreakStruct.h
@@ -32,7 +32,8 @@ class UK2Node_BreakStruct : public UK2Node_StructMemberGet
 	// End K2Node interface
 
 private:
-	/** Constructing FText strings can be costly, so we cache the node's title */
+	/** Constructing FText strings can be costly, so we cache the node's title/tooltip */
+	FNodeTextCache CachedTooltip;
 	FNodeTextCache CachedNodeTitle;
 };
 
diff --git a/Engine/Source/Editor/BlueprintGraph/Classes/K2Node_CallFunction.h b/Engine/Source/Editor/BlueprintGraph/Classes/K2Node_CallFunction.h
--- a/Engine/Source/Editor/BlueprintGraph/Classes/K2Node_CallFunction.h
+++ b/Engine/Source/Editor/BlueprintGraph/Classes/K2Node_CallFunction.h
@@ -47,6 +47,9 @@ private:
 	UPROPERTY()
 	TSubclassOf<class UObject> CallFunctionClass_DEPRECATED;
 
+	/** Constructing FText strings can be costly, so we cache the node's tooltip */
+	FNodeTextCache CachedTooltip;
+
 public:
 
 	// UObject interface
diff --git a/Engine/Source/Editor/BlueprintGraph/Classes/K2Node_CastByteToEnum.h b/Engine/Source/Editor/BlueprintGraph/Classes/K2Node_CastByteToEnum.h
--- a/Engine/Source/Editor/BlueprintGraph/Classes/K2Node_CastByteToEnum.h
+++ b/Engine/Source/Editor/BlueprintGraph/Classes/K2Node_CastByteToEnum.h
@@ -2,6 +2,7 @@
 
 #pragma once
 #include "K2Node.h"
+#include "EdGraph/EdGraphNodeUtils.h" // for FNodeTextCache
 #include "K2Node_CastByteToEnum.generated.h"
 
 UCLASS(MinimalAPI)
@@ -36,5 +37,9 @@ class UK2Node_CastByteToEnum : public UK2Node
 	// End UK2Node interface
 
 	virtual FName GetFunctionName() const;
+
+private:
+	/** Constructing FText strings can be costly, so we cache the node's tooltip */
+	FNodeTextCache CachedTooltip;
 };
 
diff --git a/Engine/Source/Editor/BlueprintGraph/Classes/K2Node_DelegateSet.h b/Engine/Source/Editor/BlueprintGraph/Classes/K2Node_DelegateSet.h
--- a/Engine/Source/Editor/BlueprintGraph/Classes/K2Node_DelegateSet.h
+++ b/Engine/Source/Editor/BlueprintGraph/Classes/K2Node_DelegateSet.h
@@ -51,7 +51,8 @@ class UK2Node_DelegateSet : public UK2Node
 	BLUEPRINTGRAPH_API UFunction* GetDelegateSignature() const;
 
 private:
-	/** Constructing FText strings can be costly, so we cache the node's title */
+	/** Constructing FText strings can be costly, so we cache the node's title/tooltip */
+	FNodeTextCache CachedTooltip;
 	FNodeTextCache CachedNodeTitle;
 };
 
diff --git a/Engine/Source/Editor/BlueprintGraph/Classes/K2Node_Event.h b/Engine/Source/Editor/BlueprintGraph/Classes/K2Node_Event.h
--- a/Engine/Source/Editor/BlueprintGraph/Classes/K2Node_Event.h
+++ b/Engine/Source/Editor/BlueprintGraph/Classes/K2Node_Event.h
@@ -3,6 +3,7 @@
 
 #pragma once
 #include "K2Node_EditablePinBase.h"
+#include "EdGraph/EdGraphNodeUtils.h" // for FNodeTextCache
 #include "K2Node_Event.generated.h"
 
 UCLASS(MinimalAPI)
@@ -79,5 +80,9 @@ class UK2Node_Event : public UK2Node_EditablePinBase
 	 *		Calling - whether this function is being called ("sending") or showing implementation ("receiving"). Determined whether we output "Replicated To Server" or "Replicated From Client".
 	 */
 	static FString GetLocalizedNetString(uint32 NetFlags, bool Calling);
+
+private:
+	/** Constructing FText strings can be costly, so we cache the node's tooltip */
+	FNodeTextCache CachedTooltip;
 };
 
diff --git a/Engine/Source/Editor/BlueprintGraph/Classes/K2Node_GetInputAxisKeyValue.h b/Engine/Source/Editor/BlueprintGraph/Classes/K2Node_GetInputAxisKeyValue.h
--- a/Engine/Source/Editor/BlueprintGraph/Classes/K2Node_GetInputAxisKeyValue.h
+++ b/Engine/Source/Editor/BlueprintGraph/Classes/K2Node_GetInputAxisKeyValue.h
@@ -41,6 +41,7 @@ class UK2Node_GetInputAxisKeyValue : public UK2Node_CallFunction
 	void Initialize(const FKey AxisKey);
 
 private:
-	/** Constructing FText strings can be costly, so we cache the node's title */
+	/** Constructing FText strings can be costly, so we cache the node's title/tooltip */
+	FNodeTextCache CachedTooltip;
 	FNodeTextCache CachedNodeTitle;
 };
\ No newline at end of file
diff --git a/Engine/Source/Editor/BlueprintGraph/Classes/K2Node_GetInputAxisValue.h b/Engine/Source/Editor/BlueprintGraph/Classes/K2Node_GetInputAxisValue.h
--- a/Engine/Source/Editor/BlueprintGraph/Classes/K2Node_GetInputAxisValue.h
+++ b/Engine/Source/Editor/BlueprintGraph/Classes/K2Node_GetInputAxisValue.h
@@ -40,6 +40,7 @@ class UK2Node_GetInputAxisValue : public UK2Node_CallFunction
 	void Initialize(const FName AxisName);
 
 private:
-	/** Constructing FText strings can be costly, so we cache the node's title */
+	/** Constructing FText strings can be costly, so we cache the node's title/tooltip */
+	FNodeTextCache CachedTooltip;
 	FNodeTextCache CachedNodeTitle;
 };
diff --git a/Engine/Source/Editor/BlueprintGraph/Classes/K2Node_GetNumEnumEntries.h b/Engine/Source/Editor/BlueprintGraph/Classes/K2Node_GetNumEnumEntries.h
--- a/Engine/Source/Editor/BlueprintGraph/Classes/K2Node_GetNumEnumEntries.h
+++ b/Engine/Source/Editor/BlueprintGraph/Classes/K2Node_GetNumEnumEntries.h
@@ -33,7 +33,8 @@ class UK2Node_GetNumEnumEntries : public UK2Node, public INodeDependingOnEnumInt
 	// End of INodeDependingOnEnumInterface
 
 private:
-	/** Constructing FText strings can be costly, so we cache the node's title */
+	/** Constructing FText strings can be costly, so we cache the node's title/tooltip */
+	FNodeTextCache CachedTooltip;
 	FNodeTextCache CachedNodeTitle;
 };
 
diff --git a/Engine/Source/Editor/BlueprintGraph/Classes/K2Node_InputAction.h b/Engine/Source/Editor/BlueprintGraph/Classes/K2Node_InputAction.h
--- a/Engine/Source/Editor/BlueprintGraph/Classes/K2Node_InputAction.h
+++ b/Engine/Source/Editor/BlueprintGraph/Classes/K2Node_InputAction.h
@@ -56,6 +56,7 @@ class UK2Node_InputAction : public UK2Node
 private:
 	void CreateInputActionEvent(FKismetCompilerContext& CompilerContext, UEdGraph* SourceGraph, UEdGraphPin* InputActionPin, const EInputEvent InputKeyEvent);
 
-	/** Constructing FText strings can be costly, so we cache the node's title */
+	/** Constructing FText strings can be costly, so we cache the node's title/tooltip */
+	FNodeTextCache CachedTooltip;
 	FNodeTextCache CachedNodeTitle;
 };
diff --git a/Engine/Source/Editor/BlueprintGraph/Classes/K2Node_InputAxisEvent.h b/Engine/Source/Editor/BlueprintGraph/Classes/K2Node_InputAxisEvent.h
--- a/Engine/Source/Editor/BlueprintGraph/Classes/K2Node_InputAxisEvent.h
+++ b/Engine/Source/Editor/BlueprintGraph/Classes/K2Node_InputAxisEvent.h
@@ -47,6 +47,7 @@ class UK2Node_InputAxisEvent : public UK2Node_Event
 	void Initialize(const FName AxisName);
 
 private:
-	/** Constructing FText strings can be costly, so we cache the node's title */
+	/** Constructing FText strings can be costly, so we cache the node's title/tooltip */
+	FNodeTextCache CachedTooltip;
 	FNodeTextCache CachedNodeTitle;
 };
diff --git a/Engine/Source/Editor/BlueprintGraph/Classes/K2Node_InputAxisKeyEvent.h b/Engine/Source/Editor/BlueprintGraph/Classes/K2Node_InputAxisKeyEvent.h
--- a/Engine/Source/Editor/BlueprintGraph/Classes/K2Node_InputAxisKeyEvent.h
+++ b/Engine/Source/Editor/BlueprintGraph/Classes/K2Node_InputAxisKeyEvent.h
@@ -41,4 +41,8 @@ class UK2Node_InputAxisKeyEvent : public UK2Node_Event
 	// End UK2Node interface
 
 	void Initialize(const FKey AxisKey);
+
+private:
+	/** Constructing FText strings can be costly, so we cache the node's tooltip */
+	FNodeTextCache CachedTooltip;
 };
diff --git a/Engine/Source/Editor/BlueprintGraph/Classes/K2Node_InputKey.h b/Engine/Source/Editor/BlueprintGraph/Classes/K2Node_InputKey.h
--- a/Engine/Source/Editor/BlueprintGraph/Classes/K2Node_InputKey.h
+++ b/Engine/Source/Editor/BlueprintGraph/Classes/K2Node_InputKey.h
@@ -79,6 +79,7 @@ class UK2Node_InputKey : public UK2Node
 private:
 	void CreateInputKeyEvent(FKismetCompilerContext& CompilerContext, UEdGraph* SourceGraph, UEdGraphPin* InputKeyPin, const EInputEvent KeyEvent);
 
-	/** Constructing FText strings can be costly, so we cache the node's title */
+	/** Constructing FText strings can be costly, so we cache the node's title/tooltip */
+	FNodeTextCache CachedTooltip;
 	FNodeTextCache CachedNodeTitle;
 };
diff --git a/Engine/Source/Editor/BlueprintGraph/Classes/K2Node_MacroInstance.h b/Engine/Source/Editor/BlueprintGraph/Classes/K2Node_MacroInstance.h
--- a/Engine/Source/Editor/BlueprintGraph/Classes/K2Node_MacroInstance.h
+++ b/Engine/Source/Editor/BlueprintGraph/Classes/K2Node_MacroInstance.h
@@ -4,6 +4,7 @@
 #pragma once
 #include "K2Node_Tunnel.h"
 #include "EdGraph/EdGraph.h"
+#include "EdGraph/EdGraphNodeUtils.h" // for FNodeTextCache
 #include "K2Node_MacroInstance.generated.h"
 
 UCLASS(MinimalAPI)
@@ -66,5 +67,9 @@ public:
 	// Finds the associated metadata for the macro instance if there is any; this function is not particularly fast.
 	BLUEPRINTGRAPH_API static FKismetUserDeclaredFunctionMetadata* GetAssociatedGraphMetadata(const UEdGraph* AssociatedMacroGraph);
 	static void FindInContentBrowser(TWeakObjectPtr<UK2Node_MacroInstance> MacroInstance);
+
+private:
+	/** Constructing FText strings can be costly, so we cache the node's tooltip */
+	FNodeTextCache CachedTooltip;
 };
 
diff --git a/Engine/Source/Editor/BlueprintGraph/Classes/K2Node_MakeStruct.h b/Engine/Source/Editor/BlueprintGraph/Classes/K2Node_MakeStruct.h
--- a/Engine/Source/Editor/BlueprintGraph/Classes/K2Node_MakeStruct.h
+++ b/Engine/Source/Editor/BlueprintGraph/Classes/K2Node_MakeStruct.h
@@ -44,6 +44,7 @@ protected:
 	};
 
 private:
-	/** Constructing FText strings can be costly, so we cache the node's title */
+	/** Constructing FText strings can be costly, so we cache the node's title/tooltip */
+	FNodeTextCache CachedTooltip;
 	FNodeTextCache CachedNodeTitle;
 };
diff --git a/Engine/Source/Editor/BlueprintGraph/Classes/K2Node_SetFieldsInStruct.h b/Engine/Source/Editor/BlueprintGraph/Classes/K2Node_SetFieldsInStruct.h
--- a/Engine/Source/Editor/BlueprintGraph/Classes/K2Node_SetFieldsInStruct.h
+++ b/Engine/Source/Editor/BlueprintGraph/Classes/K2Node_SetFieldsInStruct.h
@@ -36,6 +36,7 @@ class UK2Node_SetFieldsInStruct : public UK2Node_MakeStruct
 	BLUEPRINTGRAPH_API void RestoreAllPins();
 
 private:
-	/** Constructing FText strings can be costly, so we cache the node's title */
+	/** Constructing FText strings can be costly, so we cache the node's title/tooltip */
+	FNodeTextCache CachedTooltip;
 	FNodeTextCache CachedNodeTitle;
 };
diff --git a/Engine/Source/Editor/BlueprintGraph/Classes/K2Node_StructMemberGet.h b/Engine/Source/Editor/BlueprintGraph/Classes/K2Node_StructMemberGet.h
--- a/Engine/Source/Editor/BlueprintGraph/Classes/K2Node_StructMemberGet.h
+++ b/Engine/Source/Editor/BlueprintGraph/Classes/K2Node_StructMemberGet.h
@@ -34,7 +34,8 @@ class UK2Node_StructMemberGet : public UK2Node_StructOperation
 	BLUEPRINTGRAPH_API void AllocatePinsForSingleMemberGet(FName MemberName);
 
 private:
-	/** Constructing FText strings can be costly, so we cache the node's title */
+	/** Constructing FText strings can be costly, so we cache the node's title/tooltip */
+	FNodeTextCache CachedTooltip;
 	FNodeTextCache CachedNodeTitle;
 };
 
diff --git a/Engine/Source/Editor/BlueprintGraph/Classes/K2Node_StructMemberSet.h b/Engine/Source/Editor/BlueprintGraph/Classes/K2Node_StructMemberSet.h
--- a/Engine/Source/Editor/BlueprintGraph/Classes/K2Node_StructMemberSet.h
+++ b/Engine/Source/Editor/BlueprintGraph/Classes/K2Node_StructMemberSet.h
@@ -32,7 +32,8 @@ class UK2Node_StructMemberSet : public UK2Node_StructOperation
 	// End of UK2Node interface
 
 private:
-	/** Constructing FText strings can be costly, so we cache the node's title */
+	/** Constructing FText strings can be costly, so we cache the node's title/tooltip */
+	FNodeTextCache CachedTooltip;
 	FNodeTextCache CachedNodeTitle;
 };
 
diff --git a/Engine/Source/Editor/BlueprintGraph/Classes/K2Node_TemporaryVariable.h b/Engine/Source/Editor/BlueprintGraph/Classes/K2Node_TemporaryVariable.h
--- a/Engine/Source/Editor/BlueprintGraph/Classes/K2Node_TemporaryVariable.h
+++ b/Engine/Source/Editor/BlueprintGraph/Classes/K2Node_TemporaryVariable.h
@@ -36,7 +36,8 @@ class UK2Node_TemporaryVariable : public UK2Node
 	// End UK2Node interface.
 
 private:
-	/** Constructing FText strings can be costly, so we cache the node's title */
+	/** Constructing FText strings can be costly, so we cache the node's title/tooltip */
+	FNodeTextCache CachedTooltip;
 	FNodeTextCache CachedNodeTitle;
 };
 
diff --git a/Engine/Source/Editor/BlueprintGraph/Classes/K2Node_VariableGet.h b/Engine/Source/Editor/BlueprintGraph/Classes/K2Node_VariableGet.h
--- a/Engine/Source/Editor/BlueprintGraph/Classes/K2Node_VariableGet.h
+++ b/Engine/Source/Editor/BlueprintGraph/Classes/K2Node_VariableGet.h
@@ -24,7 +24,8 @@ class BLUEPRINTGRAPH_API UK2Node_VariableGet : public UK2Node_Variable
 	virtual class FNodeHandlingFunctor* CreateNodeHandler(class FKismetCompilerContext& CompilerContext) const override;
 	// End K2Node interface
 private:
-	/** Constructing FText strings can be costly, so we cache the node's title */
+	/** Constructing FText strings can be costly, so we cache the node's title/tooltip */
+	FNodeTextCache CachedTooltip;
 	FNodeTextCache CachedNodeTitle;
 };
 
diff --git a/Engine/Source/Editor/BlueprintGraph/Classes/K2Node_VariableSet.h b/Engine/Source/Editor/BlueprintGraph/Classes/K2Node_VariableSet.h
--- a/Engine/Source/Editor/BlueprintGraph/Classes/K2Node_VariableSet.h
+++ b/Engine/Source/Editor/BlueprintGraph/Classes/K2Node_VariableSet.h
@@ -29,7 +29,8 @@ class UK2Node_VariableSet : public UK2Node_Variable
 	BLUEPRINTGRAPH_API bool ShouldFlushDormancyOnSet() const;
 
 private:
-	/** Constructing FText strings can be costly, so we cache the node's title */
+	/** Constructing FText strings can be costly, so we cache the node's title/tooltip */
+	FNodeTextCache CachedTooltip;
 	FNodeTextCache CachedNodeTitle;
 };
 
diff --git a/Engine/Source/Editor/BlueprintGraph/Private/K2Node_BreakStruct.cpp b/Engine/Source/Editor/BlueprintGraph/Private/K2Node_BreakStruct.cpp
--- a/Engine/Source/Editor/BlueprintGraph/Private/K2Node_BreakStruct.cpp
+++ b/Engine/Source/Editor/BlueprintGraph/Private/K2Node_BreakStruct.cpp
@@ -194,10 +194,14 @@ void UK2Node_BreakStruct::AllocateDefaultPins()
 
 FText UK2Node_BreakStruct::GetNodeTitle(ENodeTitleType::Type TitleType) const
 {
-	if (CachedNodeTitle.IsOutOfDate())
+	if (StructType == nullptr)
+	{
+		return LOCTEXT("BreakNullStruct_Title", "Break <unknown struct>");
+	}
+	else if (CachedNodeTitle.IsOutOfDate())
 	{
 		FFormatNamedArguments Args;
-		Args.Add(TEXT("StructName"), FText::FromString(StructType ? StructType->GetName() : FString()));
+		Args.Add(TEXT("StructName"), FText::FromString(StructType->GetName()));
 
 		// FText::Format() is slow, so we cache this to save on performance
 		CachedNodeTitle = FText::Format(LOCTEXT("BreakNodeTitle", "Break {StructName}"), Args);
@@ -207,10 +211,19 @@ FText UK2Node_BreakStruct::GetNodeTitle(ENodeTitleType::Type TitleType) const
 
 FText UK2Node_BreakStruct::GetTooltipText() const
 {
-	return FText::Format(
-		LOCTEXT("MakeStruct_Tooltip", "Adds a node that breaks a '{0}' into its member fields"),
-		StructType ? FText::FromName(StructType->GetFName()) : FText::GetEmpty()
+	if (StructType == nullptr)
+	{
+		return LOCTEXT("BreakNullStruct_Tooltip", "Adds a node that breaks an '<unknown struct>' into its member fields");
+	}
+	else if (CachedTooltip.IsOutOfDate())
+	{
+		// FText::Format() is slow, so we cache this to save on performance
+		CachedTooltip = FText::Format(
+			LOCTEXT("BreakStruct_Tooltip", "Adds a node that breaks a '{0}' into its member fields"),
+			FText::FromName(StructType->GetFName())
 		);
+	}
+	return CachedTooltip;
 }
 
 void UK2Node_BreakStruct::ValidateNodeDuringCompilation(class FCompilerResultsLog& MessageLog) const
diff --git a/Engine/Source/Editor/BlueprintGraph/Private/K2Node_CallFunction.cpp b/Engine/Source/Editor/BlueprintGraph/Private/K2Node_CallFunction.cpp
--- a/Engine/Source/Editor/BlueprintGraph/Private/K2Node_CallFunction.cpp
+++ b/Engine/Source/Editor/BlueprintGraph/Private/K2Node_CallFunction.cpp
@@ -750,19 +750,24 @@ FLinearColor UK2Node_CallFunction::GetNodeTitleColor() const
 		return GetDefault<UGraphEditorSettings>()->PureFunctionCallNodeTitleColor;
 	}
 
-		return Super::GetNodeTitleColor();
-	}
+	return Super::GetNodeTitleColor();
+}
 
 FText UK2Node_CallFunction::GetTooltipText() const
 {
 	FText Tooltip;
 
-	if (UFunction* Function = GetTargetFunction())
+	UFunction* Function = GetTargetFunction();
+	if (Function == nullptr)
 	{
-		Tooltip = FText::FromString(GetDefaultTooltipForFunction(Function));
+		return FText::Format(LOCTEXT("CallUnknownFunction", "Call unknown function {0}"), FText::FromName(FunctionReference.GetMemberName()));
+	}
+	else if (CachedTooltip.IsOutOfDate())
+	{
+		FText BaseTooltip = FText::FromString(GetDefaultTooltipForFunction(Function));
 
 		FFormatNamedArguments Args;
-		Args.Add(TEXT("DefaultTooltip"), Tooltip);
+		Args.Add(TEXT("DefaultTooltip"), BaseTooltip);
 
 		if (Function->HasAllFunctionFlags(FUNC_BlueprintAuthorityOnly))
 		{
@@ -770,7 +775,8 @@ FText UK2Node_CallFunction::GetTooltipText() const
 				TEXT("ClientString"),
 				NSLOCTEXT("K2Node", "ServerFunction", "Authority Only. This function will only execute on the server.")
 			);
-			Tooltip = FText::Format(LOCTEXT("CallFunction_SubtitledTooltip", "{DefaultTooltip}\n\n{ClientString}"), Args);
+			// FText::Format() is slow, so we cache this to save on performance
+			CachedTooltip = FText::Format(LOCTEXT("CallFunction_SubtitledTooltip", "{DefaultTooltip}\n\n{ClientString}"), Args);
 		}
 		else if (Function->HasAllFunctionFlags(FUNC_BlueprintCosmetic))
 		{
@@ -778,15 +784,15 @@ FText UK2Node_CallFunction::GetTooltipText() const
 				TEXT("ClientString"),
 				NSLOCTEXT("K2Node", "ClientEvent", "Cosmetic. This event is only for cosmetic, non-gameplay actions.")
 			);
-			Tooltip = FText::Format(LOCTEXT("CallFunction_SubtitledTooltip", "{DefaultTooltip}\n\n{ClientString}"), Args);
+			// FText::Format() is slow, so we cache this to save on performance
+			CachedTooltip = FText::Format(LOCTEXT("CallFunction_SubtitledTooltip", "{DefaultTooltip}\n\n{ClientString}"), Args);
 		} 
+		else
+		{
+			CachedTooltip = BaseTooltip;
+		}
 	}
-	else
-	{
-		Tooltip = FText::Format(LOCTEXT("CallUnknownFunction", "Call unknown function {0}"), FText::FromName(FunctionReference.GetMemberName()));
-	}
-
-	return Tooltip;
+	return CachedTooltip;
 }
 
 FString UK2Node_CallFunction::GetUserFacingFunctionName(const UFunction* Function)
diff --git a/Engine/Source/Editor/BlueprintGraph/Private/K2Node_CastByteToEnum.cpp b/Engine/Source/Editor/BlueprintGraph/Private/K2Node_CastByteToEnum.cpp
--- a/Engine/Source/Editor/BlueprintGraph/Private/K2Node_CastByteToEnum.cpp
+++ b/Engine/Source/Editor/BlueprintGraph/Private/K2Node_CastByteToEnum.cpp
@@ -28,10 +28,15 @@ void UK2Node_CastByteToEnum::AllocateDefaultPins()
 
 FText UK2Node_CastByteToEnum::GetTooltipText() const
 {
-	return FText::Format(
-		NSLOCTEXT("K2Node", "CastByteToEnum_Tooltip", "Byte to Enum {0}"),
-		FText::FromName(Enum->GetFName())
-	);
+	if (CachedTooltip.IsOutOfDate())
+	{
+		// FText::Format() is slow, so we cache this to save on performance
+		CachedTooltip = FText::Format(
+			NSLOCTEXT("K2Node", "CastByteToEnum_Tooltip", "Byte to Enum {0}"),
+			FText::FromName(Enum->GetFName())
+		);
+	}
+	return CachedTooltip;
 }
 
 FText UK2Node_CastByteToEnum::GetNodeTitle(ENodeTitleType::Type TitleType) const
diff --git a/Engine/Source/Editor/BlueprintGraph/Private/K2Node_DelegateSet.cpp b/Engine/Source/Editor/BlueprintGraph/Private/K2Node_DelegateSet.cpp
--- a/Engine/Source/Editor/BlueprintGraph/Private/K2Node_DelegateSet.cpp
+++ b/Engine/Source/Editor/BlueprintGraph/Private/K2Node_DelegateSet.cpp
@@ -149,17 +149,21 @@ void UK2Node_DelegateSet::AllocateDefaultPins()
 
 FText UK2Node_DelegateSet::GetTooltipText() const
 {
-	FText Tooltip = FText::Format(NSLOCTEXT("K2Node", "CreateEventForDelegate", "Create an event tied to the delegate {0}"), FText::FromName(DelegatePropertyName));
-	if (UFunction* Function = GetDelegateSignature())
+	if (CachedTooltip.IsOutOfDate())
 	{
-		const FText SignatureTooltip = Function->GetToolTipText();
-
-		if (!SignatureTooltip.IsEmpty())
+		// FText::Format() is slow, so we cache this to save on performance
+		CachedTooltip = FText::Format(NSLOCTEXT("K2Node", "CreateEventForDelegate", "Create an event tied to the delegate {0}"), FText::FromName(DelegatePropertyName));
+		if (UFunction* Function = GetDelegateSignature())
 		{
-			Tooltip = FText::Format(LOCTEXT("DelegateSet_SubtitledTooltip", "{0}\n{1}"), Tooltip, SignatureTooltip);
+			const FText SignatureTooltip = Function->GetToolTipText();
+
+			if (!SignatureTooltip.IsEmpty())
+			{
+				CachedTooltip = FText::Format(LOCTEXT("DelegateSet_SubtitledTooltip", "{0}\n{1}"), (FText&)CachedTooltip, SignatureTooltip);
+			}
 		}
 	}
-	return Tooltip;
+	return CachedTooltip;
 }
 
 FText UK2Node_DelegateSet::GetNodeTitle(ENodeTitleType::Type TitleType) const
diff --git a/Engine/Source/Editor/BlueprintGraph/Private/K2Node_Event.cpp b/Engine/Source/Editor/BlueprintGraph/Private/K2Node_Event.cpp
--- a/Engine/Source/Editor/BlueprintGraph/Private/K2Node_Event.cpp
+++ b/Engine/Source/Editor/BlueprintGraph/Private/K2Node_Event.cpp
@@ -91,17 +91,15 @@ FText UK2Node_Event::GetNodeTitle(ENodeTitleType::Type TitleType) const
 
 FText UK2Node_Event::GetTooltipText() const
 {
-	FText Tooltip;
-
 	UFunction* Function = FindField<UFunction>(EventSignatureClass, EventSignatureName);
-	if (Function != NULL)
+	if (CachedTooltip.IsOutOfDate() && (Function != nullptr))
 	{
-		Tooltip = Function->GetToolTipText();
+		CachedTooltip = Function->GetToolTipText();
 
 		if (bOverrideFunction || (CustomFunctionName == NAME_None))
 		{
 			FFormatNamedArguments Args;
-			Args.Add(TEXT("FunctionTooltip"), Tooltip);
+			Args.Add(TEXT("FunctionTooltip"), (FText&)CachedTooltip);
 
 			//@TODO: KISMETREPLICATION: Should do this for events with a custom function name, if it's a newly introduced replicating thingy
 			if (Function->HasAllFunctionFlags(FUNC_BlueprintCosmetic) || IsCosmeticTickEvent())
@@ -110,7 +108,8 @@ FText UK2Node_Event::GetTooltipText() const
 					TEXT("ClientString"),
 					NSLOCTEXT("K2Node", "ClientEvent", "\n\nCosmetic. This event is only for cosmetic, non-gameplay actions.")
 				);
-				Tooltip = FText::Format(LOCTEXT("Event_SubtitledTooltip", "{FunctionTooltip}\n\n{ClientString}"), Args);
+				// FText::Format() is slow, so we cache this to save on performance
+				CachedTooltip = FText::Format(LOCTEXT("Event_SubtitledTooltip", "{FunctionTooltip}\n\n{ClientString}"), Args);
 			} 
 			else if(Function->HasAllFunctionFlags(FUNC_BlueprintAuthorityOnly))
 			{
@@ -118,12 +117,13 @@ FText UK2Node_Event::GetTooltipText() const
 					TEXT("ClientString"),
 					NSLOCTEXT("K2Node", "ServerEvent", "\n\nAuthority Only. This event only fires on the server.")
 				);
-				Tooltip = FText::Format(LOCTEXT("Event_SubtitledTooltip", "{FunctionTooltip}\n\n{ClientString}"), Args);
+				// FText::Format() is slow, so we cache this to save on performance
+				CachedTooltip = FText::Format(LOCTEXT("Event_SubtitledTooltip", "{FunctionTooltip}\n\n{ClientString}"), Args);
 			}			
 		}		
 	}
 
-	return Tooltip;
+	return CachedTooltip;
 }
 
 FString UK2Node_Event::GetKeywords() const
diff --git a/Engine/Source/Editor/BlueprintGraph/Private/K2Node_GetInputAxisKeyValue.cpp b/Engine/Source/Editor/BlueprintGraph/Private/K2Node_GetInputAxisKeyValue.cpp
--- a/Engine/Source/Editor/BlueprintGraph/Private/K2Node_GetInputAxisKeyValue.cpp
+++ b/Engine/Source/Editor/BlueprintGraph/Private/K2Node_GetInputAxisKeyValue.cpp
@@ -30,6 +30,8 @@ void UK2Node_GetInputAxisKeyValue::Initialize(const FKey AxisKey)
 {
 	InputAxisKey = AxisKey;
 	SetFromFunction(AActor::StaticClass()->FindFunctionByName(TEXT("GetInputAxisKeyValue")));
+	
+	CachedTooltip.MarkDirty();
 	CachedNodeTitle.MarkDirty();
 }
 
@@ -52,9 +54,14 @@ FText UK2Node_GetInputAxisKeyValue::GetNodeTitle(ENodeTitleType::Type TitleType)
 
 FText UK2Node_GetInputAxisKeyValue::GetTooltipText() const
 {
-	FFormatNamedArguments Args;
-	Args.Add(TEXT("AxisKey"), InputAxisKey.GetDisplayName());
-	return FText::Format(NSLOCTEXT("K2Node", "GetInputAxisKey_Tooltip", "Returns the current value of input axis key {AxisKey}.  If input is disabled for the actor the value will be 0."), Args);
+	if (CachedTooltip.IsOutOfDate())
+	{
+		FFormatNamedArguments Args;
+		Args.Add(TEXT("AxisKey"), InputAxisKey.GetDisplayName());
+		// FText::Format() is slow, so we cache this to save on performance
+		CachedTooltip = FText::Format(NSLOCTEXT("K2Node", "GetInputAxisKey_Tooltip", "Returns the current value of input axis key {AxisKey}.  If input is disabled for the actor the value will be 0."), Args);
+	}
+	return CachedTooltip;
 }
 
 bool UK2Node_GetInputAxisKeyValue::IsCompatibleWithGraph(UEdGraph const* Graph) const
diff --git a/Engine/Source/Editor/BlueprintGraph/Private/K2Node_GetInputAxisValue.cpp b/Engine/Source/Editor/BlueprintGraph/Private/K2Node_GetInputAxisValue.cpp
--- a/Engine/Source/Editor/BlueprintGraph/Private/K2Node_GetInputAxisValue.cpp
+++ b/Engine/Source/Editor/BlueprintGraph/Private/K2Node_GetInputAxisValue.cpp
@@ -53,7 +53,12 @@ FText UK2Node_GetInputAxisValue::GetNodeTitle(ENodeTitleType::Type TitleType) co
 
 FText UK2Node_GetInputAxisValue::GetTooltipText() const
 {
-	return FText::Format(NSLOCTEXT("K2Node", "GetInputAxis_Tooltip", "Returns the current value of input axis {0}.  If input is disabled for the actor the value will be 0."), FText::FromName(InputAxisName));
+	if (CachedTooltip.IsOutOfDate())
+	{
+		// FText::Format() is slow, so we cache this to save on performance
+		CachedTooltip = FText::Format(NSLOCTEXT("K2Node", "GetInputAxis_Tooltip", "Returns the current value of input axis {0}.  If input is disabled for the actor the value will be 0."), FText::FromName(InputAxisName));
+	}
+	return CachedTooltip;
 }
 
 bool UK2Node_GetInputAxisValue::IsCompatibleWithGraph(UEdGraph const* Graph) const
diff --git a/Engine/Source/Editor/BlueprintGraph/Private/K2Node_GetNumEnumEntries.cpp b/Engine/Source/Editor/BlueprintGraph/Private/K2Node_GetNumEnumEntries.cpp
--- a/Engine/Source/Editor/BlueprintGraph/Private/K2Node_GetNumEnumEntries.cpp
+++ b/Engine/Source/Editor/BlueprintGraph/Private/K2Node_GetNumEnumEntries.cpp
@@ -26,9 +26,16 @@ void UK2Node_GetNumEnumEntries::AllocateDefaultPins()
 
 FText UK2Node_GetNumEnumEntries::GetTooltipText() const
 {
-	const FString EnumName = (Enum != NULL) ? Enum->GetName() : TEXT("(bad enum)");
-
-	return FText::Format(NSLOCTEXT("K2Node", "GetNumEnumEntries_Tooltip", "Returns {0}_MAX value"), FText::FromString(EnumName));
+	if (Enum == nullptr)
+	{
+		return NSLOCTEXT("K2Node", "GetNumEnumEntries_BadTooltip", "Returns (bad enum)_MAX value");
+	}
+	else if (CachedTooltip.IsOutOfDate())
+	{
+		// FText::Format() is slow, so we cache this to save on performance
+		CachedTooltip = FText::Format(NSLOCTEXT("K2Node", "GetNumEnumEntries_Tooltip", "Returns {0}_MAX value"), FText::FromName(Enum->GetFName()));
+	}
+	return CachedTooltip;
 }
 
 FText UK2Node_GetNumEnumEntries::GetNodeTitle(ENodeTitleType::Type TitleType) const
diff --git a/Engine/Source/Editor/BlueprintGraph/Private/K2Node_InputAction.cpp b/Engine/Source/Editor/BlueprintGraph/Private/K2Node_InputAction.cpp
--- a/Engine/Source/Editor/BlueprintGraph/Private/K2Node_InputAction.cpp
+++ b/Engine/Source/Editor/BlueprintGraph/Private/K2Node_InputAction.cpp
@@ -66,7 +66,12 @@ FText UK2Node_InputAction::GetNodeTitle(ENodeTitleType::Type TitleType) const
 
 FText UK2Node_InputAction::GetTooltipText() const
 {
-	return FText::Format(NSLOCTEXT("K2Node", "InputAction_Tooltip", "Event for when the keys bound to input action {0} are pressed or released."), FText::FromName(InputActionName));
+	if (CachedTooltip.IsOutOfDate())
+	{
+		// FText::Format() is slow, so we cache this to save on performance
+		CachedTooltip = FText::Format(NSLOCTEXT("K2Node", "InputAction_Tooltip", "Event for when the keys bound to input action {0} are pressed or released."), FText::FromName(InputActionName));
+	}
+	return CachedTooltip;
 }
 
 bool UK2Node_InputAction::IsCompatibleWithGraph(UEdGraph const* Graph) const
diff --git a/Engine/Source/Editor/BlueprintGraph/Private/K2Node_InputAxisEvent.cpp b/Engine/Source/Editor/BlueprintGraph/Private/K2Node_InputAxisEvent.cpp
--- a/Engine/Source/Editor/BlueprintGraph/Private/K2Node_InputAxisEvent.cpp
+++ b/Engine/Source/Editor/BlueprintGraph/Private/K2Node_InputAxisEvent.cpp
@@ -58,7 +58,12 @@ FText UK2Node_InputAxisEvent::GetNodeTitle(ENodeTitleType::Type TitleType) const
 
 FText UK2Node_InputAxisEvent::GetTooltipText() const
 {
-	return FText::Format(NSLOCTEXT("K2Node", "InputAxis_Tooltip", "Event that provides the current value of the {0} axis once per frame when input is enabled for the containing actor."), FText::FromName(InputAxisName));
+	if (CachedTooltip.IsOutOfDate())
+	{
+		// FText::Format() is slow, so we cache this to save on performance
+		CachedTooltip = FText::Format(NSLOCTEXT("K2Node", "InputAxis_Tooltip", "Event that provides the current value of the {0} axis once per frame when input is enabled for the containing actor."), FText::FromName(InputAxisName));
+	}
+	return CachedTooltip;
 }
 
 void UK2Node_InputAxisEvent::ValidateNodeDuringCompilation(class FCompilerResultsLog& MessageLog) const
diff --git a/Engine/Source/Editor/BlueprintGraph/Private/K2Node_InputAxisKeyEvent.cpp b/Engine/Source/Editor/BlueprintGraph/Private/K2Node_InputAxisKeyEvent.cpp
--- a/Engine/Source/Editor/BlueprintGraph/Private/K2Node_InputAxisKeyEvent.cpp
+++ b/Engine/Source/Editor/BlueprintGraph/Private/K2Node_InputAxisKeyEvent.cpp
@@ -36,7 +36,12 @@ FText UK2Node_InputAxisKeyEvent::GetNodeTitle(ENodeTitleType::Type TitleType) co
 
 FText UK2Node_InputAxisKeyEvent::GetTooltipText() const
 {
-	return FText::Format(NSLOCTEXT("K2Node", "InputAxisKey_Tooltip", "Event that provides the current value of the {0} axis once per frame when input is enabled for the containing actor."), AxisKey.GetDisplayName());
+	if (CachedTooltip.IsOutOfDate())
+	{
+		// FText::Format() is slow, so we cache this to save on performance
+		CachedTooltip = FText::Format(NSLOCTEXT("K2Node", "InputAxisKey_Tooltip", "Event that provides the current value of the {0} axis once per frame when input is enabled for the containing actor."), AxisKey.GetDisplayName());
+	}
+	return CachedTooltip;
 }
 
 void UK2Node_InputAxisKeyEvent::ValidateNodeDuringCompilation(class FCompilerResultsLog& MessageLog) const
diff --git a/Engine/Source/Editor/BlueprintGraph/Private/K2Node_InputKey.cpp b/Engine/Source/Editor/BlueprintGraph/Private/K2Node_InputKey.cpp
--- a/Engine/Source/Editor/BlueprintGraph/Private/K2Node_InputKey.cpp
+++ b/Engine/Source/Editor/BlueprintGraph/Private/K2Node_InputKey.cpp
@@ -250,15 +250,22 @@ FText UK2Node_InputKey::GetNodeTitle(ENodeTitleType::Type TitleType) const
 
 FText UK2Node_InputKey::GetTooltipText() const
 {
-	FText ModifierText = GetModifierText();
-	FText KeyText = GetKeyText();
-
-	if ( !ModifierText.IsEmpty() )
+	if (CachedTooltip.IsOutOfDate())
 	{
-		return FText::Format(NSLOCTEXT("K2Node", "InputKey_Tooltip_Modifiers", "Events for when the {0} key is pressed or released while {1} is also held."), KeyText, ModifierText);
-	}
+		FText ModifierText = GetModifierText();
+		FText KeyText = GetKeyText();
 
-	return FText::Format(NSLOCTEXT("K2Node", "InputKey_Tooltip", "Events for when the {0} key is pressed or released."), KeyText);
+		// FText::Format() is slow, so we cache this to save on performance
+		if (!ModifierText.IsEmpty())
+		{
+			CachedTooltip = FText::Format(NSLOCTEXT("K2Node", "InputKey_Tooltip_Modifiers", "Events for when the {0} key is pressed or released while {1} is also held."), KeyText, ModifierText);
+		}
+		else
+		{
+			CachedTooltip = FText::Format(NSLOCTEXT("K2Node", "InputKey_Tooltip", "Events for when the {0} key is pressed or released."), KeyText);
+		}
+	}
+	return CachedTooltip;
 }
 
 FName UK2Node_InputKey::GetPaletteIcon(FLinearColor& OutColor) const
diff --git a/Engine/Source/Editor/BlueprintGraph/Private/K2Node_MacroInstance.cpp b/Engine/Source/Editor/BlueprintGraph/Private/K2Node_MacroInstance.cpp
--- a/Engine/Source/Editor/BlueprintGraph/Private/K2Node_MacroInstance.cpp
+++ b/Engine/Source/Editor/BlueprintGraph/Private/K2Node_MacroInstance.cpp
@@ -111,15 +111,17 @@ FText UK2Node_MacroInstance::GetTooltipText() const
 			return FText::FromString(Metadata->ToolTip);
 		}
 	}
-	
-	if (MacroGraph)
+
+	if (MacroGraph == nullptr)
 	{
-		return FText::Format(NSLOCTEXT("K2Node", "MacroGraphInstance_Tooltip", "{0} instance"), FText::FromName(MacroGraph->GetFName()));
+		return NSLOCTEXT("K2Node", "Macro_Tooltip", "Macro");
 	}
-	else
+	else if (CachedTooltip.IsOutOfDate())
 	{
-		return NSLOCTEXT("K2Node", "Macro_Tooltip", "Macro");
+		// FText::Format() is slow, so we cache this to save on performance
+		CachedTooltip = FText::Format(NSLOCTEXT("K2Node", "MacroGraphInstance_Tooltip", "{0} instance"), FText::FromName(MacroGraph->GetFName()));
 	}
+	return CachedTooltip;
 }
 
 FText UK2Node_MacroInstance::GetNodeTitle(ENodeTitleType::Type TitleType) const
diff --git a/Engine/Source/Editor/BlueprintGraph/Private/K2Node_MakeStruct.cpp b/Engine/Source/Editor/BlueprintGraph/Private/K2Node_MakeStruct.cpp
--- a/Engine/Source/Editor/BlueprintGraph/Private/K2Node_MakeStruct.cpp
+++ b/Engine/Source/Editor/BlueprintGraph/Private/K2Node_MakeStruct.cpp
@@ -152,10 +152,14 @@ void UK2Node_MakeStruct::ValidateNodeDuringCompilation(class FCompilerResultsLog
 
 FText UK2Node_MakeStruct::GetNodeTitle(ENodeTitleType::Type TitleType) const 
 {
-	if (CachedNodeTitle.IsOutOfDate())
+	if (StructType == nullptr)
+	{
+		return LOCTEXT("MakeNullStructTitle", "Make <unknown struct>");
+	}
+	else if (CachedNodeTitle.IsOutOfDate())
 	{
 		FFormatNamedArguments Args;
-		Args.Add(TEXT("StructName"), FText::FromString(StructType ? StructType->GetName() : FString()));
+		Args.Add(TEXT("StructName"), FText::FromName(StructType->GetFName()));
 		// FText::Format() is slow, so we cache this to save on performance
 		CachedNodeTitle = FText::Format(LOCTEXT("MakeNodeTitle", "Make {StructName}"), Args);
 	}
@@ -164,10 +168,19 @@ FText UK2Node_MakeStruct::GetNodeTitle(ENodeTitleType::Type TitleType) const
 
 FText UK2Node_MakeStruct::GetTooltipText() const
 {
-	return FText::Format(
-		LOCTEXT("MakeStruct_Tooltip", "Adds a node that create a '{0}' from its members"),
-		StructType ? FText::FromName(StructType->GetFName()) : FText::GetEmpty()
-	);
+	if (StructType == nullptr)
+	{
+		return LOCTEXT("MakeNullStruct_Tooltip", "Adds a node that create an '<unknown struct>' from its members");
+	}
+	else if (CachedTooltip.IsOutOfDate())
+	{
+		// FText::Format() is slow, so we cache this to save on performance
+		CachedTooltip = FText::Format(
+			LOCTEXT("MakeStruct_Tooltip", "Adds a node that create a '{0}' from its members"),
+			FText::FromName(StructType->GetFName())
+		);
+	}
+	return CachedTooltip;
 }
 
 FLinearColor UK2Node_MakeStruct::GetNodeTitleColor() const
diff --git a/Engine/Source/Editor/BlueprintGraph/Private/K2Node_SetFieldsInStruct.cpp b/Engine/Source/Editor/BlueprintGraph/Private/K2Node_SetFieldsInStruct.cpp
--- a/Engine/Source/Editor/BlueprintGraph/Private/K2Node_SetFieldsInStruct.cpp
+++ b/Engine/Source/Editor/BlueprintGraph/Private/K2Node_SetFieldsInStruct.cpp
@@ -55,10 +55,14 @@ void UK2Node_SetFieldsInStruct::AllocateDefaultPins()
 
 FText UK2Node_SetFieldsInStruct::GetNodeTitle(ENodeTitleType::Type TitleType) const
 {
-	if (CachedNodeTitle.IsOutOfDate())
+	if (StructType == nullptr)
+	{
+		return LOCTEXT("SetFieldsInNullStructNodeTitle", "Set members in <unknown struct>");
+	}
+	else if (CachedNodeTitle.IsOutOfDate())
 	{
 		FFormatNamedArguments Args;
-		Args.Add(TEXT("StructName"), FText::FromString(StructType ? StructType->GetName() : FString()));
+		Args.Add(TEXT("StructName"), FText::FromName(StructType->GetFName()));
 		// FText::Format() is slow, so we cache this to save on performance
 		CachedNodeTitle = FText::Format(LOCTEXT("SetFieldsInStructNodeTitle", "Set members in {StructName}"), Args);
 	}
@@ -67,10 +71,19 @@ FText UK2Node_SetFieldsInStruct::GetNodeTitle(ENodeTitleType::Type TitleType) co
 
 FText UK2Node_SetFieldsInStruct::GetTooltipText() const
 {
-	return FText::Format(
-		LOCTEXT("SetFieldsInStruct_Tooltip", "Adds a node that modifies a '{0}'"),
-		StructType ? FText::FromName(StructType->GetFName()) : FText::GetEmpty()
-	);
+	if (StructType == nullptr)
+	{
+		return LOCTEXT("SetFieldsInStruct_NullTooltip", "Adds a node that modifies an '<unknown struct>'");
+	}
+	else if (CachedTooltip.IsOutOfDate())
+	{
+		// FText::Format() is slow, so we cache this to save on performance
+		CachedTooltip = FText::Format(
+			LOCTEXT("SetFieldsInStruct_Tooltip", "Adds a node that modifies a '{0}'"),
+			FText::FromName(StructType->GetFName())
+		);
+	}
+	return CachedTooltip;
 }
 
 FName UK2Node_SetFieldsInStruct::GetPaletteIcon(FLinearColor& OutColor) const
diff --git a/Engine/Source/Editor/BlueprintGraph/Private/K2Node_StructMemberGet.cpp b/Engine/Source/Editor/BlueprintGraph/Private/K2Node_StructMemberGet.cpp
--- a/Engine/Source/Editor/BlueprintGraph/Private/K2Node_StructMemberGet.cpp
+++ b/Engine/Source/Editor/BlueprintGraph/Private/K2Node_StructMemberGet.cpp
@@ -71,9 +71,14 @@ void UK2Node_StructMemberGet::AllocatePinsForSingleMemberGet(FName MemberName)
 
 FText UK2Node_StructMemberGet::GetTooltipText() const
 {
-	FFormatNamedArguments Args;
-	Args.Add(TEXT("VariableName"), FText::FromString(GetVarNameString()));
-	return FText::Format(LOCTEXT("K2Node_StructMemberGet_Tooltip", "Get member variables of {VariableName}"), Args);
+	if (CachedTooltip.IsOutOfDate())
+	{
+		FFormatNamedArguments Args;
+		Args.Add(TEXT("VariableName"), FText::FromString(GetVarNameString()));
+		// FText::Format() is slow, so we cache this to save on performance
+		CachedTooltip = FText::Format(LOCTEXT("K2Node_StructMemberGet_Tooltip", "Get member variables of {VariableName}"), Args);
+	}
+	return CachedTooltip;
 }
 
 FText UK2Node_StructMemberGet::GetNodeTitle(ENodeTitleType::Type TitleType) const
diff --git a/Engine/Source/Editor/BlueprintGraph/Private/K2Node_StructMemberSet.cpp b/Engine/Source/Editor/BlueprintGraph/Private/K2Node_StructMemberSet.cpp
--- a/Engine/Source/Editor/BlueprintGraph/Private/K2Node_StructMemberSet.cpp
+++ b/Engine/Source/Editor/BlueprintGraph/Private/K2Node_StructMemberSet.cpp
@@ -44,9 +44,14 @@ void UK2Node_StructMemberSet::AllocateDefaultPins()
 
 FText UK2Node_StructMemberSet::GetTooltipText() const
 {
-	FFormatNamedArguments Args;
-	Args.Add(TEXT("VariableName"), FText::FromName(VariableReference.GetMemberName()));
-	return FText::Format(LOCTEXT("K2Node_StructMemberSet_Tooltip", "Set member variables of {VariableName}"), Args);
+	if (CachedTooltip.IsOutOfDate())
+	{
+		FFormatNamedArguments Args;
+		Args.Add(TEXT("VariableName"), FText::FromName(VariableReference.GetMemberName()));
+		// FText::Format() is slow, so we cache this to save on performance
+		CachedTooltip = FText::Format(LOCTEXT("K2Node_StructMemberSet_Tooltip", "Set member variables of {VariableName}"), Args);
+	}
+	return CachedTooltip;
 }
 
 FText UK2Node_StructMemberSet::GetNodeTitle(ENodeTitleType::Type TitleType) const
diff --git a/Engine/Source/Editor/BlueprintGraph/Private/K2Node_TemporaryVariable.cpp b/Engine/Source/Editor/BlueprintGraph/Private/K2Node_TemporaryVariable.cpp
--- a/Engine/Source/Editor/BlueprintGraph/Private/K2Node_TemporaryVariable.cpp
+++ b/Engine/Source/Editor/BlueprintGraph/Private/K2Node_TemporaryVariable.cpp
@@ -49,9 +49,14 @@ void UK2Node_TemporaryVariable::AllocateDefaultPins()
 
 FText UK2Node_TemporaryVariable::GetTooltipText() const
 {
-	FFormatNamedArguments Args;
-	Args.Add(TEXT("VariableType"), UEdGraphSchema_K2::TypeToText(VariableType));
-	return FText::Format(NSLOCTEXT("K2Node", "LocalTemporaryVariable", "Local temporary {VariableType} variable"), Args);
+	if (CachedTooltip.IsOutOfDate())
+	{
+		FFormatNamedArguments Args;
+		Args.Add(TEXT("VariableType"), UEdGraphSchema_K2::TypeToText(VariableType));
+		// FText::Format() is slow, so we cache this to save on performance
+		CachedTooltip = FText::Format(NSLOCTEXT("K2Node", "LocalTemporaryVariable", "Local temporary {VariableType} variable"), Args);
+	}
+	return CachedTooltip;
 }
 
 FText UK2Node_TemporaryVariable::GetNodeTitle(ENodeTitleType::Type TitleType) const
diff --git a/Engine/Source/Editor/BlueprintGraph/Private/K2Node_VariableGet.cpp b/Engine/Source/Editor/BlueprintGraph/Private/K2Node_VariableGet.cpp
--- a/Engine/Source/Editor/BlueprintGraph/Private/K2Node_VariableGet.cpp
+++ b/Engine/Source/Editor/BlueprintGraph/Private/K2Node_VariableGet.cpp
@@ -87,6 +87,11 @@ void UK2Node_VariableGet::ReallocatePinsDuringReconstruction(TArray<UEdGraphPin*
 
 FText UK2Node_VariableGet::GetTooltipText() const
 {
+	if (!CachedTooltip.IsOutOfDate())
+	{
+		return CachedTooltip;
+	}
+
 	FFormatNamedArguments Args;
 	Args.Add( TEXT( "VarName" ), FText::FromString( GetVarNameString() ));
 	Args.Add( TEXT( "TextPartition" ), FText::GetEmpty());
@@ -123,7 +128,9 @@ FText UK2Node_VariableGet::GetTooltipText() const
 			}
 		}
 	}
-	return FText::Format( NSLOCTEXT( "K2Node", "GetVariable_ToolTip", "Read the value of variable {VarName}{TextPartition}{MetaData}"), Args );
+	// FText::Format() is slow, so we cache this to save on performance
+	CachedTooltip = FText::Format(NSLOCTEXT("K2Node", "GetVariable_ToolTip", "Read the value of variable {VarName}{TextPartition}{MetaData}"), Args);
+	return CachedTooltip;
 }
 
 FText UK2Node_VariableGet::GetNodeTitle(ENodeTitleType::Type TitleType) const
diff --git a/Engine/Source/Editor/BlueprintGraph/Private/K2Node_VariableSet.cpp b/Engine/Source/Editor/BlueprintGraph/Private/K2Node_VariableSet.cpp
--- a/Engine/Source/Editor/BlueprintGraph/Private/K2Node_VariableSet.cpp
+++ b/Engine/Source/Editor/BlueprintGraph/Private/K2Node_VariableSet.cpp
@@ -50,6 +50,11 @@ void UK2Node_VariableSet::ReallocatePinsDuringReconstruction(TArray<UEdGraphPin*
 
 FText UK2Node_VariableSet::GetTooltipText() const
 {
+	if (!CachedTooltip.IsOutOfDate())
+	{
+		return CachedTooltip;
+	}
+
 	FFormatNamedArguments Args;
 	Args.Add( TEXT( "VarName" ), FText::FromString( GetVarNameString() ));
 	Args.Add( TEXT( "ReplicationCall" ), FText::GetEmpty());
@@ -81,7 +86,9 @@ FText UK2Node_VariableSet::GetTooltipText() const
 			}
 		}
 	}
-	return FText::Format( NSLOCTEXT( "K2Node", "SetValueOfVariable", "Set the value of variable {VarName}{ReplicationCall}{ReplicationNotifyName}{TextPartition}{MetaData}"), Args );
+	// FText::Format() is slow, so we cache this to save on performance
+	CachedTooltip = FText::Format(NSLOCTEXT("K2Node", "SetValueOfVariable", "Set the value of variable {VarName}{ReplicationCall}{ReplicationNotifyName}{TextPartition}{MetaData}"), Args);
+	return CachedTooltip;
 }
 
 FText UK2Node_VariableSet::GetNodeTitle(ENodeTitleType::Type TitleType) const
diff --git a/Engine/Source/Runtime/Engine/Classes/EdGraph/EdGraphNode_Comment.h b/Engine/Source/Runtime/Engine/Classes/EdGraph/EdGraphNode_Comment.h
--- a/Engine/Source/Runtime/Engine/Classes/EdGraph/EdGraphNode_Comment.h
+++ b/Engine/Source/Runtime/Engine/Classes/EdGraph/EdGraphNode_Comment.h
@@ -2,6 +2,7 @@
 
 #pragma once
 #include "EdGraph/EdGraphNode.h"
+#include "EdGraph/EdGraphNodeUtils.h" // for FNodeTextCache
 #include "EdGraphNode_Comment.generated.h"
 
 typedef TArray<class UObject*> FCommentNodeSet;
@@ -74,6 +75,9 @@ public:
 private:
 	/** Nodes currently within the region of the comment */
 	FCommentNodeSet	NodesUnderComment;
+
+	/** Constructing FText strings can be costly, so we cache the node's tooltip */
+	FNodeTextCache CachedTooltip;
 #endif
 };
 
diff --git a/Engine/Source/Runtime/Engine/Private/EdGraph/EdGraphNode_Comment.cpp b/Engine/Source/Runtime/Engine/Private/EdGraph/EdGraphNode_Comment.cpp
--- a/Engine/Source/Runtime/Engine/Private/EdGraph/EdGraphNode_Comment.cpp
+++ b/Engine/Source/Runtime/Engine/Private/EdGraph/EdGraphNode_Comment.cpp
@@ -57,7 +57,11 @@ void UEdGraphNode_Comment::PostPlacedNewNode()
 
 FText UEdGraphNode_Comment::GetTooltipText() const
 {
-	return FText::Format(NSLOCTEXT("K2Node", "CommentBlock_Tooltip", "Comment:\n{0}"), FText::FromString(NodeComment));
+	if (CachedTooltip.IsOutOfDate())
+	{
+		CachedTooltip = FText::Format(NSLOCTEXT("K2Node", "CommentBlock_Tooltip", "Comment:\n{0}"), FText::FromString(NodeComment));
+	}
+	return CachedTooltip;
 }
 
 FString UEdGraphNode_Comment::GetDocumentationLink() const
@@ -137,6 +141,7 @@ const FCommentNodeSet& UEdGraphNode_Comment::GetNodesUnderComment() const
 void UEdGraphNode_Comment::OnRenameNode(const FString& NewName)
 {
 	NodeComment = NewName;
+	CachedTooltip.MarkDirty();
 }
 
 TSharedPtr<class INameValidatorInterface> UEdGraphNode_Comment::MakeNameValidator() const
diff --git a/Engine/Source/Runtime/Engine/Private/EdGraph/EdGraphNode_Documentation.cpp b/Engine/Source/Runtime/Engine/Private/EdGraph/EdGraphNode_Documentation.cpp
--- a/Engine/Source/Runtime/Engine/Private/EdGraph/EdGraphNode_Documentation.cpp
+++ b/Engine/Source/Runtime/Engine/Private/EdGraph/EdGraphNode_Documentation.cpp
@@ -32,7 +32,7 @@ void UEdGraphNode_Documentation::PostPlacedNewNode()
 
 FText UEdGraphNode_Documentation::GetTooltipText() const
 {
-	return FText::Format(NSLOCTEXT("K2Node", "DocumentationBlock_Tooltip", "Documentation:\n{0}"), FText::FromString(NodeComment));
+	return NSLOCTEXT("K2Node", "DocumentationBlock_Tooltip", "UDN Documentation Excerpt");
 }
 
 FText UEdGraphNode_Documentation::GetNodeTitle( ENodeTitleType::Type TitleType ) const
diff --git a/Engine/Source/Runtime/Engine/Public/EdGraph/EdGraphNodeUtils.h b/Engine/Source/Runtime/Engine/Public/EdGraph/EdGraphNodeUtils.h
--- a/Engine/Source/Runtime/Engine/Public/EdGraph/EdGraphNodeUtils.h
+++ b/Engine/Source/Runtime/Engine/Public/EdGraph/EdGraphNodeUtils.h
@@ -83,7 +83,7 @@ public:
 	}
 
 	/** */
-	FORCEINLINE FText GetCachedTitle(ENodeTitleType::Type TitleType) const
+	FORCEINLINE FText& GetCachedTitle(ENodeTitleType::Type TitleType) const
 	{
 		return CachedNodeTitles[TitleType].GetCachedText();
 	}
@@ -97,6 +97,12 @@ public:
 		}
 	}
 
+	/** */
+	FORCEINLINE FText& operator[](ENodeTitleType::Type TitleIndex) const
+	{
+		return GetCachedTitle(TitleIndex);
+	}
+
 private:
 	/** */
 	FNodeTextCache CachedNodeTitles[ENodeTitleType::MAX_TitleTypes];
