Caching more node titles for various blueprint nodes (to optimize the new BP menu system).

diff --git a/Engine/Plugins/Experimental/LiveEditor/Source/LiveEditor/Classes/K2Node_LiveEditObject.h b/Engine/Plugins/Experimental/LiveEditor/Source/LiveEditor/Classes/K2Node_LiveEditObject.h
--- a/Engine/Plugins/Experimental/LiveEditor/Source/LiveEditor/Classes/K2Node_LiveEditObject.h
+++ b/Engine/Plugins/Experimental/LiveEditor/Source/LiveEditor/Classes/K2Node_LiveEditObject.h
@@ -3,7 +3,7 @@
 #pragma once
 
 #include "LiveEditorTypes.h"
-
+#include "EdGraph/EdGraphNodeUtils.h" // for FNodeTextCache
 #include "K2Node_LiveEditObject.generated.h"
 
 UCLASS(MinimalAPI)
@@ -64,5 +64,8 @@ private:
 	 * @param   PinDescription	A string describing the pin's purpose
 	 */
 	void SetPinToolTip(UEdGraphPin& MutatablePin, const FText& PinDescription) const;
+
+	/** Constructing FText strings can be costly, so we cache the node's title */
+	FNodeTextCache CachedNodeTitle;
 #endif
 };
diff --git a/Engine/Plugins/Experimental/LiveEditor/Source/LiveEditor/Private/K2Node_LiveEditObject.cpp b/Engine/Plugins/Experimental/LiveEditor/Source/LiveEditor/Private/K2Node_LiveEditObject.cpp
--- a/Engine/Plugins/Experimental/LiveEditor/Source/LiveEditor/Private/K2Node_LiveEditObject.cpp
+++ b/Engine/Plugins/Experimental/LiveEditor/Source/LiveEditor/Private/K2Node_LiveEditObject.cpp
@@ -184,19 +184,21 @@ void UK2Node_LiveEditObject::AllocateDefaultPins()
 FText UK2Node_LiveEditObject::GetNodeTitle(ENodeTitleType::Type TitleType) const
 {
 	UEdGraphPin* BaseClassPin = GetBaseClassPin();
-
-	FText SpawnString = NSLOCTEXT("K2Node", "None", "NONE");
-	if(BaseClassPin != NULL && BaseClassPin->DefaultObject != NULL )
+	if ((BaseClassPin == nullptr) || (BaseClassPin->DefaultObject == nullptr))
 	{
-		SpawnString = FText::FromString(BaseClassPin->DefaultObject->GetName());
+		return NSLOCTEXT("K2Node", "LiveEditObject_NullTitle", "LiveEditObject NONE");
 	}
+	else if (CachedNodeTitle.IsOutOfDate())
+	{
+		FNumberFormattingOptions NumberOptions;
+		NumberOptions.UseGrouping = false;
+		FFormatNamedArguments Args;
+		Args.Add(TEXT("SpawnString"), FText::FromName(BaseClassPin->DefaultObject->GetFName()));
+		Args.Add(TEXT("ID"), FText::AsNumber(GetUniqueID(), &NumberOptions));
 
-	FNumberFormattingOptions NumberOptions;
-	NumberOptions.UseGrouping = false;
-	FFormatNamedArguments Args;
-	Args.Add(TEXT("SpawnString"), SpawnString);
-	Args.Add(TEXT("ID"), FText::AsNumber(GetUniqueID(), &NumberOptions));
-	return FText::Format(NSLOCTEXT("K2Node", "LiveEditObject", "LiveEditObject {SpawnString}_{ID}"), Args );
+		CachedNodeTitle = FText::Format(NSLOCTEXT("K2Node", "LiveEditObject", "LiveEditObject {SpawnString}_{ID}"), Args);
+	}
+	return CachedNodeTitle;
 }
 
 void UK2Node_LiveEditObject::PinDefaultValueChanged(UEdGraphPin* Pin) 
diff --git a/Engine/Source/Developer/MovieSceneTools/Classes/K2Node_PlayMovieScene.h b/Engine/Source/Developer/MovieSceneTools/Classes/K2Node_PlayMovieScene.h
--- a/Engine/Source/Developer/MovieSceneTools/Classes/K2Node_PlayMovieScene.h
+++ b/Engine/Source/Developer/MovieSceneTools/Classes/K2Node_PlayMovieScene.h
@@ -3,6 +3,7 @@
 #pragma once
 
 #include "BlueprintGraphDefinitions.h"		// @todo sequencer uobject: workaround for header generator dependencies
+#include "EdGraph/EdGraphNodeUtils.h" // for FNodeTextCache
 #include "K2Node_PlayMovieScene.generated.h"
 
 #if WITH_EDITOR
@@ -135,6 +136,10 @@ protected:
 #if WITH_EDITOR
 	/** Event that is broadcast when the bindings are changed */
 	FOnMovieSceneBindingsChanged OnBindingsChangedEvent;
+
+private:
+	/** Constructing FText strings can be costly, so we cache the node's title */
+	FNodeTextCache CachedNodeTitle;
 #endif
 };
 
diff --git a/Engine/Source/Developer/MovieSceneTools/Private/K2Node_PlayMovieScene.cpp b/Engine/Source/Developer/MovieSceneTools/Private/K2Node_PlayMovieScene.cpp
--- a/Engine/Source/Developer/MovieSceneTools/Private/K2Node_PlayMovieScene.cpp
+++ b/Engine/Source/Developer/MovieSceneTools/Private/K2Node_PlayMovieScene.cpp
@@ -77,14 +77,19 @@ FLinearColor UK2Node_PlayMovieScene::GetNodeTitleColor() const
 
 FText UK2Node_PlayMovieScene::GetNodeTitle(ENodeTitleType::Type TitleType) const
 {
-	UMovieScene* MovieScene = MovieSceneBindings != NULL ? MovieSceneBindings->GetRootMovieScene() : NULL;
-
-	FFormatNamedArguments Args;
-	Args.Add(TEXT("SceneName"), FText::FromString(MovieScene->GetName()));
-
-	return MovieScene != NULL ?
-		( FText::Format(NSLOCTEXT("PlayMovieSceneNode", "NodeTitle", "Play Movie Scene: {SceneName}"), Args) ) :
-		NSLOCTEXT("PlayMovieSceneNode", "NodeTitleWithNoMovieScene", "Play Movie Scene (No Asset)");
+	UMovieScene* MovieScene = (MovieSceneBindings != nullptr) ? MovieSceneBindings->GetRootMovieScene() : nullptr;
+	if (MovieScene == nullptr)
+	{
+		return NSLOCTEXT("PlayMovieSceneNode", "NodeTitleWithNoMovieScene", "Play Movie Scene (No Asset)");
+	}
+	else if (CachedNodeTitle.IsOutOfDate())
+	{
+		FFormatNamedArguments Args;
+		Args.Add(TEXT("SceneName"), FText::FromString(MovieScene->GetName()));
+		// FText::Format() is slow, so we cache this to save on performance
+		CachedNodeTitle = FText::Format(NSLOCTEXT("PlayMovieSceneNode", "NodeTitle", "Play Movie Scene: {SceneName}"), Args);
+	}
+	return CachedNodeTitle;
 }
 
 UEdGraphPin* UK2Node_PlayMovieScene::GetPlayPin() const
@@ -187,6 +192,7 @@ void UK2Node_PlayMovieScene::CreateBindingsIfNeeded()
 		check( NewMovieSceneBindings != NULL );
 
 		MovieSceneBindings = NewMovieSceneBindings;
+		CachedNodeTitle.MarkDirty();
 	}
 }
 
diff --git a/Engine/Source/Editor/AnimGraph/Classes/AnimGraphNode_BoneDrivenController.h b/Engine/Source/Editor/AnimGraph/Classes/AnimGraphNode_BoneDrivenController.h
--- a/Engine/Source/Editor/AnimGraph/Classes/AnimGraphNode_BoneDrivenController.h
+++ b/Engine/Source/Editor/AnimGraph/Classes/AnimGraphNode_BoneDrivenController.h
@@ -4,6 +4,7 @@
 
 #include "AnimGraphNode_SkeletalControlBase.h"
 #include "Animation/BoneControllers/AnimNode_BoneDrivenController.h"
+#include "EdGraph/EdGraphNodeUtils.h" // for FNodeTitleTextTable
 #include "AnimGraphNode_BoneDrivenController.generated.h"
 
 UCLASS(MinimalAPI)
@@ -31,4 +32,8 @@ protected:
 	// UAnimGraphNode_SkeletalControlBase interface
 	virtual FText GetControllerDescription() const override;
 	//////////////////////////////////////////////////////////////////////////
+
+private:
+	/** Constructing FText strings can be costly, so we cache the node's title */
+	FNodeTitleTextTable CachedNodeTitles;
 };
\ No newline at end of file
diff --git a/Engine/Source/Editor/AnimGraph/Classes/AnimGraphNode_CopyBone.h b/Engine/Source/Editor/AnimGraph/Classes/AnimGraphNode_CopyBone.h
--- a/Engine/Source/Editor/AnimGraph/Classes/AnimGraphNode_CopyBone.h
+++ b/Engine/Source/Editor/AnimGraph/Classes/AnimGraphNode_CopyBone.h
@@ -3,6 +3,7 @@
 #pragma once
 #include "AnimGraphNode_SkeletalControlBase.h"
 #include "Animation/BoneControllers/AnimNode_CopyBone.h"
+#include "EdGraph/EdGraphNodeUtils.h" // for FNodeTitleTextTable
 #include "AnimGraphNode_CopyBone.generated.h"
 
 UCLASS(MinimalAPI)
@@ -23,4 +24,8 @@ protected:
 	// UAnimGraphNode_SkeletalControlBase interface
 	virtual FText GetControllerDescription() const override;
 	// End of UAnimGraphNode_SkeletalControlBase interface
+
+private:
+	/** Constructing FText strings can be costly, so we cache the node's title */
+	FNodeTitleTextTable CachedNodeTitles;
 };
diff --git a/Engine/Source/Editor/AnimGraph/Classes/AnimGraphNode_LookAt.h b/Engine/Source/Editor/AnimGraph/Classes/AnimGraphNode_LookAt.h
--- a/Engine/Source/Editor/AnimGraph/Classes/AnimGraphNode_LookAt.h
+++ b/Engine/Source/Editor/AnimGraph/Classes/AnimGraphNode_LookAt.h
@@ -4,6 +4,7 @@
 
 #include "AnimGraphNode_SkeletalControlBase.h"
 #include "Animation/BoneControllers/AnimNode_LookAt.h"
+#include "EdGraph/EdGraphNodeUtils.h" // for FNodeTitleTextTable
 #include "AnimGraphNode_LookAt.generated.h"
 
 UCLASS(MinimalAPI)
@@ -29,4 +30,8 @@ protected:
 	// UAnimGraphNode_SkeletalControlBase interface
 	virtual FText GetControllerDescription() const override;
 	// End of UAnimGraphNode_SkeletalControlBase interface
+
+private:
+	/** Constructing FText strings can be costly, so we cache the node's title */
+	FNodeTitleTextTable CachedNodeTitles;
 };
diff --git a/Engine/Source/Editor/AnimGraph/Classes/AnimGraphNode_ModifyBone.h b/Engine/Source/Editor/AnimGraph/Classes/AnimGraphNode_ModifyBone.h
--- a/Engine/Source/Editor/AnimGraph/Classes/AnimGraphNode_ModifyBone.h
+++ b/Engine/Source/Editor/AnimGraph/Classes/AnimGraphNode_ModifyBone.h
@@ -4,6 +4,7 @@
 
 #include "AnimGraphNode_SkeletalControlBase.h"
 #include "Animation/BoneControllers/AnimNode_ModifyBone.h"
+#include "EdGraph/EdGraphNodeUtils.h" // for FNodeTitleTextTable
 #include "AnimGraphNode_ModifyBone.generated.h"
 
 UCLASS(MinimalAPI)
@@ -25,4 +26,9 @@ protected:
 	// UAnimGraphNode_SkeletalControlBase interface
 	virtual FText GetControllerDescription() const override;
 	// End of UAnimGraphNode_SkeletalControlBase interface
+
+private:
+	/** Constructing FText strings can be costly, so we cache the node's title */
+	FNodeTitleTextTable CachedNodeTitles;
 };
+
diff --git a/Engine/Source/Editor/AnimGraph/Classes/AnimGraphNode_RotationMultiplier.h b/Engine/Source/Editor/AnimGraph/Classes/AnimGraphNode_RotationMultiplier.h
--- a/Engine/Source/Editor/AnimGraph/Classes/AnimGraphNode_RotationMultiplier.h
+++ b/Engine/Source/Editor/AnimGraph/Classes/AnimGraphNode_RotationMultiplier.h
@@ -3,6 +3,7 @@
 #pragma once
 #include "AnimGraphNode_SkeletalControlBase.h"
 #include "Animation/BoneControllers/AnimNode_RotationMultiplier.h"
+#include "EdGraph/EdGraphNodeUtils.h" // for FNodeTitleTextTable
 #include "AnimGraphNode_RotationMultiplier.generated.h"
 
 UCLASS(MinimalAPI)
@@ -23,4 +24,8 @@ protected:
 	// UAnimGraphNode_SkeletalControlBase interface
 	virtual FText GetControllerDescription() const override;
 	// End of UAnimGraphNode_SkeletalControlBase interface
+
+private:
+	/** Constructing FText strings can be costly, so we cache the node's title */
+	FNodeTitleTextTable CachedNodeTitles;
 };
diff --git a/Engine/Source/Editor/AnimGraph/Classes/AnimGraphNode_SaveCachedPose.h b/Engine/Source/Editor/AnimGraph/Classes/AnimGraphNode_SaveCachedPose.h
--- a/Engine/Source/Editor/AnimGraph/Classes/AnimGraphNode_SaveCachedPose.h
+++ b/Engine/Source/Editor/AnimGraph/Classes/AnimGraphNode_SaveCachedPose.h
@@ -3,6 +3,7 @@
 #pragma once
 #include "AnimGraphNode_Base.h"
 #include "Animation/AnimNode_SaveCachedPose.h"
+#include "EdGraph/EdGraphNodeUtils.h" // for FNodeTextCache
 #include "AnimGraphNode_SaveCachedPose.generated.h"
 
 UCLASS(MinimalAPI)
@@ -33,4 +34,8 @@ class UAnimGraphNode_SaveCachedPose : public UAnimGraphNode_Base
 	virtual void GetMenuEntries(FGraphContextMenuBuilder& ContextMenuBuilder) const override;
 	virtual bool IsSinkNode() const override { return true; }
 	// End of UAnimGraphNode_Base interface
+
+private:
+	/** Constructing FText strings can be costly, so we cache the node's title */
+	FNodeTextCache CachedNodeTitle;
 };
diff --git a/Engine/Source/Editor/AnimGraph/Classes/AnimGraphNode_SequencePlayer.h b/Engine/Source/Editor/AnimGraph/Classes/AnimGraphNode_SequencePlayer.h
--- a/Engine/Source/Editor/AnimGraph/Classes/AnimGraphNode_SequencePlayer.h
+++ b/Engine/Source/Editor/AnimGraph/Classes/AnimGraphNode_SequencePlayer.h
@@ -3,6 +3,7 @@
 #pragma once
 #include "AnimGraphNode_Base.h"
 #include "Animation/AnimNode_SequencePlayer.h"
+#include "EdGraph/EdGraphNodeUtils.h" // for FNodeTitleTextTable
 #include "AnimGraphNode_SequencePlayer.generated.h"
 
 UCLASS(MinimalAPI)
@@ -42,4 +43,7 @@ class UAnimGraphNode_SequencePlayer : public UAnimGraphNode_Base
 
 private:
 	static FText GetTitleGivenAssetInfo(const FText& AssetName, bool bKnownToBeAdditive);
+
+	/** Constructing FText strings can be costly, so we cache the node's title */
+	FNodeTitleTextTable CachedNodeTitles;
 };
diff --git a/Engine/Source/Editor/AnimGraph/Classes/AnimGraphNode_Slot.h b/Engine/Source/Editor/AnimGraph/Classes/AnimGraphNode_Slot.h
--- a/Engine/Source/Editor/AnimGraph/Classes/AnimGraphNode_Slot.h
+++ b/Engine/Source/Editor/AnimGraph/Classes/AnimGraphNode_Slot.h
@@ -3,6 +3,7 @@
 #pragma once
 #include "AnimGraphNode_Base.h"
 #include "Animation/AnimNode_Slot.h"
+#include "EdGraph/EdGraphNodeUtils.h" // for FNodeTitleTextTable
 #include "AnimGraphNode_Slot.generated.h"
 
 UCLASS(MinimalAPI)
@@ -23,4 +24,8 @@ class UAnimGraphNode_Slot : public UAnimGraphNode_Base
 	virtual FString GetNodeCategory() const override;
 	virtual void BakeDataDuringCompilation(class FCompilerResultsLog& MessageLog) override;
 	// End of UAnimGraphNode_Base interface
+
+private:
+	/** Constructing FText strings can be costly, so we cache the node's title */
+	FNodeTitleTextTable CachedNodeTitles;
 };
diff --git a/Engine/Source/Editor/AnimGraph/Classes/AnimGraphNode_SpringBone.h b/Engine/Source/Editor/AnimGraph/Classes/AnimGraphNode_SpringBone.h
--- a/Engine/Source/Editor/AnimGraph/Classes/AnimGraphNode_SpringBone.h
+++ b/Engine/Source/Editor/AnimGraph/Classes/AnimGraphNode_SpringBone.h
@@ -3,6 +3,7 @@
 #pragma once
 #include "AnimGraphNode_SkeletalControlBase.h"
 #include "Animation/BoneControllers/AnimNode_SpringBone.h"
+#include "EdGraph/EdGraphNodeUtils.h" // for FNodeTitleTextTable
 #include "AnimGraphNode_SpringBone.generated.h"
 
 UCLASS(MinimalAPI)
@@ -23,4 +24,8 @@ protected:
 	// UAnimGraphNode_SkeletalControlBase interface
 	virtual FText GetControllerDescription() const override;
 	// End of UAnimGraphNode_SkeletalControlBase interface
+
+private:
+	/** Constructing FText strings can be costly, so we cache the node's title */
+	FNodeTitleTextTable CachedNodeTitles;
 };
diff --git a/Engine/Source/Editor/AnimGraph/Classes/AnimGraphNode_StateMachineBase.h b/Engine/Source/Editor/AnimGraph/Classes/AnimGraphNode_StateMachineBase.h
--- a/Engine/Source/Editor/AnimGraph/Classes/AnimGraphNode_StateMachineBase.h
+++ b/Engine/Source/Editor/AnimGraph/Classes/AnimGraphNode_StateMachineBase.h
@@ -3,6 +3,7 @@
 #pragma once
 #include "AnimGraphNode_Base.h"
 #include "Animation/AnimNode_StateMachine.h"
+#include "EdGraph/EdGraphNodeUtils.h" // for FNodeTextCache
 #include "AnimGraphNode_StateMachineBase.generated.h"
 
 UCLASS(Abstract)
@@ -38,4 +39,8 @@ class ANIMGRAPH_API UAnimGraphNode_StateMachineBase : public UAnimGraphNode_Base
 	// Interface for derived classes to implement
 	virtual FAnimNode_StateMachine& GetNode() PURE_VIRTUAL(UAnimGraphNode_StateMachineBase::GetNode, static FAnimNode_StateMachine Dummy; return Dummy;);
 	// End of my interface
+
+private:
+	/** Constructing FText strings can be costly, so we cache the node's title */
+	FNodeTextCache CachedFullTitle;
 };
diff --git a/Engine/Source/Editor/AnimGraph/Classes/AnimGraphNode_Trail.h b/Engine/Source/Editor/AnimGraph/Classes/AnimGraphNode_Trail.h
--- a/Engine/Source/Editor/AnimGraph/Classes/AnimGraphNode_Trail.h
+++ b/Engine/Source/Editor/AnimGraph/Classes/AnimGraphNode_Trail.h
@@ -3,6 +3,7 @@
 #pragma once
 #include "AnimGraphNode_SkeletalControlBase.h"
 #include "Animation/BoneControllers/AnimNode_Trail.h"
+#include "EdGraph/EdGraphNodeUtils.h" // for FNodeTitleTextTable
 #include "AnimGraphNode_Trail.generated.h"
 
 UCLASS(MinimalAPI)
@@ -24,4 +25,8 @@ protected:
 	// UAnimGraphNode_SkeletalControlBase interface
 	virtual FText GetControllerDescription() const override;
 	// End of UAnimGraphNode_SkeletalControlBase interface
+
+private:
+	/** Constructing FText strings can be costly, so we cache the node's title */
+	FNodeTitleTextTable CachedNodeTitles;
 };
diff --git a/Engine/Source/Editor/AnimGraph/Classes/AnimGraphNode_TwoBoneIK.h b/Engine/Source/Editor/AnimGraph/Classes/AnimGraphNode_TwoBoneIK.h
--- a/Engine/Source/Editor/AnimGraph/Classes/AnimGraphNode_TwoBoneIK.h
+++ b/Engine/Source/Editor/AnimGraph/Classes/AnimGraphNode_TwoBoneIK.h
@@ -3,6 +3,7 @@
 #pragma once
 #include "AnimGraphNode_SkeletalControlBase.h"
 #include "Animation/BoneControllers/AnimNode_TwoBoneIK.h"
+#include "EdGraph/EdGraphNodeUtils.h" // for FNodeTitleTextTable
 #include "AnimGraphNode_TwoBoneIK.generated.h"
 
 UCLASS(MinimalAPI)
@@ -31,4 +32,8 @@ protected:
 	// local conversion function for drawing
 	void ConvertToComponentSpaceTransform(USkeletalMeshComponent* SkelComp, USkeleton * Skeleton, const FTransform & InTransform, FTransform & OutCSTransform, int32 BoneIndex, uint8 Space) const;
 	void DrawTargetLocation(FPrimitiveDrawInterface* PDI, USkeletalMeshComponent* SkelComp, USkeleton * Skeleton, uint8 SpaceBase, FName SpaceBoneName, const FVector & TargetLocation, const FColor & TargetColor, const FColor & BoneColor) const;
+
+private:
+	/** Constructing FText strings can be costly, so we cache the node's title */
+	FNodeTitleTextTable CachedNodeTitles;
 };
diff --git a/Engine/Source/Editor/AnimGraph/Classes/AnimGraphNode_UseCachedPose.h b/Engine/Source/Editor/AnimGraph/Classes/AnimGraphNode_UseCachedPose.h
--- a/Engine/Source/Editor/AnimGraph/Classes/AnimGraphNode_UseCachedPose.h
+++ b/Engine/Source/Editor/AnimGraph/Classes/AnimGraphNode_UseCachedPose.h
@@ -3,6 +3,7 @@
 #pragma once
 #include "AnimGraphNode_Base.h"
 #include "Animation/AnimNode_UseCachedPose.h"
+#include "EdGraph/EdGraphNodeUtils.h" // for FNodeTextCache
 #include "AnimGraphNode_UseCachedPose.generated.h"
 
 UCLASS(MinimalAPI)
@@ -27,4 +28,8 @@ public:
 	virtual FString GetNodeCategory() const override;
 	virtual void GetMenuEntries(FGraphContextMenuBuilder& ContextMenuBuilder) const override;
 	// End of UAnimGraphNode_Base interface
+
+private:
+	/** Constructing FText strings can be costly, so we cache the node's title */
+	FNodeTextCache CachedNodeTitle;
 };
diff --git a/Engine/Source/Editor/AnimGraph/Private/AnimGraphNode_BlendListByEnum.cpp b/Engine/Source/Editor/AnimGraph/Private/AnimGraphNode_BlendListByEnum.cpp
--- a/Engine/Source/Editor/AnimGraph/Private/AnimGraphNode_BlendListByEnum.cpp
+++ b/Engine/Source/Editor/AnimGraph/Private/AnimGraphNode_BlendListByEnum.cpp
@@ -40,6 +40,7 @@ FText UAnimGraphNode_BlendListByEnum::GetNodeTitle(ENodeTitleType::Type TitleTyp
 	{
 		FFormatNamedArguments Args;
 		Args.Add(TEXT("EnumName"), FText::FromString(BoundEnum->GetName()));
+		// FText::Format() is slow, so we cache this to save on performance
 		CachedNodeTitle = FText::Format(LOCTEXT("AnimGraphNode_BlendListByEnum_Title", "Blend Poses ({EnumName})"), Args);
 	}
 	return CachedNodeTitle;
diff --git a/Engine/Source/Editor/AnimGraph/Private/AnimGraphNode_BoneDrivenController.cpp b/Engine/Source/Editor/AnimGraph/Private/AnimGraphNode_BoneDrivenController.cpp
--- a/Engine/Source/Editor/AnimGraph/Private/AnimGraphNode_BoneDrivenController.cpp
+++ b/Engine/Source/Editor/AnimGraph/Private/AnimGraphNode_BoneDrivenController.cpp
@@ -18,32 +18,30 @@ FText UAnimGraphNode_BoneDrivenController::GetTooltipText() const
 
 FText UAnimGraphNode_BoneDrivenController::GetNodeTitle(ENodeTitleType::Type TitleType) const
 {
-	FFormatNamedArguments Args;
-	Args.Add(TEXT("ControllerDesc"), GetControllerDescription());
-	Args.Add(TEXT("SourceBone"), FText::FromName(Node.SourceBone.BoneName));
-	Args.Add(TEXT("TargetBone"), FText::FromName(Node.TargetBone.BoneName));
-
-	FText NodeTitle;
-	if (TitleType == ENodeTitleType::ListView)
+	if ((Node.SourceBone.BoneName == NAME_None) && (Node.TargetBone.BoneName == NAME_None) && (TitleType == ENodeTitleType::ListView))
+	{
+		return GetControllerDescription();
+	}
+	else if (!CachedNodeTitles.IsTitleCached(TitleType))
 	{
-		Args.Add(TEXT("Delim"), FText::FromString(TEXT(" - ")));
+		FFormatNamedArguments Args;
+		Args.Add(TEXT("ControllerDesc"), GetControllerDescription());
+		Args.Add(TEXT("SourceBone"), FText::FromName(Node.SourceBone.BoneName));
+		Args.Add(TEXT("TargetBone"), FText::FromName(Node.TargetBone.BoneName));
 
-		if ((Node.SourceBone.BoneName == NAME_None) && (Node.TargetBone.BoneName == NAME_None))
+		if (TitleType == ENodeTitleType::ListView)
 		{
-			NodeTitle = FText::Format(LOCTEXT("AnimGraphNode_BoneDrivenController_MenuTitle", "{ControllerDesc}"), Args);
+			Args.Add(TEXT("Delim"), FText::FromString(TEXT(" - ")));
+		}
+		else
+		{
+			Args.Add(TEXT("Delim"), FText::FromString(TEXT("\n")));
 		}
-	}
-	else
-	{
-		Args.Add(TEXT("Delim"), FText::FromString(TEXT("\n")));
-	}
 
-	if (NodeTitle.IsEmpty())
-	{
-		NodeTitle = FText::Format(LOCTEXT("AnimGraphNode_BoneDrivenController_Title", "{ControllerDesc}{Delim}Driving Bone: {SourceBone}{Delim}Driven Bone: {TargetBone}"), Args);
-	}
-	
-	return NodeTitle;
+		// FText::Format() is slow, so we cache this to save on performance
+		CachedNodeTitles.SetCachedTitle(TitleType, FText::Format(LOCTEXT("AnimGraphNode_BoneDrivenController_Title", "{ControllerDesc}{Delim}Driving Bone: {SourceBone}{Delim}Driven Bone: {TargetBone}"), Args));
+	}	
+	return CachedNodeTitles[TitleType];
 }
 
 FText UAnimGraphNode_BoneDrivenController::GetControllerDescription() const
diff --git a/Engine/Source/Editor/AnimGraph/Private/AnimGraphNode_CopyBone.cpp b/Engine/Source/Editor/AnimGraph/Private/AnimGraphNode_CopyBone.cpp
--- a/Engine/Source/Editor/AnimGraph/Private/AnimGraphNode_CopyBone.cpp
+++ b/Engine/Source/Editor/AnimGraph/Private/AnimGraphNode_CopyBone.cpp
@@ -25,28 +25,27 @@ FText UAnimGraphNode_CopyBone::GetTooltipText() const
 
 FText UAnimGraphNode_CopyBone::GetNodeTitle(ENodeTitleType::Type TitleType) const
 {
-	FFormatNamedArguments Args;
-	Args.Add(TEXT("ControllerDescription"), GetControllerDescription());
-	Args.Add(TEXT("SourceBoneName"), FText::FromName(Node.SourceBone.BoneName));
-	Args.Add(TEXT("TargetBoneName"), FText::FromName(Node.TargetBone.BoneName));
-
-	FText NodeTitle;
-	if (TitleType == ENodeTitleType::ListView)
+	if ((TitleType == ENodeTitleType::ListView) && (Node.TargetBone.BoneName == NAME_None) && (Node.SourceBone.BoneName == NAME_None))
+	{
+		return GetControllerDescription();
+	}
+	else if (!CachedNodeTitles.IsTitleCached(TitleType))
 	{
-		if ((Node.SourceBone.BoneName == NAME_None) && (Node.TargetBone.BoneName == NAME_None))
+		FFormatNamedArguments Args;
+		Args.Add(TEXT("ControllerDescription"), GetControllerDescription());
+		Args.Add(TEXT("SourceBoneName"), FText::FromName(Node.SourceBone.BoneName));
+		Args.Add(TEXT("TargetBoneName"), FText::FromName(Node.TargetBone.BoneName));
+
+		if (TitleType == ENodeTitleType::ListView)
 		{
-			NodeTitle = FText::Format(LOCTEXT("AnimGraphNode_CopyBone_MenuTitle", "{ControllerDescription}"), Args);
+			CachedNodeTitles.SetCachedTitle(TitleType, FText::Format(LOCTEXT("AnimGraphNode_CopyBone_ListTitle", "{ControllerDescription} - Source Bone: {SourceBoneName} - Target Bone: {TargetBoneName}"), Args));
 		}
 		else
 		{
-			NodeTitle = FText::Format(LOCTEXT("AnimGraphNode_CopyBone_ListTitle", "{ControllerDescription} - Source Bone: {SourceBoneName} - Target Bone: {TargetBoneName}"), Args);
+			CachedNodeTitles.SetCachedTitle(TitleType, FText::Format(LOCTEXT("AnimGraphNode_CopyBone_Title", "{ControllerDescription}\nSource Bone: {SourceBoneName}\nTarget Bone: {TargetBoneName}"), Args));
 		}
 	}
-	else
-	{
-		NodeTitle = FText::Format(LOCTEXT("AnimGraphNode_CopyBone_Title", "{ControllerDescription}\nSource Bone: {SourceBoneName}\nTarget Bone: {TargetBoneName}"), Args);
-	}
-	return NodeTitle;
+	return CachedNodeTitles[TitleType];
 }
 
 #undef LOCTEXT_NAMESPACE
diff --git a/Engine/Source/Editor/AnimGraph/Private/AnimGraphNode_LookAt.cpp b/Engine/Source/Editor/AnimGraph/Private/AnimGraphNode_LookAt.cpp
--- a/Engine/Source/Editor/AnimGraph/Private/AnimGraphNode_LookAt.cpp
+++ b/Engine/Source/Editor/AnimGraph/Private/AnimGraphNode_LookAt.cpp
@@ -30,22 +30,27 @@ FText UAnimGraphNode_LookAt::GetTooltipText() const
 
 FText UAnimGraphNode_LookAt::GetNodeTitle(ENodeTitleType::Type TitleType) const
 {
-	FFormatNamedArguments Args;
-	Args.Add(TEXT("ControllerDescription"), GetControllerDescription());
-	Args.Add(TEXT("BoneName"), FText::FromName(Node.BoneToModify.BoneName));
-
-	if(TitleType == ENodeTitleType::ListView)
+	if ((TitleType == ENodeTitleType::ListView) && (Node.BoneToModify.BoneName == NAME_None))
 	{
-		if (Node.BoneToModify.BoneName == NAME_None)
-		{
-			return FText::Format(LOCTEXT("AnimGraphNode_LookAt_MenuTitle", "{ControllerDescription}"), Args);
-		}
-		return FText::Format(LOCTEXT("AnimGraphNode_LookAt_ListTitle", "{ControllerDescription} - Bone: {BoneName}"), Args);
+		return GetControllerDescription();
 	}
-	else
+	else if (!CachedNodeTitles.IsTitleCached(TitleType))
 	{
-		return FText::Format(LOCTEXT("AnimGraphNode_LookAt_Title", "{ControllerDescription}\nBone: {BoneName}"), Args);
+		FFormatNamedArguments Args;
+		Args.Add(TEXT("ControllerDescription"), GetControllerDescription());
+		Args.Add(TEXT("BoneName"), FText::FromName(Node.BoneToModify.BoneName));
+
+		// FText::Format() is slow, so we cache this to save on performance
+		if (TitleType == ENodeTitleType::ListView)
+		{
+			CachedNodeTitles.SetCachedTitle(TitleType, FText::Format(LOCTEXT("AnimGraphNode_LookAt_ListTitle", "{ControllerDescription} - Bone: {BoneName}"), Args));
+		}
+		else
+		{
+			CachedNodeTitles.SetCachedTitle(TitleType, FText::Format(LOCTEXT("AnimGraphNode_LookAt_Title", "{ControllerDescription}\nBone: {BoneName}"), Args));
+		}
 	}
+	return CachedNodeTitles[TitleType];		
 }
 
 void UAnimGraphNode_LookAt::Draw(FPrimitiveDrawInterface* PDI, USkeletalMeshComponent* SkelMeshComp) const
diff --git a/Engine/Source/Editor/AnimGraph/Private/AnimGraphNode_ModifyBone.cpp b/Engine/Source/Editor/AnimGraph/Private/AnimGraphNode_ModifyBone.cpp
--- a/Engine/Source/Editor/AnimGraph/Private/AnimGraphNode_ModifyBone.cpp
+++ b/Engine/Source/Editor/AnimGraph/Private/AnimGraphNode_ModifyBone.cpp
@@ -30,22 +30,27 @@ FText UAnimGraphNode_ModifyBone::GetTooltipText() const
 
 FText UAnimGraphNode_ModifyBone::GetNodeTitle(ENodeTitleType::Type TitleType) const
 {
-	FFormatNamedArguments Args;
-	Args.Add(TEXT("ControllerDescription"), GetControllerDescription());
-	Args.Add(TEXT("BoneName"), FText::FromName(Node.BoneToModify.BoneName));
-
-	if(TitleType == ENodeTitleType::ListView)
+	if ((TitleType == ENodeTitleType::ListView) && (Node.BoneToModify.BoneName == NAME_None))
 	{
-		if (Node.BoneToModify.BoneName == NAME_None)
-		{
-			return FText::Format(LOCTEXT("AnimGraphNode_ModifyBone_MenuTitle", "{ControllerDescription}"), Args);
-		}
-		return FText::Format(LOCTEXT("AnimGraphNode_ModifyBone_ListTitle", "{ControllerDescription} - Bone: {BoneName}"), Args);
+		return GetControllerDescription();
 	}
-	else
+	else if (!CachedNodeTitles.IsTitleCached(TitleType))
 	{
-		return FText::Format(LOCTEXT("AnimGraphNode_ModifyBone_Title", "{ControllerDescription}\nBone: {BoneName}"), Args);
+		FFormatNamedArguments Args;
+		Args.Add(TEXT("ControllerDescription"), GetControllerDescription());
+		Args.Add(TEXT("BoneName"), FText::FromName(Node.BoneToModify.BoneName));
+
+		// FText::Format() is slow, so we cache this to save on performance
+		if (TitleType == ENodeTitleType::ListView)
+		{
+			CachedNodeTitles.SetCachedTitle(TitleType, FText::Format(LOCTEXT("AnimGraphNode_ModifyBone_ListTitle", "{ControllerDescription} - Bone: {BoneName}"), Args));
+		}
+		else
+		{
+			CachedNodeTitles.SetCachedTitle(TitleType, FText::Format(LOCTEXT("AnimGraphNode_ModifyBone_Title", "{ControllerDescription}\nBone: {BoneName}"), Args));
+		}
 	}
+	return CachedNodeTitles[TitleType];
 }
 
 #undef LOCTEXT_NAMESPACE
diff --git a/Engine/Source/Editor/AnimGraph/Private/AnimGraphNode_RotationMultiplier.cpp b/Engine/Source/Editor/AnimGraph/Private/AnimGraphNode_RotationMultiplier.cpp
--- a/Engine/Source/Editor/AnimGraph/Private/AnimGraphNode_RotationMultiplier.cpp
+++ b/Engine/Source/Editor/AnimGraph/Private/AnimGraphNode_RotationMultiplier.cpp
@@ -24,22 +24,27 @@ FText UAnimGraphNode_RotationMultiplier::GetTooltipText() const
 
 FText UAnimGraphNode_RotationMultiplier::GetNodeTitle(ENodeTitleType::Type TitleType) const
 {
-	FFormatNamedArguments Args;
-	Args.Add(TEXT("ControllerDescription"), GetControllerDescription());
-	Args.Add(TEXT("BoneName"), FText::FromName(Node.TargetBone.BoneName));
-
-	if(TitleType == ENodeTitleType::ListView)
+	if ((TitleType == ENodeTitleType::ListView) && (Node.TargetBone.BoneName == NAME_None))
 	{
-		if (Node.TargetBone.BoneName == NAME_None)
-		{
-			return FText::Format(LOCTEXT("AnimGraphNode_ModifyBone_MenuTitle", "{ControllerDescription}"), Args);
-		}
-		return FText::Format(LOCTEXT("AnimGraphNode_ModifyBone_ListTitle", "{ControllerDescription} - Bone: {BoneName}"), Args);
+		return GetControllerDescription();
 	}
-	else
+	else if (!CachedNodeTitles.IsTitleCached(TitleType))
 	{
-		return FText::Format(LOCTEXT("AnimGraphNode_ModifyBone_Title", "{ControllerDescription}\nBone: {BoneName}"), Args);
+		FFormatNamedArguments Args;
+		Args.Add(TEXT("ControllerDescription"), GetControllerDescription());
+		Args.Add(TEXT("BoneName"), FText::FromName(Node.TargetBone.BoneName));
+
+		// FText::Format() is slow, so we cache this to save on performance
+		if (TitleType == ENodeTitleType::ListView)
+		{
+			CachedNodeTitles.SetCachedTitle(TitleType, FText::Format(LOCTEXT("AnimGraphNode_ModifyBone_ListTitle", "{ControllerDescription} - Bone: {BoneName}"), Args));
+		}
+		else
+		{
+			CachedNodeTitles.SetCachedTitle(TitleType, FText::Format(LOCTEXT("AnimGraphNode_ModifyBone_Title", "{ControllerDescription}\nBone: {BoneName}"), Args));
+		}
 	}
+	return CachedNodeTitles[TitleType];
 }
 
 #undef LOCTEXT_NAMESPACE
diff --git a/Engine/Source/Editor/AnimGraph/Private/AnimGraphNode_SaveCachedPose.cpp b/Engine/Source/Editor/AnimGraph/Private/AnimGraphNode_SaveCachedPose.cpp
--- a/Engine/Source/Editor/AnimGraph/Private/AnimGraphNode_SaveCachedPose.cpp
+++ b/Engine/Source/Editor/AnimGraph/Private/AnimGraphNode_SaveCachedPose.cpp
@@ -55,12 +55,13 @@ FText UAnimGraphNode_SaveCachedPose::GetNodeTitle(ENodeTitleType::Type TitleType
 	{
 		return LOCTEXT("NewSaveCachedPose", "New Save cached pose...");
 	}
-	else
+	else if (CachedNodeTitle.IsOutOfDate())
 	{
 		FFormatNamedArguments Args;
 		Args.Add(TEXT("NodeTitle"), FText::FromString(CacheName));
-		return FText::Format(LOCTEXT("AnimGraphNode_SaveCachedPose_Title", "Save cached pose '{NodeTitle}'"), Args);
+		CachedNodeTitle = FText::Format(LOCTEXT("AnimGraphNode_SaveCachedPose_Title", "Save cached pose '{NodeTitle}'"), Args);
 	}
+	return CachedNodeTitle;
 }
 
 FString UAnimGraphNode_SaveCachedPose::GetNodeCategory() const
diff --git a/Engine/Source/Editor/AnimGraph/Private/AnimGraphNode_SequencePlayer.cpp b/Engine/Source/Editor/AnimGraph/Private/AnimGraphNode_SequencePlayer.cpp
--- a/Engine/Source/Editor/AnimGraph/Private/AnimGraphNode_SequencePlayer.cpp
+++ b/Engine/Source/Editor/AnimGraph/Private/AnimGraphNode_SequencePlayer.cpp
@@ -77,19 +77,34 @@ FText UAnimGraphNode_SequencePlayer::GetTooltipText() const
 
 FText UAnimGraphNode_SequencePlayer::GetNodeTitle(ENodeTitleType::Type TitleType) const
 {
-	const bool bAdditive = ((Node.Sequence != NULL) && Node.Sequence->IsValidAdditive());
-	FText Title = GetTitleGivenAssetInfo((Node.Sequence != NULL) ? FText::FromString(Node.Sequence->GetName()) : LOCTEXT("None", "(None)"), bAdditive);
-
-	if ((TitleType == ENodeTitleType::FullTitle) && (SyncGroup.GroupName != NAME_None))
+	if (Node.Sequence == nullptr)
 	{
-		FFormatNamedArguments Args;
-		Args.Add(TEXT("Title"), Title);
-		Args.Add(TEXT("SyncGroup"), FText::FromName(SyncGroup.GroupName));
-
-		Title = FText::Format(LOCTEXT("SequenceNodeGroupWithSubtitle", "{Title}\nSync group {SyncGroup}"), Args);
+		return LOCTEXT("SequenceNullTitle", "Play (None)");
 	}
+	else if (!CachedNodeTitles.IsTitleCached(TitleType))
+	{
+		if (!CachedNodeTitles.IsTitleCached(ENodeTitleType::ListView))
+		{
+			const bool bAdditive = Node.Sequence->IsValidAdditive();
+			// FText::Format() is slow, so we cache this to save on performance
+			CachedNodeTitles.SetCachedTitle(ENodeTitleType::ListView, GetTitleGivenAssetInfo(FText::FromName(Node.Sequence->GetFName()), bAdditive));
+		}
 
-	return Title;
+		if ((TitleType == ENodeTitleType::FullTitle) && (SyncGroup.GroupName != NAME_None))
+		{
+			FFormatNamedArguments Args;
+			Args.Add(TEXT("Title"), CachedNodeTitles[ENodeTitleType::ListView]);
+			Args.Add(TEXT("SyncGroup"), FText::FromName(SyncGroup.GroupName));
+
+			// FText::Format() is slow, so we cache this to save on performance
+			CachedNodeTitles.SetCachedTitle(TitleType, FText::Format(LOCTEXT("SequenceNodeGroupWithSubtitle", "{Title}\nSync group {SyncGroup}"), Args));
+		}
+		else
+		{
+			return CachedNodeTitles[ENodeTitleType::ListView];
+		}
+	}
+	return CachedNodeTitles[TitleType];
 }
 
 FText UAnimGraphNode_SequencePlayer::GetTitleGivenAssetInfo(const FText& AssetName, bool bKnownToBeAdditive)
diff --git a/Engine/Source/Editor/AnimGraph/Private/AnimGraphNode_Slot.cpp b/Engine/Source/Editor/AnimGraph/Private/AnimGraphNode_Slot.cpp
--- a/Engine/Source/Editor/AnimGraph/Private/AnimGraphNode_Slot.cpp
+++ b/Engine/Source/Editor/AnimGraph/Private/AnimGraphNode_Slot.cpp
@@ -28,19 +28,33 @@ FText UAnimGraphNode_Slot::GetTooltipText() const
 
 FText UAnimGraphNode_Slot::GetNodeTitle(ENodeTitleType::Type TitleType) const
 {
-	const FText SlotName = (Node.SlotName != NAME_None) ? FText::FromName(Node.SlotName) : LOCTEXT("NoSlotName", "(No slot name)");
-
-	FFormatNamedArguments Args;
-	Args.Add(TEXT("SlotName"), SlotName);
-
-	if (TitleType == ENodeTitleType::ListView)
+	if (Node.SlotName != NAME_None)
 	{
-		return FText::Format(LOCTEXT("SlotNodeListTitle", "Slot '{SlotName}'"), Args);
+		if (TitleType == ENodeTitleType::ListView)
+		{
+			return LOCTEXT("SlotNodeListTitle_NoName", "Slot '(No slot name)'");
+		}
+		else
+		{
+			return LOCTEXT("SlotNodeTitle_NoName", "(No slot name)\nSlot");
+		}
 	}
-	else
+	else if (!CachedNodeTitles.IsTitleCached(TitleType))
 	{
-		return FText::Format(LOCTEXT("SlotNodeTitle", "{SlotName}\nSlot"), Args);
+		FFormatNamedArguments Args;
+		Args.Add(TEXT("SlotName"), FText::FromName(Node.SlotName));
+
+		// FText::Format() is slow, so we cache this to save on performance
+		if (TitleType == ENodeTitleType::ListView)
+		{
+			CachedNodeTitles.SetCachedTitle(TitleType, FText::Format(LOCTEXT("SlotNodeListTitle", "Slot '{SlotName}'"), Args));
+		}
+		else
+		{
+			CachedNodeTitles.SetCachedTitle(TitleType, FText::Format(LOCTEXT("SlotNodeTitle", "{SlotName}\nSlot"), Args));
+		}
 	}
+	return CachedNodeTitles[TitleType];
 }
 
 FString UAnimGraphNode_Slot::GetNodeCategory() const
diff --git a/Engine/Source/Editor/AnimGraph/Private/AnimGraphNode_SpringBone.cpp b/Engine/Source/Editor/AnimGraph/Private/AnimGraphNode_SpringBone.cpp
--- a/Engine/Source/Editor/AnimGraph/Private/AnimGraphNode_SpringBone.cpp
+++ b/Engine/Source/Editor/AnimGraph/Private/AnimGraphNode_SpringBone.cpp
@@ -25,22 +25,28 @@ FText UAnimGraphNode_SpringBone::GetTooltipText() const
 
 FText UAnimGraphNode_SpringBone::GetNodeTitle(ENodeTitleType::Type TitleType) const
 {
-	FFormatNamedArguments Args;
-	Args.Add(TEXT("ControllerDescription"), GetControllerDescription());
-	Args.Add(TEXT("BoneName"), FText::FromName(Node.SpringBone.BoneName));
-
-	if(TitleType == ENodeTitleType::ListView)
+	if ((TitleType == ENodeTitleType::ListView) && (Node.SpringBone.BoneName == NAME_None))
 	{
-		if (Node.SpringBone.BoneName == NAME_None)
-		{
-			return FText::Format(LOCTEXT("AnimGraphNode_SpringBone_MenuTitle", "{ControllerDescription}"), Args);
-		}
-		return FText::Format(LOCTEXT("AnimGraphNode_SpringBone_ListTitle", "{ControllerDescription} - Bone: {BoneName}"), Args);
+		return GetControllerDescription();
 	}
-	else
+	else if(!CachedNodeTitles.IsTitleCached(TitleType))
 	{
-		return FText::Format(LOCTEXT("AnimGraphNode_SpringBone_Title", "{ControllerDescription}\nBone: {BoneName}"), Args);
+		FFormatNamedArguments Args;
+		Args.Add(TEXT("ControllerDescription"), GetControllerDescription());
+		Args.Add(TEXT("BoneName"), FText::FromName(Node.SpringBone.BoneName));
+
+		// FText::Format() is slow, so we cache this to save on performance
+		if (TitleType == ENodeTitleType::ListView)
+		{
+			CachedNodeTitles.SetCachedTitle(TitleType, FText::Format(LOCTEXT("AnimGraphNode_SpringBone_ListTitle", "{ControllerDescription} - Bone: {BoneName}"), Args));
+		}
+		else
+		{
+			CachedNodeTitles.SetCachedTitle(TitleType, FText::Format(LOCTEXT("AnimGraphNode_SpringBone_Title", "{ControllerDescription}\nBone: {BoneName}"), Args));
+		}
 	}
+
+	return CachedNodeTitles[TitleType];
 }
 
 #undef LOCTEXT_NAMESPACE
\ No newline at end of file
diff --git a/Engine/Source/Editor/AnimGraph/Private/AnimGraphNode_StateMachineBase.cpp b/Engine/Source/Editor/AnimGraph/Private/AnimGraphNode_StateMachineBase.cpp
--- a/Engine/Source/Editor/AnimGraph/Private/AnimGraphNode_StateMachineBase.cpp
+++ b/Engine/Source/Editor/AnimGraph/Private/AnimGraphNode_StateMachineBase.cpp
@@ -59,20 +59,34 @@ FText UAnimGraphNode_StateMachineBase::GetTooltipText() const
 
 FText UAnimGraphNode_StateMachineBase::GetNodeTitle(ENodeTitleType::Type TitleType) const
 {
-	const FText FirstLine = (EditorStateMachineGraph != NULL) ? FText::FromString(EditorStateMachineGraph->GetName()) : LOCTEXT("ErrorNoGraph", "Error: No Graph");
 	if ((TitleType == ENodeTitleType::ListView) && (EditorStateMachineGraph == nullptr))
 	{
 		return LOCTEXT("AddNewStateMachine", "Add New State Machine...");
 	}
-	else if(TitleType == ENodeTitleType::FullTitle)
+	else if (EditorStateMachineGraph == nullptr)
 	{
-		FFormatNamedArguments Args;
-		Args.Add(TEXT("Title"), FirstLine);
-
-		return FText::Format(LOCTEXT("StateMachineFullTitle", "{Title}\nState Machine"), Args);
+		if (TitleType == ENodeTitleType::FullTitle)
+		{
+			return LOCTEXT("NullStateMachineFullTitle", "Error: No Graph\nState Machine");
+		}
+		else
+		{
+			return LOCTEXT("ErrorNoGraph", "Error: No Graph");
+		}
+	}
+	else if (TitleType == ENodeTitleType::FullTitle)
+	{
+		if (CachedFullTitle.IsOutOfDate())
+		{
+			FFormatNamedArguments Args;
+			Args.Add(TEXT("Title"), FText::FromName(EditorStateMachineGraph->GetFName()));
+			// FText::Format() is slow, so we cache this to save on performance
+			CachedFullTitle = FText::Format(LOCTEXT("StateMachineFullTitle", "{Title}\nState Machine"), Args);
+		}
+		return CachedFullTitle;
 	}
 
-	return FirstLine;
+	return FText::FromName(EditorStateMachineGraph->GetFName());
 }
 
 FString UAnimGraphNode_StateMachineBase::GetNodeCategory() const
diff --git a/Engine/Source/Editor/AnimGraph/Private/AnimGraphNode_Trail.cpp b/Engine/Source/Editor/AnimGraph/Private/AnimGraphNode_Trail.cpp
--- a/Engine/Source/Editor/AnimGraph/Private/AnimGraphNode_Trail.cpp
+++ b/Engine/Source/Editor/AnimGraph/Private/AnimGraphNode_Trail.cpp
@@ -25,22 +25,26 @@ FText UAnimGraphNode_Trail::GetTooltipText() const
 
 FText UAnimGraphNode_Trail::GetNodeTitle(ENodeTitleType::Type TitleType) const
 {
-	FFormatNamedArguments Args;
-	Args.Add(TEXT("ControllerDescription"), GetControllerDescription());
-	Args.Add(TEXT("BoneName"), FText::FromName(Node.TrailBone.BoneName));
-
-	if (TitleType == ENodeTitleType::ListView)
+	if ((TitleType == ENodeTitleType::ListView) && (Node.TrailBone.BoneName == NAME_None))
 	{
-		if (Node.TrailBone.BoneName == NAME_None)
-		{
-			return FText::Format(LOCTEXT("AnimGraphNode_Trail_MenuTitle", "{ControllerDescription}"), Args);
-		}
-		return FText::Format(LOCTEXT("AnimGraphNode_Trail_ListTitle", "{ControllerDescription} - Bone: {BoneName}"), Args);
+		return GetControllerDescription();
 	}
-	else
+	else if (!CachedNodeTitles.IsTitleCached(TitleType))
 	{
-		return FText::Format(LOCTEXT("AnimGraphNode_Trail_Title", "{ControllerDescription}\nBone: {BoneName}"), Args);
+		FFormatNamedArguments Args;
+		Args.Add(TEXT("ControllerDescription"), GetControllerDescription());
+		Args.Add(TEXT("BoneName"), FText::FromName(Node.TrailBone.BoneName));
+
+		if (TitleType == ENodeTitleType::ListView)
+		{
+			CachedNodeTitles.SetCachedTitle(TitleType, FText::Format(LOCTEXT("AnimGraphNode_Trail_ListTitle", "{ControllerDescription} - Bone: {BoneName}"), Args));
+		}
+		else
+		{
+			CachedNodeTitles.SetCachedTitle(TitleType, FText::Format(LOCTEXT("AnimGraphNode_Trail_Title", "{ControllerDescription}\nBone: {BoneName}"), Args));
+		}
 	}
+	return CachedNodeTitles[TitleType];
 }
 
 #undef LOCTEXT_NAMESPACE
\ No newline at end of file
diff --git a/Engine/Source/Editor/AnimGraph/Private/AnimGraphNode_TwoBoneIK.cpp b/Engine/Source/Editor/AnimGraph/Private/AnimGraphNode_TwoBoneIK.cpp
--- a/Engine/Source/Editor/AnimGraph/Private/AnimGraphNode_TwoBoneIK.cpp
+++ b/Engine/Source/Editor/AnimGraph/Private/AnimGraphNode_TwoBoneIK.cpp
@@ -25,22 +25,27 @@ FText UAnimGraphNode_TwoBoneIK::GetTooltipText() const
 
 FText UAnimGraphNode_TwoBoneIK::GetNodeTitle(ENodeTitleType::Type TitleType) const
 {
-	FFormatNamedArguments Args;
-	Args.Add(TEXT("ControllerDescription"), GetControllerDescription());
-	Args.Add(TEXT("BoneName"), FText::FromName(Node.IKBone.BoneName));
-
-	if(TitleType == ENodeTitleType::ListView)
+	if ((TitleType == ENodeTitleType::ListView) && (Node.IKBone.BoneName == NAME_None))
 	{
-		if (Node.IKBone.BoneName == NAME_None)
-		{
-			return FText::Format(LOCTEXT("AnimGraphNode_IKBone_MenuTitle", "{ControllerDescription}"), Args);
-		}
-		return FText::Format(LOCTEXT("AnimGraphNode_IKBone_ListTitle", "{ControllerDescription} - Bone: {BoneName}"), Args);
+		return GetControllerDescription();
 	}
-	else
+	else if (!CachedNodeTitles.IsTitleCached(TitleType))
 	{
-		return FText::Format(LOCTEXT("AnimGraphNode_IKBone_Title", "{ControllerDescription}\nBone: {BoneName}"), Args);
+		FFormatNamedArguments Args;
+		Args.Add(TEXT("ControllerDescription"), GetControllerDescription());
+		Args.Add(TEXT("BoneName"), FText::FromName(Node.IKBone.BoneName));
+
+		// FText::Format() is slow, so we cache this to save on performance
+		if (TitleType == ENodeTitleType::ListView)
+		{
+			CachedNodeTitles.SetCachedTitle(TitleType, FText::Format(LOCTEXT("AnimGraphNode_IKBone_ListTitle", "{ControllerDescription} - Bone: {BoneName}"), Args));
+		}
+		else
+		{
+			CachedNodeTitles.SetCachedTitle(TitleType, FText::Format(LOCTEXT("AnimGraphNode_IKBone_Title", "{ControllerDescription}\nBone: {BoneName}"), Args));
+		}
 	}
+	return CachedNodeTitles[TitleType];
 }
 
 void UAnimGraphNode_TwoBoneIK::Draw( FPrimitiveDrawInterface* PDI, USkeletalMeshComponent* SkelMeshComp ) const
diff --git a/Engine/Source/Editor/AnimGraph/Private/AnimGraphNode_UseCachedPose.cpp b/Engine/Source/Editor/AnimGraph/Private/AnimGraphNode_UseCachedPose.cpp
--- a/Engine/Source/Editor/AnimGraph/Private/AnimGraphNode_UseCachedPose.cpp
+++ b/Engine/Source/Editor/AnimGraph/Private/AnimGraphNode_UseCachedPose.cpp
@@ -25,9 +25,13 @@ FText UAnimGraphNode_UseCachedPose::GetTooltipText() const
 
 FText UAnimGraphNode_UseCachedPose::GetNodeTitle(ENodeTitleType::Type TitleType) const
 {
-	FFormatNamedArguments Args;
-	Args.Add(TEXT("CachePoseName"), FText::FromString(NameOfCache));
-	return FText::Format(LOCTEXT("AnimGraphNode_UseCachedPose_Title", "Use cached pose '{CachePoseName}'"), Args);
+	if (CachedNodeTitle.IsOutOfDate())
+	{
+		FFormatNamedArguments Args;
+		Args.Add(TEXT("CachePoseName"), FText::FromString(NameOfCache));
+		CachedNodeTitle = FText::Format(LOCTEXT("AnimGraphNode_UseCachedPose_Title", "Use cached pose '{CachePoseName}'"), Args);
+	}
+	return CachedNodeTitle;
 }
 
 FString UAnimGraphNode_UseCachedPose::GetNodeCategory() const
diff --git a/Engine/Source/Editor/AnimGraph/Private/AnimStateTransitionNode.cpp b/Engine/Source/Editor/AnimGraph/Private/AnimStateTransitionNode.cpp
--- a/Engine/Source/Editor/AnimGraph/Private/AnimStateTransitionNode.cpp
+++ b/Engine/Source/Editor/AnimGraph/Private/AnimStateTransitionNode.cpp
@@ -150,6 +150,10 @@ FText UAnimStateTransitionNode::GetNodeTitle(ENodeTitleType::Type TitleType) con
 	{
 		FFormatNamedArguments Args;
 		Args.Add(TEXT("BoundGraph"), (BoundGraph != NULL) ? FText::FromString(BoundGraph->GetName()) : LOCTEXT("Null", "(null)") );
+		// @TODO: FText::Format() is slow, and we could benefit from caching 
+		//        this off like we do for a lot of other nodes (but we have to
+		//        make sure to invalidate the cached string at the appropriate 
+		//        times).
 		return FText::Format(LOCTEXT("TransitioNState", "Trans {BoundGraph}}"), Args);
 	}
 }
diff --git a/Engine/Source/Editor/AnimGraph/Private/K2Node_TransitionRuleGetter.cpp b/Engine/Source/Editor/AnimGraph/Private/K2Node_TransitionRuleGetter.cpp
--- a/Engine/Source/Editor/AnimGraph/Private/K2Node_TransitionRuleGetter.cpp
+++ b/Engine/Source/Editor/AnimGraph/Private/K2Node_TransitionRuleGetter.cpp
@@ -71,6 +71,10 @@ FText UK2Node_TransitionRuleGetter::GetFriendlyName(ETransitionGetter::Type Type
 
 FText UK2Node_TransitionRuleGetter::GetNodeTitle(ENodeTitleType::Type TitleType) const
 {
+	// @TODO: FText::Format() is slow... consider caching this tooltip like 
+	//        we do for a lot of the BP nodes now (unfamiliar with this 
+	//        node's purpose, so hesitant to muck with this at this time).
+
 	if (AssociatedAnimAssetPlayerNode != NULL)
 	{
 		UAnimationAsset* BoundAsset = AssociatedAnimAssetPlayerNode->GetAnimationAsset();
diff --git a/Engine/Source/Editor/IntroTutorials/Classes/K2Node_TutorialExcerptComplete.h b/Engine/Source/Editor/IntroTutorials/Classes/K2Node_TutorialExcerptComplete.h
--- a/Engine/Source/Editor/IntroTutorials/Classes/K2Node_TutorialExcerptComplete.h
+++ b/Engine/Source/Editor/IntroTutorials/Classes/K2Node_TutorialExcerptComplete.h
@@ -3,6 +3,7 @@
 #pragma once
 
 #include "BlueprintGraphDefinitions.h"
+#include "EdGraph/EdGraphNodeUtils.h" // for FNodeTextCache
 #include "K2Node_TutorialExcerptComplete.generated.h"
 
 
@@ -58,4 +59,7 @@ private:
 
 	/** Set of all excerpts "excerpt complete" nodes currently use */
 	static TSet<FString> AllBlueprintExcerpts;
+
+	/** Constructing FText strings can be costly, so we cache the node's title */
+	FNodeTextCache CachedNodeTitle;
 };
\ No newline at end of file
diff --git a/Engine/Source/Editor/IntroTutorials/Private/K2Node_TutorialExcerptComplete.cpp b/Engine/Source/Editor/IntroTutorials/Private/K2Node_TutorialExcerptComplete.cpp
--- a/Engine/Source/Editor/IntroTutorials/Private/K2Node_TutorialExcerptComplete.cpp
+++ b/Engine/Source/Editor/IntroTutorials/Private/K2Node_TutorialExcerptComplete.cpp
@@ -31,6 +31,7 @@ void UK2Node_TutorialExcerptComplete::PostLoad()
 		CachedExcerpt = Pin->DefaultValue;
 		Pin->bNotConnectable = true;
 		bExcerptCacheDirty = true;
+		CachedNodeTitle.MarkDirty();
 	}
 }
 
@@ -58,6 +59,7 @@ void UK2Node_TutorialExcerptComplete::PinDefaultValueChanged(UEdGraphPin* Pin)
 			bExcerptCacheDirty = true;
 		}
 		CachedExcerpt = Pin->DefaultValue;
+		CachedNodeTitle.MarkDirty();
 	}
 }
 
@@ -92,16 +94,17 @@ void UK2Node_TutorialExcerptComplete::PostPlacedNewNode()
 
 FText UK2Node_TutorialExcerptComplete::GetNodeTitle(ENodeTitleType::Type TitleType) const
 {
-	if(CachedExcerpt.Len() > 0)
+	if (CachedExcerpt.IsEmpty())
 	{
-		FFormatNamedArguments NamedArgs;
-		NamedArgs.Add(TEXT("ExcerptName"), FText::FromString(CachedExcerpt));
-		return FText::Format(LOCTEXT("TutorialExcerptComplete_TitleWithExcerpt", "Tutorial Excerpt '{ExcerptName}' Complete"), NamedArgs);
+		return LOCTEXT("TutorialExcerptComplete_Title", "Tutorial Excerpt Complete");
 	}
-	else
+	else if (CachedNodeTitle.IsOutOfDate())
 	{
-		return LOCTEXT("TutorialExcerptComplete_Title", "Tutorial Excerpt Complete");
+		FFormatNamedArguments NamedArgs;
+		NamedArgs.Add(TEXT("ExcerptName"), FText::FromString(CachedExcerpt));
+		CachedNodeTitle = FText::Format(LOCTEXT("TutorialExcerptComplete_TitleWithExcerpt", "Tutorial Excerpt '{ExcerptName}' Complete"), NamedArgs);
 	}
+	return CachedNodeTitle;
 }
 
 void UK2Node_TutorialExcerptComplete::TutorialExcerptComplete(const FString& Excerpt)
@@ -118,7 +121,7 @@ bool UK2Node_TutorialExcerptComplete::IsExcerptInteractive(const FString& InExer
 		for(TObjectIterator<UK2Node_TutorialExcerptComplete> It; It; ++It)
 		{
 			const FString& CachedExcerpt = It->CachedExcerpt;
-			if(CachedExcerpt.Len() > 0)
+			if (CachedExcerpt.Len() > 0)
 			{
 				AllBlueprintExcerpts.Add(CachedExcerpt);
 			}
