Added a new MenuTitle option for nodes to reduce issues with using ListView in a variety of circumstances. Fixes issues with some nodes being displayed by their palette action name in other menu lists.

diff --git a/Engine/Source/Editor/AnimGraph/Private/AnimGraphNode_BlendSpaceEvaluator.cpp b/Engine/Source/Editor/AnimGraph/Private/AnimGraphNode_BlendSpaceEvaluator.cpp
--- a/Engine/Source/Editor/AnimGraph/Private/AnimGraphNode_BlendSpaceEvaluator.cpp
+++ b/Engine/Source/Editor/AnimGraph/Private/AnimGraphNode_BlendSpaceEvaluator.cpp
@@ -25,7 +25,7 @@ FText UAnimGraphNode_BlendSpaceEvaluator::GetNodeTitle(ENodeTitleType::Type Titl
 {
 	if (Node.BlendSpace == nullptr)
 	{
-		if (TitleType == ENodeTitleType::ListView)
+		if (TitleType == ENodeTitleType::ListView || TitleType == ENodeTitleType::MenuTitle)
 		{
 			return LOCTEXT("BlendSpaceEvaluator_NONE_ListTitle", "Blendspace Evaluator '(None)'");
 		}
@@ -38,7 +38,7 @@ FText UAnimGraphNode_BlendSpaceEvaluator::GetNodeTitle(ENodeTitleType::Type Titl
 	{
 		const FText BlendSpaceName = FText::FromString(Node.BlendSpace->GetName());
 
-		if (TitleType == ENodeTitleType::ListView)
+		if (TitleType == ENodeTitleType::ListView || TitleType == ENodeTitleType::MenuTitle)
 		{
 			FFormatNamedArguments Args;
 			Args.Add(TEXT("BlendSpaceName"), BlendSpaceName);
diff --git a/Engine/Source/Editor/AnimGraph/Private/AnimGraphNode_BlendSpacePlayer.cpp b/Engine/Source/Editor/AnimGraph/Private/AnimGraphNode_BlendSpacePlayer.cpp
--- a/Engine/Source/Editor/AnimGraph/Private/AnimGraphNode_BlendSpacePlayer.cpp
+++ b/Engine/Source/Editor/AnimGraph/Private/AnimGraphNode_BlendSpacePlayer.cpp
@@ -27,7 +27,7 @@ FText UAnimGraphNode_BlendSpacePlayer::GetNodeTitle(ENodeTitleType::Type TitleTy
 {	
 	if (Node.BlendSpace == nullptr)
 	{
-		if (TitleType == ENodeTitleType::ListView)
+		if (TitleType == ENodeTitleType::ListView || TitleType == ENodeTitleType::MenuTitle)
 		{
 			return LOCTEXT("BlendspacePlayer_NONE_ListTitle", "Blendspace Player '(None)'");
 		}
@@ -40,7 +40,7 @@ FText UAnimGraphNode_BlendSpacePlayer::GetNodeTitle(ENodeTitleType::Type TitleTy
 	{
 		const FText BlendSpaceName = FText::FromString(Node.BlendSpace->GetName());
 
-		if (TitleType == ENodeTitleType::ListView)
+		if (TitleType == ENodeTitleType::ListView || TitleType == ENodeTitleType::MenuTitle)
 		{
 			FFormatNamedArguments Args;
 			Args.Add(TEXT("BlendSpaceName"), BlendSpaceName);
diff --git a/Engine/Source/Editor/AnimGraph/Private/AnimGraphNode_BoneDrivenController.cpp b/Engine/Source/Editor/AnimGraph/Private/AnimGraphNode_BoneDrivenController.cpp
--- a/Engine/Source/Editor/AnimGraph/Private/AnimGraphNode_BoneDrivenController.cpp
+++ b/Engine/Source/Editor/AnimGraph/Private/AnimGraphNode_BoneDrivenController.cpp
@@ -18,7 +18,7 @@ FText UAnimGraphNode_BoneDrivenController::GetTooltipText() const
 
 FText UAnimGraphNode_BoneDrivenController::GetNodeTitle(ENodeTitleType::Type TitleType) const
 {
-	if ((Node.SourceBone.BoneName == NAME_None) && (Node.TargetBone.BoneName == NAME_None) && (TitleType == ENodeTitleType::ListView))
+	if ((Node.SourceBone.BoneName == NAME_None) && (Node.TargetBone.BoneName == NAME_None) && (TitleType == ENodeTitleType::ListView || TitleType == ENodeTitleType::MenuTitle))
 	{
 		return GetControllerDescription();
 	}
@@ -29,7 +29,7 @@ FText UAnimGraphNode_BoneDrivenController::GetNodeTitle(ENodeTitleType::Type Tit
 		Args.Add(TEXT("SourceBone"), FText::FromName(Node.SourceBone.BoneName));
 		Args.Add(TEXT("TargetBone"), FText::FromName(Node.TargetBone.BoneName));
 
-		if (TitleType == ENodeTitleType::ListView)
+		if (TitleType == ENodeTitleType::ListView || TitleType == ENodeTitleType::MenuTitle)
 		{
 			Args.Add(TEXT("Delim"), FText::FromString(TEXT(" - ")));
 		}
diff --git a/Engine/Source/Editor/AnimGraph/Private/AnimGraphNode_CopyBone.cpp b/Engine/Source/Editor/AnimGraph/Private/AnimGraphNode_CopyBone.cpp
--- a/Engine/Source/Editor/AnimGraph/Private/AnimGraphNode_CopyBone.cpp
+++ b/Engine/Source/Editor/AnimGraph/Private/AnimGraphNode_CopyBone.cpp
@@ -25,7 +25,7 @@ FText UAnimGraphNode_CopyBone::GetTooltipText() const
 
 FText UAnimGraphNode_CopyBone::GetNodeTitle(ENodeTitleType::Type TitleType) const
 {
-	if ((TitleType == ENodeTitleType::ListView) && (Node.TargetBone.BoneName == NAME_None) && (Node.SourceBone.BoneName == NAME_None))
+	if ((TitleType == ENodeTitleType::ListView || TitleType == ENodeTitleType::MenuTitle) && (Node.TargetBone.BoneName == NAME_None) && (Node.SourceBone.BoneName == NAME_None))
 	{
 		return GetControllerDescription();
 	}
@@ -36,7 +36,7 @@ FText UAnimGraphNode_CopyBone::GetNodeTitle(ENodeTitleType::Type TitleType) cons
 		Args.Add(TEXT("SourceBoneName"), FText::FromName(Node.SourceBone.BoneName));
 		Args.Add(TEXT("TargetBoneName"), FText::FromName(Node.TargetBone.BoneName));
 
-		if (TitleType == ENodeTitleType::ListView)
+		if (TitleType == ENodeTitleType::ListView || TitleType == ENodeTitleType::MenuTitle)
 		{
 			CachedNodeTitles.SetCachedTitle(TitleType, FText::Format(LOCTEXT("AnimGraphNode_CopyBone_ListTitle", "{ControllerDescription} - Source Bone: {SourceBoneName} - Target Bone: {TargetBoneName}"), Args));
 		}
diff --git a/Engine/Source/Editor/AnimGraph/Private/AnimGraphNode_LookAt.cpp b/Engine/Source/Editor/AnimGraph/Private/AnimGraphNode_LookAt.cpp
--- a/Engine/Source/Editor/AnimGraph/Private/AnimGraphNode_LookAt.cpp
+++ b/Engine/Source/Editor/AnimGraph/Private/AnimGraphNode_LookAt.cpp
@@ -30,7 +30,7 @@ FText UAnimGraphNode_LookAt::GetTooltipText() const
 
 FText UAnimGraphNode_LookAt::GetNodeTitle(ENodeTitleType::Type TitleType) const
 {
-	if ((TitleType == ENodeTitleType::ListView) && (Node.BoneToModify.BoneName == NAME_None))
+	if ((TitleType == ENodeTitleType::ListView || TitleType == ENodeTitleType::MenuTitle) && (Node.BoneToModify.BoneName == NAME_None))
 	{
 		return GetControllerDescription();
 	}
@@ -41,7 +41,7 @@ FText UAnimGraphNode_LookAt::GetNodeTitle(ENodeTitleType::Type TitleType) const
 		Args.Add(TEXT("BoneName"), FText::FromName(Node.BoneToModify.BoneName));
 
 		// FText::Format() is slow, so we cache this to save on performance
-		if (TitleType == ENodeTitleType::ListView)
+		if (TitleType == ENodeTitleType::ListView || TitleType == ENodeTitleType::MenuTitle)
 		{
 			CachedNodeTitles.SetCachedTitle(TitleType, FText::Format(LOCTEXT("AnimGraphNode_LookAt_ListTitle", "{ControllerDescription} - Bone: {BoneName}"), Args));
 		}
diff --git a/Engine/Source/Editor/AnimGraph/Private/AnimGraphNode_ModifyBone.cpp b/Engine/Source/Editor/AnimGraph/Private/AnimGraphNode_ModifyBone.cpp
--- a/Engine/Source/Editor/AnimGraph/Private/AnimGraphNode_ModifyBone.cpp
+++ b/Engine/Source/Editor/AnimGraph/Private/AnimGraphNode_ModifyBone.cpp
@@ -30,7 +30,7 @@ FText UAnimGraphNode_ModifyBone::GetTooltipText() const
 
 FText UAnimGraphNode_ModifyBone::GetNodeTitle(ENodeTitleType::Type TitleType) const
 {
-	if ((TitleType == ENodeTitleType::ListView) && (Node.BoneToModify.BoneName == NAME_None))
+	if ((TitleType == ENodeTitleType::ListView || TitleType == ENodeTitleType::MenuTitle) && (Node.BoneToModify.BoneName == NAME_None))
 	{
 		return GetControllerDescription();
 	}
@@ -41,7 +41,7 @@ FText UAnimGraphNode_ModifyBone::GetNodeTitle(ENodeTitleType::Type TitleType) co
 		Args.Add(TEXT("BoneName"), FText::FromName(Node.BoneToModify.BoneName));
 
 		// FText::Format() is slow, so we cache this to save on performance
-		if (TitleType == ENodeTitleType::ListView)
+		if (TitleType == ENodeTitleType::ListView || TitleType == ENodeTitleType::MenuTitle)
 		{
 			CachedNodeTitles.SetCachedTitle(TitleType, FText::Format(LOCTEXT("AnimGraphNode_ModifyBone_ListTitle", "{ControllerDescription} - Bone: {BoneName}"), Args));
 		}
diff --git a/Engine/Source/Editor/AnimGraph/Private/AnimGraphNode_RotationMultiplier.cpp b/Engine/Source/Editor/AnimGraph/Private/AnimGraphNode_RotationMultiplier.cpp
--- a/Engine/Source/Editor/AnimGraph/Private/AnimGraphNode_RotationMultiplier.cpp
+++ b/Engine/Source/Editor/AnimGraph/Private/AnimGraphNode_RotationMultiplier.cpp
@@ -24,7 +24,7 @@ FText UAnimGraphNode_RotationMultiplier::GetTooltipText() const
 
 FText UAnimGraphNode_RotationMultiplier::GetNodeTitle(ENodeTitleType::Type TitleType) const
 {
-	if ((TitleType == ENodeTitleType::ListView) && (Node.TargetBone.BoneName == NAME_None))
+	if ((TitleType == ENodeTitleType::ListView || TitleType == ENodeTitleType::MenuTitle) && (Node.TargetBone.BoneName == NAME_None))
 	{
 		return GetControllerDescription();
 	}
@@ -35,7 +35,7 @@ FText UAnimGraphNode_RotationMultiplier::GetNodeTitle(ENodeTitleType::Type Title
 		Args.Add(TEXT("BoneName"), FText::FromName(Node.TargetBone.BoneName));
 
 		// FText::Format() is slow, so we cache this to save on performance
-		if (TitleType == ENodeTitleType::ListView)
+		if (TitleType == ENodeTitleType::ListView || TitleType == ENodeTitleType::MenuTitle)
 		{
 			CachedNodeTitles.SetCachedTitle(TitleType, FText::Format(LOCTEXT("AnimGraphNode_ModifyBone_ListTitle", "{ControllerDescription} - Bone: {BoneName}"), Args));
 		}
diff --git a/Engine/Source/Editor/AnimGraph/Private/AnimGraphNode_RotationOffsetBlendSpace.cpp b/Engine/Source/Editor/AnimGraph/Private/AnimGraphNode_RotationOffsetBlendSpace.cpp
--- a/Engine/Source/Editor/AnimGraph/Private/AnimGraphNode_RotationOffsetBlendSpace.cpp
+++ b/Engine/Source/Editor/AnimGraph/Private/AnimGraphNode_RotationOffsetBlendSpace.cpp
@@ -27,7 +27,7 @@ FText UAnimGraphNode_RotationOffsetBlendSpace::GetNodeTitle(ENodeTitleType::Type
 {
 	if (Node.BlendSpace == nullptr)
 	{
-		if (TitleType == ENodeTitleType::ListView)
+		if (TitleType == ENodeTitleType::ListView || TitleType == ENodeTitleType::MenuTitle)
 		{
 			return LOCTEXT("RotationOffsetBlend_NONE_ListTitle", "AimOffset '(None)'");
 		}
@@ -44,7 +44,7 @@ FText UAnimGraphNode_RotationOffsetBlendSpace::GetNodeTitle(ENodeTitleType::Type
 		Args.Add(TEXT("BlendSpaceName"), BlendSpaceName);
 
 		// FText::Format() is slow, so we cache this to save on performance
-		if (TitleType == ENodeTitleType::ListView)
+		if (TitleType == ENodeTitleType::ListView || TitleType == ENodeTitleType::MenuTitle)
 		{
 			CachedNodeTitles.SetCachedTitle(TitleType, FText::Format(LOCTEXT("AimOffsetListTitle", "AimOffset '{BlendSpaceName}'"), Args));
 		}
diff --git a/Engine/Source/Editor/AnimGraph/Private/AnimGraphNode_SaveCachedPose.cpp b/Engine/Source/Editor/AnimGraph/Private/AnimGraphNode_SaveCachedPose.cpp
--- a/Engine/Source/Editor/AnimGraph/Private/AnimGraphNode_SaveCachedPose.cpp
+++ b/Engine/Source/Editor/AnimGraph/Private/AnimGraphNode_SaveCachedPose.cpp
@@ -51,7 +51,7 @@ FText UAnimGraphNode_SaveCachedPose::GetNodeTitle(ENodeTitleType::Type TitleType
 	{
 		return FText::FromString(CacheName);
 	}
-	else if ((TitleType == ENodeTitleType::ListView) && CacheName.IsEmpty())
+	else if ((TitleType == ENodeTitleType::MenuTitle) && CacheName.IsEmpty())
 	{
 		return LOCTEXT("NewSaveCachedPose", "New Save cached pose...");
 	}
diff --git a/Engine/Source/Editor/AnimGraph/Private/AnimGraphNode_SequencePlayer.cpp b/Engine/Source/Editor/AnimGraph/Private/AnimGraphNode_SequencePlayer.cpp
--- a/Engine/Source/Editor/AnimGraph/Private/AnimGraphNode_SequencePlayer.cpp
+++ b/Engine/Source/Editor/AnimGraph/Private/AnimGraphNode_SequencePlayer.cpp
@@ -83,14 +83,18 @@ FText UAnimGraphNode_SequencePlayer::GetNodeTitle(ENodeTitleType::Type TitleType
 	}
 	else if (!CachedNodeTitles.IsTitleCached(TitleType))
 	{
-		if (!CachedNodeTitles.IsTitleCached(ENodeTitleType::ListView))
+		if(SyncGroup.GroupName == NAME_None)
 		{
-			const bool bAdditive = Node.Sequence->IsValidAdditive();
-			// FText::Format() is slow, so we cache this to save on performance
-			CachedNodeTitles.SetCachedTitle(ENodeTitleType::ListView, GetTitleGivenAssetInfo(FText::FromName(Node.Sequence->GetFName()), bAdditive));
-		}
+			TitleType = ENodeTitleType::ListView;
 
-		if ((TitleType == ENodeTitleType::FullTitle) && (SyncGroup.GroupName != NAME_None))
+			if(TitleType == ENodeTitleType::ListView || TitleType == ENodeTitleType::MenuTitle)
+			{
+				const bool bAdditive = Node.Sequence->IsValidAdditive();
+				// FText::Format() is slow, so we cache this to save on performance
+				CachedNodeTitles.SetCachedTitle(TitleType, GetTitleGivenAssetInfo(FText::FromName(Node.Sequence->GetFName()), bAdditive));
+			}
+		}
+		else if (TitleType == ENodeTitleType::FullTitle)
 		{
 			FFormatNamedArguments Args;
 			Args.Add(TEXT("Title"), CachedNodeTitles[ENodeTitleType::ListView]);
@@ -99,10 +103,6 @@ FText UAnimGraphNode_SequencePlayer::GetNodeTitle(ENodeTitleType::Type TitleType
 			// FText::Format() is slow, so we cache this to save on performance
 			CachedNodeTitles.SetCachedTitle(TitleType, FText::Format(LOCTEXT("SequenceNodeGroupWithSubtitle", "{Title}\nSync group {SyncGroup}"), Args));
 		}
-		else
-		{
-			return CachedNodeTitles[ENodeTitleType::ListView];
-		}
 	}
 	return CachedNodeTitles[TitleType];
 }
diff --git a/Engine/Source/Editor/AnimGraph/Private/AnimGraphNode_Slot.cpp b/Engine/Source/Editor/AnimGraph/Private/AnimGraphNode_Slot.cpp
--- a/Engine/Source/Editor/AnimGraph/Private/AnimGraphNode_Slot.cpp
+++ b/Engine/Source/Editor/AnimGraph/Private/AnimGraphNode_Slot.cpp
@@ -30,7 +30,7 @@ FText UAnimGraphNode_Slot::GetNodeTitle(ENodeTitleType::Type TitleType) const
 {
 	if (Node.SlotName != NAME_None)
 	{
-		if (TitleType == ENodeTitleType::ListView)
+		if (TitleType == ENodeTitleType::ListView || TitleType == ENodeTitleType::MenuTitle)
 		{
 			return LOCTEXT("SlotNodeListTitle_NoName", "Slot '(No slot name)'");
 		}
@@ -45,7 +45,7 @@ FText UAnimGraphNode_Slot::GetNodeTitle(ENodeTitleType::Type TitleType) const
 		Args.Add(TEXT("SlotName"), FText::FromName(Node.SlotName));
 
 		// FText::Format() is slow, so we cache this to save on performance
-		if (TitleType == ENodeTitleType::ListView)
+		if (TitleType == ENodeTitleType::ListView || TitleType == ENodeTitleType::MenuTitle)
 		{
 			CachedNodeTitles.SetCachedTitle(TitleType, FText::Format(LOCTEXT("SlotNodeListTitle", "Slot '{SlotName}'"), Args));
 		}
diff --git a/Engine/Source/Editor/AnimGraph/Private/AnimGraphNode_SpringBone.cpp b/Engine/Source/Editor/AnimGraph/Private/AnimGraphNode_SpringBone.cpp
--- a/Engine/Source/Editor/AnimGraph/Private/AnimGraphNode_SpringBone.cpp
+++ b/Engine/Source/Editor/AnimGraph/Private/AnimGraphNode_SpringBone.cpp
@@ -25,7 +25,7 @@ FText UAnimGraphNode_SpringBone::GetTooltipText() const
 
 FText UAnimGraphNode_SpringBone::GetNodeTitle(ENodeTitleType::Type TitleType) const
 {
-	if ((TitleType == ENodeTitleType::ListView) && (Node.SpringBone.BoneName == NAME_None))
+	if ((TitleType == ENodeTitleType::ListView || TitleType == ENodeTitleType::MenuTitle) && (Node.SpringBone.BoneName == NAME_None))
 	{
 		return GetControllerDescription();
 	}
@@ -36,7 +36,7 @@ FText UAnimGraphNode_SpringBone::GetNodeTitle(ENodeTitleType::Type TitleType) co
 		Args.Add(TEXT("BoneName"), FText::FromName(Node.SpringBone.BoneName));
 
 		// FText::Format() is slow, so we cache this to save on performance
-		if (TitleType == ENodeTitleType::ListView)
+		if (TitleType == ENodeTitleType::ListView || TitleType == ENodeTitleType::MenuTitle)
 		{
 			CachedNodeTitles.SetCachedTitle(TitleType, FText::Format(LOCTEXT("AnimGraphNode_SpringBone_ListTitle", "{ControllerDescription} - Bone: {BoneName}"), Args));
 		}
diff --git a/Engine/Source/Editor/AnimGraph/Private/AnimGraphNode_Trail.cpp b/Engine/Source/Editor/AnimGraph/Private/AnimGraphNode_Trail.cpp
--- a/Engine/Source/Editor/AnimGraph/Private/AnimGraphNode_Trail.cpp
+++ b/Engine/Source/Editor/AnimGraph/Private/AnimGraphNode_Trail.cpp
@@ -25,7 +25,7 @@ FText UAnimGraphNode_Trail::GetTooltipText() const
 
 FText UAnimGraphNode_Trail::GetNodeTitle(ENodeTitleType::Type TitleType) const
 {
-	if ((TitleType == ENodeTitleType::ListView) && (Node.TrailBone.BoneName == NAME_None))
+	if ((TitleType == ENodeTitleType::ListView || TitleType == ENodeTitleType::MenuTitle) && (Node.TrailBone.BoneName == NAME_None))
 	{
 		return GetControllerDescription();
 	}
@@ -35,7 +35,7 @@ FText UAnimGraphNode_Trail::GetNodeTitle(ENodeTitleType::Type TitleType) const
 		Args.Add(TEXT("ControllerDescription"), GetControllerDescription());
 		Args.Add(TEXT("BoneName"), FText::FromName(Node.TrailBone.BoneName));
 
-		if (TitleType == ENodeTitleType::ListView)
+		if (TitleType == ENodeTitleType::ListView || TitleType == ENodeTitleType::MenuTitle)
 		{
 			CachedNodeTitles.SetCachedTitle(TitleType, FText::Format(LOCTEXT("AnimGraphNode_Trail_ListTitle", "{ControllerDescription} - Bone: {BoneName}"), Args));
 		}
diff --git a/Engine/Source/Editor/AnimGraph/Private/AnimGraphNode_TwoBoneIK.cpp b/Engine/Source/Editor/AnimGraph/Private/AnimGraphNode_TwoBoneIK.cpp
--- a/Engine/Source/Editor/AnimGraph/Private/AnimGraphNode_TwoBoneIK.cpp
+++ b/Engine/Source/Editor/AnimGraph/Private/AnimGraphNode_TwoBoneIK.cpp
@@ -25,7 +25,7 @@ FText UAnimGraphNode_TwoBoneIK::GetTooltipText() const
 
 FText UAnimGraphNode_TwoBoneIK::GetNodeTitle(ENodeTitleType::Type TitleType) const
 {
-	if ((TitleType == ENodeTitleType::ListView) && (Node.IKBone.BoneName == NAME_None))
+	if ((TitleType == ENodeTitleType::ListView || TitleType == ENodeTitleType::MenuTitle) && (Node.IKBone.BoneName == NAME_None))
 	{
 		return GetControllerDescription();
 	}
@@ -36,7 +36,7 @@ FText UAnimGraphNode_TwoBoneIK::GetNodeTitle(ENodeTitleType::Type TitleType) con
 		Args.Add(TEXT("BoneName"), FText::FromName(Node.IKBone.BoneName));
 
 		// FText::Format() is slow, so we cache this to save on performance
-		if (TitleType == ENodeTitleType::ListView)
+		if (TitleType == ENodeTitleType::ListView || TitleType == ENodeTitleType::MenuTitle)
 		{
 			CachedNodeTitles.SetCachedTitle(TitleType, FText::Format(LOCTEXT("AnimGraphNode_IKBone_ListTitle", "{ControllerDescription} - Bone: {BoneName}"), Args));
 		}
diff --git a/Engine/Source/Editor/AnimGraph/Private/AnimGraphNode_WheelHandler.cpp b/Engine/Source/Editor/AnimGraph/Private/AnimGraphNode_WheelHandler.cpp
--- a/Engine/Source/Editor/AnimGraph/Private/AnimGraphNode_WheelHandler.cpp
+++ b/Engine/Source/Editor/AnimGraph/Private/AnimGraphNode_WheelHandler.cpp
@@ -32,7 +32,7 @@ FText UAnimGraphNode_WheelHandler::GetTooltipText() const
 FText UAnimGraphNode_WheelHandler::GetNodeTitle(ENodeTitleType::Type TitleType) const
 {
 	FText NodeTitle;
-	if (TitleType == ENodeTitleType::ListView)
+	if (TitleType == ENodeTitleType::ListView || TitleType == ENodeTitleType::MenuTitle)
 	{
 		NodeTitle = GetControllerDescription();
 	}
diff --git a/Engine/Source/Editor/BlueprintGraph/Private/BlueprintNodeSpawner.cpp b/Engine/Source/Editor/BlueprintGraph/Private/BlueprintNodeSpawner.cpp
--- a/Engine/Source/Editor/BlueprintGraph/Private/BlueprintNodeSpawner.cpp
+++ b/Engine/Source/Editor/BlueprintGraph/Private/BlueprintNodeSpawner.cpp
@@ -90,7 +90,7 @@ void UBlueprintNodeSpawner::Prime()
 			K2NodeTemplate->GetMenuCategory();
 		}
 		CachedTemplateNode->GetTooltipText();
-		CachedTemplateNode->GetNodeTitle(ENodeTitleType::ListView);
+		CachedTemplateNode->GetNodeTitle(ENodeTitleType::MenuTitle);
 	}
 
 	// in case any of these cache FText::Format() operations
diff --git a/Engine/Source/Editor/BlueprintGraph/Private/K2Node_AssignDelegate.cpp b/Engine/Source/Editor/BlueprintGraph/Private/K2Node_AssignDelegate.cpp
--- a/Engine/Source/Editor/BlueprintGraph/Private/K2Node_AssignDelegate.cpp
+++ b/Engine/Source/Editor/BlueprintGraph/Private/K2Node_AssignDelegate.cpp
@@ -20,7 +20,7 @@ UK2Node_AssignDelegate::UK2Node_AssignDelegate(const class FPostConstructInitial
 //------------------------------------------------------------------------------
 FText UK2Node_AssignDelegate::GetNodeTitle(ENodeTitleType::Type TitleType) const
 {
-	if (TitleType == ENodeTitleType::ListView)
+	if (TitleType == ENodeTitleType::ListView || TitleType == ENodeTitleType::MenuTitle)
 	{
 		if (CachedListTitle.IsOutOfDate())
 		{
diff --git a/Engine/Source/Editor/BlueprintGraph/Private/K2Node_ConstructObjectFromClass.cpp b/Engine/Source/Editor/BlueprintGraph/Private/K2Node_ConstructObjectFromClass.cpp
--- a/Engine/Source/Editor/BlueprintGraph/Private/K2Node_ConstructObjectFromClass.cpp
+++ b/Engine/Source/Editor/BlueprintGraph/Private/K2Node_ConstructObjectFromClass.cpp
@@ -246,7 +246,7 @@ FText UK2Node_ConstructObjectFromClass::GetNodeTitleFormat() const
 
 FText UK2Node_ConstructObjectFromClass::GetNodeTitle(ENodeTitleType::Type TitleType) const
 {
-	if (TitleType == ENodeTitleType::ListView)
+	if (TitleType == ENodeTitleType::ListView || TitleType == ENodeTitleType::MenuTitle)
 	{
 		return GetBaseNodeTitle();
 	}
diff --git a/Engine/Source/Editor/BlueprintGraph/Private/K2Node_CustomEvent.cpp b/Engine/Source/Editor/BlueprintGraph/Private/K2Node_CustomEvent.cpp
--- a/Engine/Source/Editor/BlueprintGraph/Private/K2Node_CustomEvent.cpp
+++ b/Engine/Source/Editor/BlueprintGraph/Private/K2Node_CustomEvent.cpp
@@ -106,11 +106,7 @@ UK2Node_CustomEvent::UK2Node_CustomEvent(const class FPostConstructInitializePro
 
 FText UK2Node_CustomEvent::GetNodeTitle(ENodeTitleType::Type TitleType) const
 {
-	if (CustomFunctionName.IsNone())
-	{
-		return LOCTEXT("ActionMenuTitle", "Custom Event...");
-	}
-	else if (TitleType != ENodeTitleType::FullTitle)
+	if (TitleType != ENodeTitleType::FullTitle)
 	{
 		return FText::FromName(CustomFunctionName);
 	}
diff --git a/Engine/Source/Editor/BlueprintGraph/Private/K2Node_GetDataTableRow.cpp b/Engine/Source/Editor/BlueprintGraph/Private/K2Node_GetDataTableRow.cpp
--- a/Engine/Source/Editor/BlueprintGraph/Private/K2Node_GetDataTableRow.cpp
+++ b/Engine/Source/Editor/BlueprintGraph/Private/K2Node_GetDataTableRow.cpp
@@ -279,7 +279,7 @@ UEdGraphPin* UK2Node_GetDataTableRow::GetResultPin() const
 
 FText UK2Node_GetDataTableRow::GetNodeTitle(ENodeTitleType::Type TitleType) const
 {
-	if (TitleType == ENodeTitleType::ListView)
+	if (TitleType == ENodeTitleType::MenuTitle)
 	{
 		return LOCTEXT("ListViewTitle", "Get Data Table Row");
 	}
diff --git a/Engine/Source/Editor/BlueprintGraph/Private/K2Node_GetInputAxisKeyValue.cpp b/Engine/Source/Editor/BlueprintGraph/Private/K2Node_GetInputAxisKeyValue.cpp
--- a/Engine/Source/Editor/BlueprintGraph/Private/K2Node_GetInputAxisKeyValue.cpp
+++ b/Engine/Source/Editor/BlueprintGraph/Private/K2Node_GetInputAxisKeyValue.cpp
@@ -37,7 +37,7 @@ void UK2Node_GetInputAxisKeyValue::Initialize(const FKey AxisKey)
 
 FText UK2Node_GetInputAxisKeyValue::GetNodeTitle(ENodeTitleType::Type TitleType) const
 {
-	if (TitleType == ENodeTitleType::ListView)
+	if (TitleType == ENodeTitleType::MenuTitle)
 	{
 		return InputAxisKey.GetDisplayName();
 	}
diff --git a/Engine/Source/Editor/BlueprintGraph/Private/K2Node_GetInputAxisValue.cpp b/Engine/Source/Editor/BlueprintGraph/Private/K2Node_GetInputAxisValue.cpp
--- a/Engine/Source/Editor/BlueprintGraph/Private/K2Node_GetInputAxisValue.cpp
+++ b/Engine/Source/Editor/BlueprintGraph/Private/K2Node_GetInputAxisValue.cpp
@@ -34,7 +34,7 @@ void UK2Node_GetInputAxisValue::Initialize(const FName AxisName)
 
 FText UK2Node_GetInputAxisValue::GetNodeTitle(ENodeTitleType::Type TitleType) const
 {
-	if (TitleType == ENodeTitleType::ListView)
+	if (TitleType == ENodeTitleType::MenuTitle)
 	{
 		return FText::FromName(InputAxisName);
 	}
diff --git a/Engine/Source/Editor/BlueprintGraph/Private/K2Node_InputAction.cpp b/Engine/Source/Editor/BlueprintGraph/Private/K2Node_InputAction.cpp
--- a/Engine/Source/Editor/BlueprintGraph/Private/K2Node_InputAction.cpp
+++ b/Engine/Source/Editor/BlueprintGraph/Private/K2Node_InputAction.cpp
@@ -47,7 +47,7 @@ FLinearColor UK2Node_InputAction::GetNodeTitleColor() const
 
 FText UK2Node_InputAction::GetNodeTitle(ENodeTitleType::Type TitleType) const
 {
-	if (TitleType == ENodeTitleType::ListView)
+	if (TitleType == ENodeTitleType::MenuTitle)
 	{
 		return FText::FromName(InputActionName);
 	}
diff --git a/Engine/Source/Editor/BlueprintGraph/Private/K2Node_InputAxisEvent.cpp b/Engine/Source/Editor/BlueprintGraph/Private/K2Node_InputAxisEvent.cpp
--- a/Engine/Source/Editor/BlueprintGraph/Private/K2Node_InputAxisEvent.cpp
+++ b/Engine/Source/Editor/BlueprintGraph/Private/K2Node_InputAxisEvent.cpp
@@ -39,7 +39,7 @@ void UK2Node_InputAxisEvent::Initialize(const FName AxisName)
 
 FText UK2Node_InputAxisEvent::GetNodeTitle(ENodeTitleType::Type TitleType) const
 {
-	if (TitleType == ENodeTitleType::ListView)
+	if (TitleType == ENodeTitleType::MenuTitle)
 	{
 		return FText::FromName(InputAxisName);
 	}
diff --git a/Engine/Source/Editor/BlueprintGraph/Private/K2Node_InputTouch.cpp b/Engine/Source/Editor/BlueprintGraph/Private/K2Node_InputTouch.cpp
--- a/Engine/Source/Editor/BlueprintGraph/Private/K2Node_InputTouch.cpp
+++ b/Engine/Source/Editor/BlueprintGraph/Private/K2Node_InputTouch.cpp
@@ -61,7 +61,7 @@ FLinearColor UK2Node_InputTouch::GetNodeTitleColor() const
 FText UK2Node_InputTouch::GetNodeTitle(ENodeTitleType::Type TitleType) const
 {
 	FText Title = NSLOCTEXT("K2Node", "InputTouch_Name", "InputTouch");
-	if (TitleType == ENodeTitleType::ListView)
+	if (TitleType == ENodeTitleType::MenuTitle)
 	{
 		Title = NSLOCTEXT("K2Node", "InputTouch_ListTitle", "Touch");
 	}
diff --git a/Engine/Source/Editor/BlueprintGraph/Private/K2Node_Knot.cpp b/Engine/Source/Editor/BlueprintGraph/Private/K2Node_Knot.cpp
--- a/Engine/Source/Editor/BlueprintGraph/Private/K2Node_Knot.cpp
+++ b/Engine/Source/Editor/BlueprintGraph/Private/K2Node_Knot.cpp
@@ -42,7 +42,7 @@ FText UK2Node_Knot::GetNodeTitle(ENodeTitleType::Type TitleType) const
 	{
 		return FText::FromString(NodeComment);
 	}
-	else if (TitleType == ENodeTitleType::ListView)
+	else if (TitleType == ENodeTitleType::MenuTitle)
 	{
 		return LOCTEXT("KnotListTitle", "Add Reroute Node...");
 	}
diff --git a/Engine/Source/Editor/BlueprintGraph/Private/K2Node_Literal.cpp b/Engine/Source/Editor/BlueprintGraph/Private/K2Node_Literal.cpp
--- a/Engine/Source/Editor/BlueprintGraph/Private/K2Node_Literal.cpp
+++ b/Engine/Source/Editor/BlueprintGraph/Private/K2Node_Literal.cpp
@@ -118,7 +118,7 @@ FText UK2Node_Literal::GetTooltipText() const
 FText UK2Node_Literal::GetNodeTitle(ENodeTitleType::Type TitleType) const
 {
 	AActor* Actor = Cast<AActor>(ObjectRef);
-	if(TitleType == ENodeTitleType::ListView)
+	if(TitleType == ENodeTitleType::MenuTitle)
 	{
 		return NSLOCTEXT("K2Node", "Literal_Title", "Get Selected Actor Reference(s)");
 	}
diff --git a/Engine/Source/Editor/BlueprintGraph/Private/K2Node_LocalVariable.cpp b/Engine/Source/Editor/BlueprintGraph/Private/K2Node_LocalVariable.cpp
--- a/Engine/Source/Editor/BlueprintGraph/Private/K2Node_LocalVariable.cpp
+++ b/Engine/Source/Editor/BlueprintGraph/Private/K2Node_LocalVariable.cpp
@@ -31,7 +31,7 @@ FText UDEPRECATED_K2Node_LocalVariable::GetNodeTitle(ENodeTitleType::Type TitleT
 	{
 		return FText::FromName(CustomVariableName);
 	}
-	else if(TitleType == ENodeTitleType::ListView)
+	else if(TitleType == ENodeTitleType::ListView || TitleType == ENodeTitleType::MenuTitle)
 	{
 		FFormatNamedArguments Args;
 		Args.Add(TEXT("TypeName"), UEdGraphSchema_K2::TypeToText(VariableType));
diff --git a/Engine/Source/Editor/BlueprintGraph/Private/K2Node_MathExpression.cpp b/Engine/Source/Editor/BlueprintGraph/Private/K2Node_MathExpression.cpp
--- a/Engine/Source/Editor/BlueprintGraph/Private/K2Node_MathExpression.cpp
+++ b/Engine/Source/Editor/BlueprintGraph/Private/K2Node_MathExpression.cpp
@@ -2453,7 +2453,7 @@ void UK2Node_MathExpression::ValidateNodeDuringCompilation(FCompilerResultsLog&
 //------------------------------------------------------------------------------
 FText UK2Node_MathExpression::GetNodeTitle(ENodeTitleType::Type TitleType) const
 {
-	if (Expression.IsEmpty() && (TitleType == ENodeTitleType::ListView))
+	if (Expression.IsEmpty() && (TitleType == ENodeTitleType::MenuTitle))
 	{
 		return LOCTEXT("AddMathExprMenuOption", "Add Math Expression...");
 	}
diff --git a/Engine/Source/Editor/BlueprintGraph/Private/K2Node_Message.cpp b/Engine/Source/Editor/BlueprintGraph/Private/K2Node_Message.cpp
--- a/Engine/Source/Editor/BlueprintGraph/Private/K2Node_Message.cpp
+++ b/Engine/Source/Editor/BlueprintGraph/Private/K2Node_Message.cpp
@@ -19,7 +19,7 @@ FText UK2Node_Message::GetNodeTitle(ENodeTitleType::Type TitleType) const
 		if (!CachedNodeTitles.IsTitleCached(TitleType))
 		{
 			FText NodeNameText = FText::FromString(UK2Node_CallFunction::GetUserFacingFunctionName(Function));
-			if (TitleType == ENodeTitleType::ListView)
+			if (TitleType == ENodeTitleType::MenuTitle)
 			{
 				// FText::Format() is slow, so we cache this to save on performance
 				CachedNodeTitles.SetCachedTitle(TitleType, FText::Format(LOCTEXT("ListTitle", "{0} (Message)"), NodeNameText));
diff --git a/Engine/Source/Editor/BlueprintGraph/Private/K2Node_Self.cpp b/Engine/Source/Editor/BlueprintGraph/Private/K2Node_Self.cpp
--- a/Engine/Source/Editor/BlueprintGraph/Private/K2Node_Self.cpp
+++ b/Engine/Source/Editor/BlueprintGraph/Private/K2Node_Self.cpp
@@ -58,7 +58,7 @@ FString UK2Node_Self::GetKeywords() const
 FText UK2Node_Self::GetNodeTitle(ENodeTitleType::Type TitleType) const
 {
 	FText NodeTitle = NSLOCTEXT("K2Node", "SelfReferenceName", "Self-Reference");
-	if (TitleType == ENodeTitleType::ListView)
+	if (TitleType == ENodeTitleType::MenuTitle)
 	{
 		NodeTitle = LOCTEXT("ListTitle", "Get a reference to self");
 	}
diff --git a/Engine/Source/Editor/BlueprintGraph/Private/K2Node_SpawnActorFromClass.cpp b/Engine/Source/Editor/BlueprintGraph/Private/K2Node_SpawnActorFromClass.cpp
--- a/Engine/Source/Editor/BlueprintGraph/Private/K2Node_SpawnActorFromClass.cpp
+++ b/Engine/Source/Editor/BlueprintGraph/Private/K2Node_SpawnActorFromClass.cpp
@@ -260,7 +260,7 @@ FLinearColor UK2Node_SpawnActorFromClass::GetNodeTitleColor() const
 FText UK2Node_SpawnActorFromClass::GetNodeTitle(ENodeTitleType::Type TitleType) const
 {
 	FText NodeTitle = NSLOCTEXT("K2Node", "SpawnActor_BaseTitle", "Spawn Actor from Class");
-	if (TitleType != ENodeTitleType::ListView)
+	if (TitleType != ENodeTitleType::MenuTitle)
 	{
 		FText SpawnString = NSLOCTEXT("K2Node", "None", "NONE");
 		if (UEdGraphPin* ClassPin = GetClassPin())
diff --git a/Engine/Source/Editor/Kismet/Private/BlueprintActionMenuBuilder.cpp b/Engine/Source/Editor/Kismet/Private/BlueprintActionMenuBuilder.cpp
--- a/Engine/Source/Editor/Kismet/Private/BlueprintActionMenuBuilder.cpp
+++ b/Engine/Source/Editor/Kismet/Private/BlueprintActionMenuBuilder.cpp
@@ -204,7 +204,7 @@ FText FBlueprintActionMenuItemFactory::GetMenuNameForAction(TWeakPtr<FBlueprintE
 	{
 		if (UEdGraphNode* NodeTemplate = GetTemplateNode(Action, EditorContext))
 		{
-			MenuName = NodeTemplate->GetNodeTitle(ENodeTitleType::ListView);
+			MenuName = NodeTemplate->GetNodeTitle(ENodeTitleType::MenuTitle);
 		}
 		else
 		{
diff --git a/Engine/Source/Editor/UnrealEd/Private/MaterialGraphNode.cpp b/Engine/Source/Editor/UnrealEd/Private/MaterialGraphNode.cpp
--- a/Engine/Source/Editor/UnrealEd/Private/MaterialGraphNode.cpp
+++ b/Engine/Source/Editor/UnrealEd/Private/MaterialGraphNode.cpp
@@ -154,7 +154,7 @@ FText UMaterialGraphNode::GetNodeTitle(ENodeTitleType::Type TitleType) const
 	{
 		return FText::FromString(GetParameterName());
 	}
-	else if (TitleType == ENodeTitleType::ListView)
+	else if (TitleType == ENodeTitleType::ListView || TitleType == ENodeTitleType::MenuTitle)
 	{
 		return FText::FromString(MaterialExpression->GetClass()->GetDescription());
 	}
diff --git a/Engine/Source/Runtime/Engine/Classes/EdGraph/EdGraphNode.h b/Engine/Source/Runtime/Engine/Classes/EdGraph/EdGraphNode.h
--- a/Engine/Source/Runtime/Engine/Classes/EdGraph/EdGraphNode.h
+++ b/Engine/Source/Runtime/Engine/Classes/EdGraph/EdGraphNode.h
@@ -32,6 +32,8 @@ namespace ENodeTitleType
 		ListView,
 		// Returns the editable title (which might not be a title at all)
 		EditableTitle,
+		// Menu Title for context menus to be displayed in context menus referencing the node
+		MenuTitle,
 
 		MAX_TitleTypes,
 	};
diff --git a/Engine/Source/Runtime/Engine/Private/EdGraph/EdGraphNode_Comment.cpp b/Engine/Source/Runtime/Engine/Private/EdGraph/EdGraphNode_Comment.cpp
--- a/Engine/Source/Runtime/Engine/Private/EdGraph/EdGraphNode_Comment.cpp
+++ b/Engine/Source/Runtime/Engine/Private/EdGraph/EdGraphNode_Comment.cpp
@@ -76,16 +76,13 @@ FString UEdGraphNode_Comment::GetDocumentationExcerptName() const
 
 FText UEdGraphNode_Comment::GetNodeTitle(ENodeTitleType::Type TitleType) const
 {
-	if(TitleType == ENodeTitleType::ListView)
+	if(TitleType == ENodeTitleType::MenuTitle)
 	{
-		if (NodeComment.IsEmpty())
-		{
-			return NSLOCTEXT("K2Node", "NoComment_ListTitle", "Add Comment...");
-		}
-		else 
-		{
-			return NSLOCTEXT("K2Node", "CommentBlock_ListTitle", "Comment");
-		}		
+		return NSLOCTEXT("K2Node", "NoComment_ListTitle", "Add Comment...");
+	}
+	else if(TitleType == ENodeTitleType::ListView)
+	{
+		return NSLOCTEXT("K2Node", "CommentBlock_ListTitle", "Comment");	
 	}
 
 	return FText::FromString(NodeComment);
diff --git a/Engine/Source/Runtime/Engine/Private/EdGraph/EdGraphNode_Documentation.cpp b/Engine/Source/Runtime/Engine/Private/EdGraph/EdGraphNode_Documentation.cpp
--- a/Engine/Source/Runtime/Engine/Private/EdGraph/EdGraphNode_Documentation.cpp
+++ b/Engine/Source/Runtime/Engine/Private/EdGraph/EdGraphNode_Documentation.cpp
@@ -38,7 +38,7 @@ FText UEdGraphNode_Documentation::GetTooltipText() const
 FText UEdGraphNode_Documentation::GetNodeTitle( ENodeTitleType::Type TitleType ) const
 {
 	FText NodeTitle;
-	if (TitleType == ENodeTitleType::ListView)
+	if (TitleType == ENodeTitleType::MenuTitle)
 	{
 		NodeTitle = NSLOCTEXT("K2Node", "DocumentationBlock_ListTitle", "Add Documentation Node...");
 	}
