Stop transistion nodes being selected instead of the new state nodes when dragging from an existing state nodes

#TTP 343232 - ANIM GRAPH: Transition node is selected when State node is added

diff --git a/Engine/Source/Editor/AnimGraph/Classes/AnimationStateMachineSchema.h b/Engine/Source/Editor/AnimGraph/Classes/AnimationStateMachineSchema.h
--- a/Engine/Source/Editor/AnimGraph/Classes/AnimationStateMachineSchema.h
+++ b/Engine/Source/Editor/AnimGraph/Classes/AnimationStateMachineSchema.h
@@ -26,12 +26,12 @@ struct ANIMGRAPH_API FEdGraphSchemaAction_NewStateNode : public FEdGraphSchemaAc
 	virtual void AddReferencedObjects( FReferenceCollector& Collector ) override;
 
 	template <typename NodeType>
-	static NodeType* SpawnNodeFromTemplate(class UEdGraph* ParentGraph, NodeType* InTemplateNode, const FVector2D Location = FVector2D(0.0f, 0.0f))
+	static NodeType* SpawnNodeFromTemplate(class UEdGraph* ParentGraph, NodeType* InTemplateNode, const FVector2D Location = FVector2D(0.0f, 0.0f), bool bSelectNewNode = true)
 	{
 		FEdGraphSchemaAction_NewStateNode Action;
 		Action.NodeTemplate = InTemplateNode;
 
-		return Cast<NodeType>(Action.PerformAction(ParentGraph, NULL, Location));
+		return Cast<NodeType>(Action.PerformAction(ParentGraph, NULL, Location, bSelectNewNode));
 	}
 };
 
diff --git a/Engine/Source/Editor/AnimGraph/Private/AnimationStateMachineSchema.cpp b/Engine/Source/Editor/AnimGraph/Private/AnimationStateMachineSchema.cpp
--- a/Engine/Source/Editor/AnimGraph/Private/AnimationStateMachineSchema.cpp
+++ b/Engine/Source/Editor/AnimGraph/Private/AnimationStateMachineSchema.cpp
@@ -241,7 +241,7 @@ bool UAnimationStateMachineSchema::CreateAutomaticConversionNodeAndConnections(U
 		&& (NodeA->GetInputPin() != NULL) && (NodeA->GetOutputPin() != NULL)
 		&& (NodeB->GetInputPin() != NULL) && (NodeB->GetOutputPin() != NULL))
 	{
-		UAnimStateTransitionNode* TransitionNode = FEdGraphSchemaAction_NewStateNode::SpawnNodeFromTemplate<UAnimStateTransitionNode>(NodeA->GetGraph(), NewObject<UAnimStateTransitionNode>());
+		UAnimStateTransitionNode* TransitionNode = FEdGraphSchemaAction_NewStateNode::SpawnNodeFromTemplate<UAnimStateTransitionNode>(NodeA->GetGraph(), NewObject<UAnimStateTransitionNode>(), FVector2D(0.0f, 0.0f), false);
 
 		if (PinA->Direction == EGPD_Output)
 		{
