TTP # 316516 : POLISH: ANIM: AnimBlueprint: Two Bone IK node is confusing

diff --git a/Engine/Source/Editor/AnimGraph/Classes/AnimGraphNode_Base.h b/Engine/Source/Editor/AnimGraph/Classes/AnimGraphNode_Base.h
--- a/Engine/Source/Editor/AnimGraph/Classes/AnimGraphNode_Base.h
+++ b/Engine/Source/Editor/AnimGraph/Classes/AnimGraphNode_Base.h
@@ -174,6 +174,9 @@ class UAnimGraphNode_Base : public UK2Node
 	virtual UScriptStruct* GetTimePropertyStruct() const { return NULL; }
 	// END Interface to support transition getter
 
+	// can customize details tab 
+	virtual void CustomizeDetails(class IDetailLayoutBuilder& DetailBuilder){ }
+
 protected:
 	friend class FAnimBlueprintCompiler;
 	friend class FAnimGraphNodeDetails;
diff --git a/Engine/Source/Editor/AnimGraph/Classes/AnimGraphNode_TwoBoneIK.h b/Engine/Source/Editor/AnimGraph/Classes/AnimGraphNode_TwoBoneIK.h
--- a/Engine/Source/Editor/AnimGraph/Classes/AnimGraphNode_TwoBoneIK.h
+++ b/Engine/Source/Editor/AnimGraph/Classes/AnimGraphNode_TwoBoneIK.h
@@ -20,6 +20,10 @@ public:
 	virtual FText GetTooltipText() const override;
 	// End of UEdGraphNode interface
 
+	// UAnimGraphNode_Base interface
+	virtual void CustomizeDetails(class IDetailLayoutBuilder& DetailBuilder) override;
+	// End of UAnimGraphNode_Base interface
+
 	// UAnimGraphNode_SkeletalControlBase interface
 	ANIMGRAPH_API virtual void Draw( FPrimitiveDrawInterface* PDI, USkeletalMeshComponent* SkelMeshComp) const override;
 	// End of UAnimGraphNode_SkeletalControlBase interface
@@ -33,6 +37,9 @@ protected:
 	void ConvertToComponentSpaceTransform(USkeletalMeshComponent* SkelComp, USkeleton * Skeleton, const FTransform & InTransform, FTransform & OutCSTransform, int32 BoneIndex, uint8 Space) const;
 	void DrawTargetLocation(FPrimitiveDrawInterface* PDI, USkeletalMeshComponent* SkelComp, USkeleton * Skeleton, uint8 SpaceBase, FName SpaceBoneName, const FVector & TargetLocation, const FColor & TargetColor, const FColor & BoneColor) const;
 
+	// make Pins showed / hidden by options
+	void SetPinsVisibility(bool bShow);
+
 private:
 	/** Constructing FText strings can be costly, so we cache the node's title */
 	FNodeTitleTextTable CachedNodeTitles;
diff --git a/Engine/Source/Editor/AnimGraph/Private/AnimGraphNode_TwoBoneIK.cpp b/Engine/Source/Editor/AnimGraph/Private/AnimGraphNode_TwoBoneIK.cpp
--- a/Engine/Source/Editor/AnimGraph/Private/AnimGraphNode_TwoBoneIK.cpp
+++ b/Engine/Source/Editor/AnimGraph/Private/AnimGraphNode_TwoBoneIK.cpp
@@ -3,10 +3,33 @@
 #include "AnimGraphPrivatePCH.h"
 #include "AnimGraphNode_TwoBoneIK.h"
 
+// for customization details
+#include "../../PropertyEditor/Public/PropertyHandle.h"
+#include "../../PropertyEditor/Public/DetailLayoutBuilder.h"
+#include "../../PropertyEditor/Public/DetailCategoryBuilder.h"
+#include "../../PropertyEditor/Public/DetailWidgetRow.h"
+#include "../../PropertyEditor/Public/IDetailPropertyRow.h"
+
+#define LOCTEXT_NAMESPACE "A3Nodes"
+
+/////////////////////////////////////////////////////
+// FTwoBoneIKDelegate
+
+class FTwoBoneIKDelegate : public TSharedFromThis<FTwoBoneIKDelegate>
+{
+public:
+	void UpdateLocationSpace(class IDetailLayoutBuilder* DetailBuilder)
+	{
+		if (DetailBuilder)
+		{
+			DetailBuilder->ForceRefreshDetails();
+		}
+	}
+};
+
 /////////////////////////////////////////////////////
 // UAnimGraphNode_TwoBoneIK
 
-#define LOCTEXT_NAMESPACE "A3Nodes"
 
 UAnimGraphNode_TwoBoneIK::UAnimGraphNode_TwoBoneIK(const FPostConstructInitializeProperties& PCIP)
 	: Super(PCIP)
@@ -50,12 +73,12 @@ FText UAnimGraphNode_TwoBoneIK::GetNodeTitle(ENodeTitleType::Type TitleType) con
 
 void UAnimGraphNode_TwoBoneIK::Draw( FPrimitiveDrawInterface* PDI, USkeletalMeshComponent* SkelMeshComp ) const
 {
-	if ( SkelMeshComp && SkelMeshComp->SkeletalMesh && SkelMeshComp->SkeletalMesh->Skeleton )
+	if (SkelMeshComp && SkelMeshComp->SkeletalMesh && SkelMeshComp->SkeletalMesh->Skeleton)
 	{
 		USkeleton * Skeleton = SkelMeshComp->SkeletalMesh->Skeleton;
 
-		DrawTargetLocation(PDI, SkelMeshComp, Skeleton, Node.EffectorLocationSpace, Node.EffectorSpaceBoneName,Node.EffectorLocation, FColor(255,128,128), FColor(180,128,128));
-		DrawTargetLocation(PDI, SkelMeshComp, Skeleton, Node.JointTargetLocationSpace, Node.JointTargetSpaceBoneName,Node.JointTargetLocation, FColor(128, 255, 128), FColor(128,180,128));
+		DrawTargetLocation(PDI, SkelMeshComp, Skeleton, Node.EffectorLocationSpace, Node.EffectorSpaceBoneName, Node.EffectorLocation, FColor(255, 128, 128), FColor(180, 128, 128));
+		DrawTargetLocation(PDI, SkelMeshComp, Skeleton, Node.JointTargetLocationSpace, Node.JointTargetSpaceBoneName, Node.JointTargetLocation, FColor(128, 255, 128), FColor(128, 180, 128));
 	}
 }
 
@@ -149,5 +172,94 @@ void UAnimGraphNode_TwoBoneIK::ConvertToComponentSpaceTransform(USkeletalMeshCom
 	}
 }
 
+// just for refreshing UIs when bone space was changed
+TSharedPtr<FTwoBoneIKDelegate> TwoBoneIKDelegate;
+
+void UAnimGraphNode_TwoBoneIK::CustomizeDetails(class IDetailLayoutBuilder& DetailBuilder)
+{
+	// initialize just once
+	if (!TwoBoneIKDelegate.IsValid())
+	{
+		TwoBoneIKDelegate = MakeShareable(new FTwoBoneIKDelegate());
+	}
+
+	EBoneControlSpace Space = Node.EffectorLocationSpace;
+	if (Space == BCS_BoneSpace || Space == BCS_ParentBoneSpace)
+	{
+		IDetailCategoryBuilder& IKCategory = DetailBuilder.EditCategory("IK");
+		IDetailCategoryBuilder& EffectorCategory = DetailBuilder.EditCategory("EndEffector");
+		TSharedPtr<IPropertyHandle> PropertyHandle;
+		PropertyHandle = DetailBuilder.GetProperty(FName("Node.bTakeRotationFromEffectorSpace"), GetClass());
+		EffectorCategory.AddProperty(PropertyHandle);
+		PropertyHandle = DetailBuilder.GetProperty(FName("Node.bMaintainEffectorRelRot"), GetClass());
+		EffectorCategory.AddProperty(PropertyHandle);
+		PropertyHandle = DetailBuilder.GetProperty(FName("Node.EffectorSpaceBoneName"), GetClass());
+		EffectorCategory.AddProperty(PropertyHandle);
+	}
+	else // hide all properties in EndEffector category
+	{
+		TSharedPtr<IPropertyHandle> PropertyHandle = DetailBuilder.GetProperty(FName("Node.EffectorLocation"), GetClass());
+		DetailBuilder.HideProperty(PropertyHandle);
+		PropertyHandle = DetailBuilder.GetProperty(FName("Node.bTakeRotationFromEffectorSpace"), GetClass());
+		DetailBuilder.HideProperty(PropertyHandle);
+		PropertyHandle = DetailBuilder.GetProperty(FName("Node.bMaintainEffectorRelRot"), GetClass());
+		DetailBuilder.HideProperty(PropertyHandle);
+		PropertyHandle = DetailBuilder.GetProperty(FName("Node.EffectorSpaceBoneName"), GetClass());
+		DetailBuilder.HideProperty(PropertyHandle);
+	}
+
+	Space = Node.JointTargetLocationSpace;
+	if (Space == BCS_BoneSpace || Space == BCS_ParentBoneSpace)
+	{
+		IDetailCategoryBuilder& IKCategory = DetailBuilder.EditCategory("IK");
+		IDetailCategoryBuilder& EffectorCategory = DetailBuilder.EditCategory("JointTarget");
+		TSharedPtr<IPropertyHandle> PropertyHandle;
+		PropertyHandle = DetailBuilder.GetProperty(FName("Node.JointTargetSpaceBoneName"), GetClass());
+		EffectorCategory.AddProperty(PropertyHandle);
+
+		SetPinsVisibility(true);
+	}
+	else // hide all properties in JointTarget category except for JointTargetLocationSpace
+	{
+		TSharedPtr<IPropertyHandle> PropertyHandle = DetailBuilder.GetProperty(FName("Node.JointTargetLocation"), GetClass());
+		DetailBuilder.HideProperty(PropertyHandle);
+		PropertyHandle = DetailBuilder.GetProperty(FName("Node.JointTargetSpaceBoneName"), GetClass());
+		DetailBuilder.HideProperty(PropertyHandle);
+
+		SetPinsVisibility(false);
+	}
+
+
+	// refresh UIs when bone space is changed
+	TSharedRef<IPropertyHandle> EffectorLocHandle = DetailBuilder.GetProperty(FName("Node.EffectorLocationSpace"), GetClass());
+	if (EffectorLocHandle->IsValidHandle())
+	{
+		FSimpleDelegate UpdateEffectorSpaceDelegate = FSimpleDelegate::CreateSP(TwoBoneIKDelegate.Get(), &FTwoBoneIKDelegate::UpdateLocationSpace, &DetailBuilder);
+		EffectorLocHandle->SetOnPropertyValueChanged(UpdateEffectorSpaceDelegate);
+	}
+
+	TSharedRef<IPropertyHandle> JointTragetLocHandle = DetailBuilder.GetProperty(FName("Node.JointTargetLocationSpace"), GetClass());
+	if (JointTragetLocHandle->IsValidHandle())
+	{
+		FSimpleDelegate UpdateJointSpaceDelegate = FSimpleDelegate::CreateSP(TwoBoneIKDelegate.Get(), &FTwoBoneIKDelegate::UpdateLocationSpace, &DetailBuilder);
+		JointTragetLocHandle->SetOnPropertyValueChanged(UpdateJointSpaceDelegate);
+	}
+
+	// reconstruct node for showing/hiding Pins
+	ReconstructNode();
+}
+
+void UAnimGraphNode_TwoBoneIK::SetPinsVisibility(bool bShow)
+{
+	for (FOptionalPinFromProperty& Pin : ShowPinForProperties)
+	{
+		if (Pin.PropertyName == "JointTargetLocation")
+		{
+			PreEditChange(NULL);
+			Pin.bShowPin = bShow;
+			break;
+		}
+	}
+}
 
 #undef LOCTEXT_NAMESPACE
diff --git a/Engine/Source/Editor/Persona/Private/AnimGraphNodeDetails.cpp b/Engine/Source/Editor/Persona/Private/AnimGraphNodeDetails.cpp
--- a/Engine/Source/Editor/Persona/Private/AnimGraphNodeDetails.cpp
+++ b/Engine/Source/Editor/Persona/Private/AnimGraphNodeDetails.cpp
@@ -74,6 +74,9 @@ void FAnimGraphNodeDetails::CustomizeDetails(class IDetailLayoutBuilder& DetailB
 		return;
 	}
 
+	// customize anim graph node's own details if needed
+	AnimGraphNode->CustomizeDetails(DetailBuilder);
+
 	// Now customize each property in the pins array
 	for (int CustomPinIndex = 0; CustomPinIndex < AnimGraphNode->ShowPinForProperties.Num(); ++CustomPinIndex)
 	{
diff --git a/Engine/Source/Runtime/Engine/Classes/Animation/BoneControllers/AnimNode_TwoBoneIK.h b/Engine/Source/Runtime/Engine/Classes/Animation/BoneControllers/AnimNode_TwoBoneIK.h
--- a/Engine/Source/Runtime/Engine/Classes/Animation/BoneControllers/AnimNode_TwoBoneIK.h
+++ b/Engine/Source/Runtime/Engine/Classes/Animation/BoneControllers/AnimNode_TwoBoneIK.h
@@ -23,7 +23,7 @@ struct ENGINE_API FAnimNode_TwoBoneIK : public FAnimNode_SkeletalControlBase
 	FVector EffectorLocation;
 
 	/** Joint Target Location. Location used to orient Joint bone. **/
-	UPROPERTY(EditAnywhere, BlueprintReadWrite, Category=JointTarget, meta=(PinShownByDefault))
+	UPROPERTY(EditAnywhere, BlueprintReadWrite, Category = JointTarget, meta=(PinShownByDefault))
 	FVector JointTargetLocation;
 
 	/** Limits to use if stretching is allowed. */
@@ -47,7 +47,7 @@ struct ENGINE_API FAnimNode_TwoBoneIK : public FAnimNode_SkeletalControlBase
 	uint32 bAllowStretching:1;
 	
 	/** Reference frame of Effector Location. */
-	UPROPERTY(EditAnywhere, BlueprintReadWrite, Category=EndEffector)
+	UPROPERTY(EditAnywhere, BlueprintReadWrite, Category = IK)
 	TEnumAsByte<enum EBoneControlSpace> EffectorLocationSpace;
 
 	/** Reference frame of Joint Target Location. */
