Fixed state machine graph node showing up with wrong title in new blueprint context menu. Would show "Error: No Graph" prior to this fix.

diff --git a/Engine/Source/Editor/AnimGraph/Private/AnimGraphNode_StateMachineBase.cpp b/Engine/Source/Editor/AnimGraph/Private/AnimGraphNode_StateMachineBase.cpp
--- a/Engine/Source/Editor/AnimGraph/Private/AnimGraphNode_StateMachineBase.cpp
+++ b/Engine/Source/Editor/AnimGraph/Private/AnimGraphNode_StateMachineBase.cpp
@@ -59,7 +59,7 @@ FText UAnimGraphNode_StateMachineBase::GetTooltipText() const
 
 FText UAnimGraphNode_StateMachineBase::GetNodeTitle(ENodeTitleType::Type TitleType) const
 {
-	if ((TitleType == ENodeTitleType::ListView) && (EditorStateMachineGraph == nullptr))
+	if ((TitleType == ENodeTitleType::MenuTitle || TitleType == ENodeTitleType::ListView) && (EditorStateMachineGraph == nullptr))
 	{
 		return LOCTEXT("AddNewStateMachine", "Add New State Machine...");
 	}
