"Use Cached Pose" works with the new Blueprint context menu system.

No longer has issues with renaming the "Save Cached Pose" node.

Fixed up the details panel to not show unneeded data.

diff --git a/Engine/Source/Editor/AnimGraph/Classes/AnimGraphNode_UseCachedPose.h b/Engine/Source/Editor/AnimGraph/Classes/AnimGraphNode_UseCachedPose.h
--- a/Engine/Source/Editor/AnimGraph/Classes/AnimGraphNode_UseCachedPose.h
+++ b/Engine/Source/Editor/AnimGraph/Classes/AnimGraphNode_UseCachedPose.h
@@ -12,24 +12,33 @@ class UAnimGraphNode_UseCachedPose : public UAnimGraphNode_Base
 	GENERATED_UCLASS_BODY()
 
 public:
-	UPROPERTY(EditAnywhere, Category=Settings)
+	UPROPERTY()
 	FAnimNode_UseCachedPose Node;
 
-	UPROPERTY(EditAnywhere, Category=CachedPose)
-	FString NameOfCache;
+	UPROPERTY()
+	TWeakObjectPtr<UAnimGraphNode_SaveCachedPose> SaveCachedPoseNode;
 
 public:
+	// Begin UObject interface
+	virtual void PostLoad() override;
+	// End UObject interface
+
 	// UEdGraphNode interface
 	virtual FText GetTooltipText() const override;
 	virtual FText GetNodeTitle(ENodeTitleType::Type TitleType) const override;
 	// End of UEdGraphNode interface
 
+	// UK2Node interface.
+	virtual void GetMenuActions(FBlueprintActionDatabaseRegistrar& ActionRegistrar) const override;
+	virtual bool IsActionFilteredOut(class FBlueprintActionFilter const& Filter) override;
+	// End of UK2Node interface
+
 	// UAnimGraphNode_Base interface
 	virtual FString GetNodeCategory() const override;
 	virtual void GetMenuEntries(FGraphContextMenuBuilder& ContextMenuBuilder) const override;
 	// End of UAnimGraphNode_Base interface
 
 private:
-	/** Constructing FText strings can be costly, so we cache the node's title */
-	FNodeTextCache CachedNodeTitle;
+	UPROPERTY()
+	mutable FString NameOfCache;
 };
diff --git a/Engine/Source/Editor/AnimGraph/Private/AnimGraphNode_UseCachedPose.cpp b/Engine/Source/Editor/AnimGraph/Private/AnimGraphNode_UseCachedPose.cpp
--- a/Engine/Source/Editor/AnimGraph/Private/AnimGraphNode_UseCachedPose.cpp
+++ b/Engine/Source/Editor/AnimGraph/Private/AnimGraphNode_UseCachedPose.cpp
@@ -7,6 +7,9 @@
 #include "K2ActionMenuBuilder.h" // for FK2ActionMenuBuilder::AddNewNodeAction()
 #include "AnimGraphNode_SaveCachedPose.h"
 #include "AnimGraphNode_UseCachedPose.h"
+#include "BlueprintNodeSpawner.h"
+#include "BlueprintActionDatabaseRegistrar.h"
+#include "BlueprintActionFilter.h"
 
 /////////////////////////////////////////////////////
 // UAnimGraphNode_UseCachedPose
@@ -18,6 +21,37 @@ UAnimGraphNode_UseCachedPose::UAnimGraphNode_UseCachedPose(const FPostConstructI
 {
 }
 
+void UAnimGraphNode_UseCachedPose::PostLoad()
+{
+	Super::PostLoad();
+
+	// If there is no SaveCachedPose node set but there is a NameOfCache available, we may be updating to the new system
+	// Go through and find the cached node, if possible.
+	if(!SaveCachedPoseNode.IsValid() && !NameOfCache.IsEmpty())
+	{
+		TArray<UEdGraph*> AllAnimationGraphs;
+		GetGraph()->GetAllChildrenGraphs(AllAnimationGraphs);
+		AllAnimationGraphs.Add(GetGraph());
+
+		for(UEdGraph* Graph : AllAnimationGraphs)
+		{
+			// Get a list of all save cached pose nodes
+			TArray<UAnimGraphNode_SaveCachedPose*> CachedPoseNodes;
+			Graph->GetNodesOfClass(CachedPoseNodes);
+
+			// Go through all the nodes and find one with a title that matches ours
+			for (auto NodeIt = CachedPoseNodes.CreateIterator(); NodeIt; ++NodeIt)
+			{
+				if((*NodeIt)->CacheName == NameOfCache)
+				{
+					SaveCachedPoseNode = *NodeIt;
+					break;
+				}
+			}
+		}
+	}
+}
+
 FText UAnimGraphNode_UseCachedPose::GetTooltipText() const
 {
 	return LOCTEXT("AnimGraphNode_UseCachedPose_Tooltip", "References an animation tree elsewhere in the blueprint, which will be evaluated at most once per frame.");
@@ -25,16 +59,13 @@ FText UAnimGraphNode_UseCachedPose::GetTooltipText() const
 
 FText UAnimGraphNode_UseCachedPose::GetNodeTitle(ENodeTitleType::Type TitleType) const
 {
-	// @TODO: don't know enough about this node type to comfortably assert that
-	//        the NameOfCache won't change after the node has spawned... until
-	//        then, we'll leave this optimization off
-	//if (CachedNodeTitle.IsOutOfDate())
+	FFormatNamedArguments Args;
+	if(SaveCachedPoseNode.IsValid())
 	{
-		FFormatNamedArguments Args;
-		Args.Add(TEXT("CachePoseName"), FText::FromString(NameOfCache));
-		CachedNodeTitle = FText::Format(LOCTEXT("AnimGraphNode_UseCachedPose_Title", "Use cached pose '{CachePoseName}'"), Args);
+		NameOfCache = SaveCachedPoseNode->CacheName;
 	}
-	return CachedNodeTitle;
+	Args.Add(TEXT("CachePoseName"), FText::FromString(NameOfCache));
+	return FText::Format(LOCTEXT("AnimGraphNode_UseCachedPose_Title", "Use cached pose '{CachePoseName}'"), Args);
 }
 
 FString UAnimGraphNode_UseCachedPose::GetNodeCategory() const
@@ -56,6 +87,7 @@ void UAnimGraphNode_UseCachedPose::GetMenuEntries(FGraphContextMenuBuilder& Cont
 		{
 			UAnimGraphNode_UseCachedPose* UseCachedPose = NewObject<UAnimGraphNode_UseCachedPose>();
 			UseCachedPose->NameOfCache = (*NodeIt)->CacheName;
+			UseCachedPose->SaveCachedPoseNode = *NodeIt;
 
 			TSharedPtr<FEdGraphSchemaAction_K2NewNode> UseCachedPoseAction = FK2ActionMenuBuilder::AddNewNodeAction(ContextMenuBuilder, GetNodeCategory(), UseCachedPose->GetNodeTitle(ENodeTitleType::ListView), UseCachedPose->GetTooltipText().ToString(), 0, UseCachedPose->GetKeywords());
 			UseCachedPoseAction->NodeTemplate = UseCachedPose;
@@ -63,4 +95,52 @@ void UAnimGraphNode_UseCachedPose::GetMenuEntries(FGraphContextMenuBuilder& Cont
 	}
 }
 
+
+void UAnimGraphNode_UseCachedPose::GetMenuActions(FBlueprintActionDatabaseRegistrar& ActionRegistrar) const
+{
+	auto PostSpawnSetupLambda = [](UEdGraphNode* NewNode, bool bIsTemplateNode, FString CacheNodeName, UAnimGraphNode_SaveCachedPose* SaveCachePoseNode)
+	{
+		UAnimGraphNode_UseCachedPose* UseCachedPose = CastChecked<UAnimGraphNode_UseCachedPose>(NewNode);
+		// we use an empty CacheName in GetNodeTitle() to relay the proper menu title
+		UseCachedPose->SaveCachedPoseNode = SaveCachePoseNode;
+	};
+
+
+	UObject const* ActionKey = ActionRegistrar.GetActionKeyFilter();
+
+	if(UBlueprint const* Blueprint = Cast<UBlueprint>(ActionKey))
+	{
+		// Get a list of all save cached pose nodes
+		TArray<UAnimGraphNode_SaveCachedPose*> CachedPoseNodes;
+		FBlueprintEditorUtils::GetAllNodesOfClass<UAnimGraphNode_SaveCachedPose>(Blueprint, /*out*/ CachedPoseNodes);
+
+		// Offer a use node for each of them
+		for (auto NodeIt = CachedPoseNodes.CreateIterator(); NodeIt; ++NodeIt)
+		{
+			UBlueprintNodeSpawner* NodeSpawner = UBlueprintNodeSpawner::Create(GetClass());
+			NodeSpawner->CustomizeNodeDelegate = UBlueprintNodeSpawner::FCustomizeNodeDelegate::CreateStatic(PostSpawnSetupLambda, (*NodeIt)->CacheName, *NodeIt);
+
+			ActionRegistrar.AddBlueprintAction(ActionKey, NodeSpawner);
+		}
+	}
+}
+
+bool UAnimGraphNode_UseCachedPose::IsActionFilteredOut(class FBlueprintActionFilter const& Filter)
+{
+	bool bIsFilteredOut = false;
+	if(SaveCachedPoseNode.IsValid())
+	{
+		FBlueprintActionContext const& FilterContext = Filter.Context;
+		for(UBlueprint* Blueprint : FilterContext.Blueprints)
+		{
+			if(SaveCachedPoseNode->GetBlueprint() != Blueprint)
+			{
+				bIsFilteredOut = true;
+				break;
+			}
+		}
+	}
+	return bIsFilteredOut;
+}
+
 #undef LOCTEXT_NAMESPACE
\ No newline at end of file
diff --git a/Engine/Source/Editor/KismetCompiler/Private/AnimBlueprintCompiler.cpp b/Engine/Source/Editor/KismetCompiler/Private/AnimBlueprintCompiler.cpp
--- a/Engine/Source/Editor/KismetCompiler/Private/AnimBlueprintCompiler.cpp
+++ b/Engine/Source/Editor/KismetCompiler/Private/AnimBlueprintCompiler.cpp
@@ -385,19 +385,26 @@ void FAnimBlueprintCompiler::ProcessAnimationNode(UAnimGraphNode_Base* VisualAni
 
 void FAnimBlueprintCompiler::ProcessUseCachedPose(UAnimGraphNode_UseCachedPose* UseCachedPose)
 {
+	bool bSuccessful = false;
+
 	// Link to the saved cached pose
-	if (UAnimGraphNode_SaveCachedPose* AssociatedSaveNode = SaveCachedPoseNodes.FindRef(UseCachedPose->NameOfCache))
+	if(UseCachedPose->SaveCachedPoseNode.IsValid())
 	{
-		UStructProperty* LinkProperty = FindField<UStructProperty>(FAnimNode_UseCachedPose::StaticStruct(), TEXT("LinkToCachingNode"));
-		check(LinkProperty);
-
-		FPoseLinkMappingRecord LinkRecord = FPoseLinkMappingRecord::MakeFromMember(UseCachedPose, AssociatedSaveNode, LinkProperty);
-		if (LinkRecord.IsValid())
+		if (UAnimGraphNode_SaveCachedPose* AssociatedSaveNode = SaveCachedPoseNodes.FindRef(UseCachedPose->SaveCachedPoseNode->CacheName))
 		{
-			ValidPoseLinkList.Add(LinkRecord);
+			UStructProperty* LinkProperty = FindField<UStructProperty>(FAnimNode_UseCachedPose::StaticStruct(), TEXT("LinkToCachingNode"));
+			check(LinkProperty);
+
+			FPoseLinkMappingRecord LinkRecord = FPoseLinkMappingRecord::MakeFromMember(UseCachedPose, AssociatedSaveNode, LinkProperty);
+			if (LinkRecord.IsValid())
+			{
+				ValidPoseLinkList.Add(LinkRecord);
+			}
+			bSuccessful = true;
 		}
 	}
-	else
+	
+	if(!bSuccessful)
 	{
 		MessageLog.Error(*LOCTEXT("NoAssociatedSaveNode", "@@ does not have an associated Save Cached Pose node").ToString(), UseCachedPose);
 	}
