Added state blend percent accessor to new context menu

diff --git a/Engine/Source/Editor/AnimGraph/Classes/K2Node_TransitionRuleGetter.h b/Engine/Source/Editor/AnimGraph/Classes/K2Node_TransitionRuleGetter.h
--- a/Engine/Source/Editor/AnimGraph/Classes/K2Node_TransitionRuleGetter.h
+++ b/Engine/Source/Editor/AnimGraph/Classes/K2Node_TransitionRuleGetter.h
@@ -34,15 +34,21 @@ class UK2Node_TransitionRuleGetter : public UK2Node
 	class UAnimGraphNode_Base* AssociatedAnimAssetPlayerNode;
 
 	UPROPERTY()
-	UObject* AssociatedStateNode;
+	UAnimStateNode* AssociatedStateNode;
 
 	// UEdGraphNode interface
 	virtual void AllocateDefaultPins() override;
 	virtual FText GetNodeTitle(ENodeTitleType::Type TitleType) const override;
 	virtual FText GetTooltipText() const override;
-	virtual bool ShowPaletteIconOnNode() const override{ return false; }
+	virtual bool ShowPaletteIconOnNode() const override { return false; }
+	virtual bool CanCreateUnderSpecifiedSchema(const UEdGraphSchema* Schema) const override;
+	virtual bool IsActionFilteredOut(class FBlueprintActionFilter const& Filter) override;
 	// End of UEdGraphNode interface
 
+	// UK2Node interface
+	virtual void GetMenuActions(FBlueprintActionDatabaseRegistrar& ActionRegistrar) const override;
+	// end of UK2Node interface
+
 	// @todo document
 	ANIMGRAPH_API UEdGraphPin* GetOutputPin() const;
 
diff --git a/Engine/Source/Editor/AnimGraph/Private/K2Node_TransitionRuleGetter.cpp b/Engine/Source/Editor/AnimGraph/Private/K2Node_TransitionRuleGetter.cpp
--- a/Engine/Source/Editor/AnimGraph/Private/K2Node_TransitionRuleGetter.cpp
+++ b/Engine/Source/Editor/AnimGraph/Private/K2Node_TransitionRuleGetter.cpp
@@ -1,11 +1,15 @@
 // Copyright 1998-2014 Epic Games, Inc. All Rights Reserved.
 
-
 #include "AnimGraphPrivatePCH.h"
-#include "K2Node_TransitionRuleGetter.h"
-#include "BlueprintUtilities.h"
+
+#include "AnimationGraphSchema.h"
 #include "AnimGraphNode_Base.h"
 #include "AnimStateNode.h"
+#include "BlueprintActionDatabaseRegistrar.h"
+#include "BlueprintActionFilter.h"
+#include "BlueprintNodeSpawner.h"
+#include "BlueprintUtilities.h"
+#include "K2Node_TransitionRuleGetter.h"
 /////////////////////////////////////////////////////
 // UK2Node_TransitionRuleGetter
 
@@ -110,11 +114,70 @@ FText UK2Node_TransitionRuleGetter::GetNodeTitle(ENodeTitleType::Type TitleType)
 	return Super::GetNodeTitle(TitleType);
 }
 
+void UK2Node_TransitionRuleGetter::GetMenuActions(FBlueprintActionDatabaseRegistrar& ActionRegistrar) const
+{
+	if( const UAnimBlueprint* AnimBlueprint = Cast<UAnimBlueprint>( ActionRegistrar.GetActionKeyFilter() ) )
+	{
+		TArray<UAnimStateNode*> States;
+		FBlueprintEditorUtils::GetAllNodesOfClass(AnimBlueprint, /*out*/ States);
+
+		for (auto StateIt = States.CreateIterator(); StateIt; ++StateIt)
+		{
+			UAnimStateNode* StateNode = *StateIt;
+			
+			auto UiSpecOverride = [](const FBlueprintActionContext& /*Context*/, const IBlueprintNodeBinder::FBindingSet& Bindings, FBlueprintActionUiSpec* UiSpecOut, UAnimStateNode* StateNode)
+			{
+				const FString OwnerName = StateNode->GetOuter()->GetName();
+				UiSpecOut->MenuName = FText::Format(LOCTEXT("TransitionRuleGetterTitle", "Current {0} for state '{1}.{2}'"), 
+														UK2Node_TransitionRuleGetter::GetFriendlyName(ETransitionGetter::ArbitraryState_GetBlendWeight), 
+														FText::FromString(OwnerName), 
+														FText::FromString(StateNode->GetStateName()));
+			};
+
+			auto PostSpawnSetupLambda = [](UEdGraphNode* NewNode, bool /*bIsTemplateNode*/, UAnimStateNode* StateNode)
+			{
+				UK2Node_TransitionRuleGetter* NewNodeTyped = CastChecked<UK2Node_TransitionRuleGetter>(NewNode);
+				NewNodeTyped->AssociatedStateNode = StateNode;
+				NewNodeTyped->GetterType = ETransitionGetter::ArbitraryState_GetBlendWeight;
+			};
+
+			UBlueprintNodeSpawner* Spawner = UBlueprintNodeSpawner::Create( UK2Node_TransitionRuleGetter::StaticClass(), nullptr, UBlueprintNodeSpawner::FCustomizeNodeDelegate::CreateStatic(PostSpawnSetupLambda, StateNode) );
+			Spawner->DynamicUiSignatureGetter = UBlueprintNodeSpawner::FUiSpecOverrideDelegate::CreateStatic(UiSpecOverride, StateNode);
+			ActionRegistrar.AddBlueprintAction( AnimBlueprint, Spawner );
+		}
+	}
+}
+
 FText UK2Node_TransitionRuleGetter::GetTooltipText() const
 {
 	return GetNodeTitle(ENodeTitleType::FullTitle);
 }
 
+bool UK2Node_TransitionRuleGetter::CanCreateUnderSpecifiedSchema(const UEdGraphSchema* Schema) const
+{
+	return Cast<UAnimationGraphSchema>(Schema) != NULL;
+}
+
+bool UK2Node_TransitionRuleGetter::IsActionFilteredOut(class FBlueprintActionFilter const& Filter)
+{
+	// only show the transition nodes if the associated state node is in every graph:
+	if( AssociatedStateNode )
+	{
+		for( auto Blueprint : Filter.Context.Blueprints )
+		{
+			TArray<UAnimStateNode*> States;
+			FBlueprintEditorUtils::GetAllNodesOfClass(Blueprint, /*out*/ States);
+
+			if( !States.Contains( AssociatedStateNode ) )
+			{
+				return true;
+			}
+		}
+		return false;
+	}
+	return true;
+}
+
 UEdGraphPin* UK2Node_TransitionRuleGetter::GetOutputPin() const
 {
 	return FindPinChecked(TEXT("Output"));
