Fixes for UseCachedPose node so that it can better find it's SaveCachedPose node and upgrade from the old system.

diff --git a/Engine/Source/Editor/AnimGraph/Classes/AnimGraphNode_UseCachedPose.h b/Engine/Source/Editor/AnimGraph/Classes/AnimGraphNode_UseCachedPose.h
--- a/Engine/Source/Editor/AnimGraph/Classes/AnimGraphNode_UseCachedPose.h
+++ b/Engine/Source/Editor/AnimGraph/Classes/AnimGraphNode_UseCachedPose.h
@@ -16,13 +16,9 @@ public:
 	FAnimNode_UseCachedPose Node;
 
 	UPROPERTY()
-	TWeakObjectPtr<UAnimGraphNode_SaveCachedPose> SaveCachedPoseNode;
+	mutable TWeakObjectPtr<UAnimGraphNode_SaveCachedPose> SaveCachedPoseNode;
 
 public:
-	// Begin UObject interface
-	virtual void PostLoad() override;
-	// End UObject interface
-
 	// UEdGraphNode interface
 	virtual FText GetTooltipText() const override;
 	virtual FText GetNodeTitle(ENodeTitleType::Type TitleType) const override;
@@ -31,6 +27,7 @@ public:
 	// UK2Node interface.
 	virtual void GetMenuActions(FBlueprintActionDatabaseRegistrar& ActionRegistrar) const override;
 	virtual bool IsActionFilteredOut(class FBlueprintActionFilter const& Filter) override;
+	virtual void ValidateNodeDuringCompilation(class FCompilerResultsLog& MessageLog) const override;
 	// End of UK2Node interface
 
 	// UAnimGraphNode_Base interface
diff --git a/Engine/Source/Editor/AnimGraph/Private/AnimGraphNode_UseCachedPose.cpp b/Engine/Source/Editor/AnimGraph/Private/AnimGraphNode_UseCachedPose.cpp
--- a/Engine/Source/Editor/AnimGraph/Private/AnimGraphNode_UseCachedPose.cpp
+++ b/Engine/Source/Editor/AnimGraph/Private/AnimGraphNode_UseCachedPose.cpp
@@ -5,6 +5,7 @@
 #include "GraphEditorActions.h"
 #include "ScopedTransaction.h"
 #include "K2ActionMenuBuilder.h" // for FK2ActionMenuBuilder::AddNewNodeAction()
+#include "EdGraphUtilities.h"
 #include "AnimGraphNode_SaveCachedPose.h"
 #include "AnimGraphNode_UseCachedPose.h"
 #include "BlueprintNodeSpawner.h"
@@ -21,19 +22,34 @@ UAnimGraphNode_UseCachedPose::UAnimGraphNode_UseCachedPose(const FPostConstructI
 {
 }
 
-void UAnimGraphNode_UseCachedPose::PostLoad()
+void UAnimGraphNode_UseCachedPose::ValidateNodeDuringCompilation(class FCompilerResultsLog& MessageLog) const
 {
-	Super::PostLoad();
+	bool bRefreshSavCachedPoseNode = true;
 
-	// If there is no SaveCachedPose node set but there is a NameOfCache available, we may be updating to the new system
-	// Go through and find the cached node, if possible.
-	if(!SaveCachedPoseNode.IsValid() && !NameOfCache.IsEmpty())
+	// Check to see the current cached node is still valid (and not deleted, by checking pin connections)
+	if(SaveCachedPoseNode.IsValid())
+	{
+		// The node has a single pin, make sure it's there
+		check(SaveCachedPoseNode->Pins.Num());
+
+		// Deleted nodes have no links, otherwise we will be doing some wasted work on unlinked nodes
+		if(SaveCachedPoseNode->Pins[0]->LinkedTo.Num())
+		{
+			// The node has links, it's valid, continue to use it
+			bRefreshSavCachedPoseNode = false;
+		}
+	}
+
+	// We need to refresh the cached pose node this node is linked to
+	if(bRefreshSavCachedPoseNode && !NameOfCache.IsEmpty())
 	{
-		TArray<UEdGraph*> AllAnimationGraphs;
-		GetGraph()->GetAllChildrenGraphs(AllAnimationGraphs);
-		AllAnimationGraphs.Add(GetGraph());
+		UBlueprint* GraphBlueprint = FBlueprintEditorUtils::FindBlueprintForGraph(GetGraph());
+		check(GraphBlueprint);
+
+		TArray<UEdGraph*> AllGraphs;
+		GraphBlueprint->GetAllGraphs(AllGraphs);
 
-		for(UEdGraph* Graph : AllAnimationGraphs)
+		for(UEdGraph* Graph : AllGraphs)
 		{
 			// Get a list of all save cached pose nodes
 			TArray<UAnimGraphNode_SaveCachedPose*> CachedPoseNodes;
@@ -44,6 +60,8 @@ void UAnimGraphNode_UseCachedPose::PostLoad()
 			{
 				if((*NodeIt)->CacheName == NameOfCache)
 				{
+					// Fix the original Blueprint node as well as the compiled version
+					MessageLog.FindSourceObjectTypeChecked<UAnimGraphNode_UseCachedPose>(this)->SaveCachedPoseNode = *NodeIt;
 					SaveCachedPoseNode = *NodeIt;
 					break;
 				}
