Fixed slot node name display (would show none even when a slot name was selected)

diff --git a/Engine/Source/Editor/AnimGraph/Private/AnimGraphNode_Slot.cpp b/Engine/Source/Editor/AnimGraph/Private/AnimGraphNode_Slot.cpp
--- a/Engine/Source/Editor/AnimGraph/Private/AnimGraphNode_Slot.cpp
+++ b/Engine/Source/Editor/AnimGraph/Private/AnimGraphNode_Slot.cpp
@@ -28,7 +28,7 @@ FText UAnimGraphNode_Slot::GetTooltipText() const
 
 FText UAnimGraphNode_Slot::GetNodeTitle(ENodeTitleType::Type TitleType) const
 {
-	if (Node.SlotName != NAME_None)
+	if (Node.SlotName == NAME_None)
 	{
 		if (TitleType == ENodeTitleType::ListView || TitleType == ENodeTitleType::MenuTitle)
 		{
