AnimTransitionSchema menu actions defined in UAnimationTransitionSchema::GetSourceStateActions have been ported to the new context menu system.

diff --git a/Engine/Source/Editor/AnimGraph/Classes/K2Node_TransitionRuleGetter.h b/Engine/Source/Editor/AnimGraph/Classes/K2Node_TransitionRuleGetter.h
--- a/Engine/Source/Editor/AnimGraph/Classes/K2Node_TransitionRuleGetter.h
+++ b/Engine/Source/Editor/AnimGraph/Classes/K2Node_TransitionRuleGetter.h
@@ -54,5 +54,18 @@ class UK2Node_TransitionRuleGetter : public UK2Node
 
 	// @todo document
 	ANIMGRAPH_API static FText GetFriendlyName(ETransitionGetter::Type TypeID);
+
+private:
+	/** Helper function for GetStateSpecificMenuActions, finds actions for AnimBlueprint states that can be used in AnimGraphSchema graphs */
+	void GetStateSpecificAnimGraphSchemaMenuActions(FBlueprintActionDatabaseRegistrar& ActionRegistrar, const UAnimBlueprint* AnimBlueprint, UAnimStateNode* StateNode) const;
+
+	/** Helper function for GetStateSpecificMenuActions, finds actions for AnimBlueprint states that can be used in AnimTransitionSchema graphs */
+	void GetStateSpecificAnimTransitionSchemaMenuActions(FBlueprintActionDatabaseRegistrar& ActionRegistrar, const UAnimBlueprint* AnimBlueprint, UAnimStateNode* StateNode) const;
+
+	/** Helper function for GetMenuActions, goes through all states in an AnimBlueprint and finds possible actions */
+	void GetStateSpecificMenuActions(FBlueprintActionDatabaseRegistrar& ActionRegistrar, const UAnimBlueprint* AnimBlueprint) const;
+
+	/** Helper function for GetMenuActions, generates non-state specific actions that use this node */
+	void GetNonStateSpecificMenuActions(FBlueprintActionDatabaseRegistrar& ActionRegistrar) const;
 };
 
diff --git a/Engine/Source/Editor/AnimGraph/Private/K2Node_TransitionRuleGetter.cpp b/Engine/Source/Editor/AnimGraph/Private/K2Node_TransitionRuleGetter.cpp
--- a/Engine/Source/Editor/AnimGraph/Private/K2Node_TransitionRuleGetter.cpp
+++ b/Engine/Source/Editor/AnimGraph/Private/K2Node_TransitionRuleGetter.cpp
@@ -3,6 +3,7 @@
 #include "AnimGraphPrivatePCH.h"
 
 #include "AnimationGraphSchema.h"
+#include "AnimationTransitionSchema.h"
 #include "AnimGraphNode_Base.h"
 #include "AnimStateNode.h"
 #include "BlueprintActionDatabaseRegistrar.h"
@@ -114,40 +115,163 @@ FText UK2Node_TransitionRuleGetter::GetNodeTitle(ENodeTitleType::Type TitleType)
 	return Super::GetNodeTitle(TitleType);
 }
 
-void UK2Node_TransitionRuleGetter::GetMenuActions(FBlueprintActionDatabaseRegistrar& ActionRegistrar) const
+void UK2Node_TransitionRuleGetter::GetStateSpecificAnimGraphSchemaMenuActions(FBlueprintActionDatabaseRegistrar& ActionRegistrar, const UAnimBlueprint* AnimBlueprint, UAnimStateNode* StateNode) const
 {
-	if( const UAnimBlueprint* AnimBlueprint = Cast<UAnimBlueprint>( ActionRegistrar.GetActionKeyFilter() ) )
+	auto UiSpecOverride = [](const FBlueprintActionContext& /*Context*/, const IBlueprintNodeBinder::FBindingSet& Bindings, FBlueprintActionUiSpec* UiSpecOut, UAnimStateNode* StateNode)
+	{
+		const FString OwnerName = StateNode->GetOuter()->GetName();
+		UiSpecOut->MenuName = FText::Format(LOCTEXT("TransitionRuleGetterTitle", "Current {0} for state '{1}.{2}'"), 
+			UK2Node_TransitionRuleGetter::GetFriendlyName(ETransitionGetter::ArbitraryState_GetBlendWeight), 
+			FText::FromString(OwnerName), 
+			FText::FromString(StateNode->GetStateName()));
+	};
+
+	auto PostSpawnSetupLambda = [](UEdGraphNode* NewNode, bool /*bIsTemplateNode*/, UAnimStateNode* StateNode)
+	{
+		UK2Node_TransitionRuleGetter* NewNodeTyped = CastChecked<UK2Node_TransitionRuleGetter>(NewNode);
+		NewNodeTyped->AssociatedStateNode = StateNode;
+		NewNodeTyped->GetterType = ETransitionGetter::ArbitraryState_GetBlendWeight;
+	};
+
+	UBlueprintNodeSpawner* Spawner = UBlueprintNodeSpawner::Create( UK2Node_TransitionRuleGetter::StaticClass(), nullptr, UBlueprintNodeSpawner::FCustomizeNodeDelegate::CreateStatic(PostSpawnSetupLambda, StateNode) );
+	Spawner->DynamicUiSignatureGetter = UBlueprintNodeSpawner::FUiSpecOverrideDelegate::CreateStatic(UiSpecOverride, StateNode);
+	ActionRegistrar.AddBlueprintAction( AnimBlueprint, Spawner );
+}
+
+void UK2Node_TransitionRuleGetter::GetStateSpecificAnimTransitionSchemaMenuActions(FBlueprintActionDatabaseRegistrar& ActionRegistrar, const UAnimBlueprint* AnimBlueprint, UAnimStateNode* StateNode) const
+{
+	// Offer options from the source state
+
+	// Sequence player positions
+	ETransitionGetter::Type SequenceSpecificGetters[] =
+	{
+		ETransitionGetter::AnimationAsset_GetCurrentTime,
+		ETransitionGetter::AnimationAsset_GetLength,
+		ETransitionGetter::AnimationAsset_GetCurrentTimeFraction,
+		ETransitionGetter::AnimationAsset_GetTimeFromEnd,
+		ETransitionGetter::AnimationAsset_GetTimeFromEndFraction
+	};
+
+	// Using the State Machine's graph, find all asset players nodes
+	TArray<UK2Node*> AssetPlayers;
+	StateNode->BoundGraph->GetNodesOfClassEx<UAnimGraphNode_Base, UK2Node>(/*out*/ AssetPlayers);
+
+	for (int32 TypeIndex = 0; TypeIndex < ARRAY_COUNT(SequenceSpecificGetters); ++TypeIndex)
+	{
+		for (auto NodeIt = AssetPlayers.CreateConstIterator(); NodeIt; ++NodeIt)
+		{
+			UAnimGraphNode_Base* AnimNode = CastChecked<UAnimGraphNode_Base>(*NodeIt);
+
+			if (AnimNode->DoesSupportTimeForTransitionGetter())
+			{
+				FString AssetName;
+				UAnimationAsset * AnimAsset = AnimNode->GetAnimationAsset();
+				UAnimGraphNode_Base* AssociatedAnimAssetPlayerNode = nullptr;
+				if (AnimAsset)
+				{
+					auto UiSpecOverride = [](const FBlueprintActionContext& /*Context*/, const IBlueprintNodeBinder::FBindingSet& Bindings, FBlueprintActionUiSpec* UiSpecOut, FString AssetName, TEnumAsByte<ETransitionGetter::Type> GetterType)
+					{
+						UiSpecOut->Category = LOCTEXT("AssetPlayer", "Asset Player");
+
+						FFormatNamedArguments Args;
+						Args.Add(TEXT("NodeName"), UK2Node_TransitionRuleGetter::GetFriendlyName(GetterType));
+						Args.Add(TEXT("AssetName"), FText::FromString(AssetName));
+						FText Title = FText::Format(LOCTEXT("TransitionFor", "{NodeName} for '{AssetName}'"), Args);
+						UiSpecOut->MenuName = Title;
+					};
+
+					auto PostSpawnSetupLambda = [](UEdGraphNode* NewNode, bool /*bIsTemplateNode*/, UAnimGraphNode_Base* AssociatedAnimAssetPlayerNode, TEnumAsByte<ETransitionGetter::Type> GetterType)
+					{
+						UK2Node_TransitionRuleGetter* NewNodeTyped = CastChecked<UK2Node_TransitionRuleGetter>(NewNode);
+						NewNodeTyped->AssociatedAnimAssetPlayerNode = AssociatedAnimAssetPlayerNode;
+						NewNodeTyped->GetterType = GetterType;
+					};
+
+					// Prepare the node spawner
+					AssociatedAnimAssetPlayerNode = AnimNode;
+					AssetName = AnimAsset->GetName();
+
+					TEnumAsByte<ETransitionGetter::Type> GetterType = SequenceSpecificGetters[TypeIndex];
+
+					UBlueprintNodeSpawner* Spawner = UBlueprintNodeSpawner::Create( UK2Node_TransitionRuleGetter::StaticClass(), nullptr, UBlueprintNodeSpawner::FCustomizeNodeDelegate::CreateStatic(PostSpawnSetupLambda, AssociatedAnimAssetPlayerNode, GetterType) );
+					Spawner->DynamicUiSignatureGetter = UBlueprintNodeSpawner::FUiSpecOverrideDelegate::CreateStatic(UiSpecOverride, AssetName, GetterType);
+					ActionRegistrar.AddBlueprintAction( AnimBlueprint, Spawner );
+				}
+			}
+		}
+	}
+}
+
+void UK2Node_TransitionRuleGetter::GetStateSpecificMenuActions(FBlueprintActionDatabaseRegistrar& ActionRegistrar, const UAnimBlueprint* AnimBlueprint) const
+{
+	TArray<UAnimStateNode*> States;
+	FBlueprintEditorUtils::GetAllNodesOfClass(AnimBlueprint, /*out*/ States);
+
+	// Go through all states to generate possible menu actions
+	for (auto StateIt = States.CreateIterator(); StateIt; ++StateIt)
+	{
+		UAnimStateNode* StateNode = *StateIt;
+
+		GetStateSpecificAnimGraphSchemaMenuActions(ActionRegistrar, AnimBlueprint, StateNode);
+		GetStateSpecificAnimTransitionSchemaMenuActions(ActionRegistrar, AnimBlueprint, StateNode);
+	}
+}
+
+void UK2Node_TransitionRuleGetter::GetNonStateSpecificMenuActions(FBlueprintActionDatabaseRegistrar& ActionRegistrar) const
+{
+	// actions get registered under specific object-keys; the idea is that 
+	// actions might have to be updated (or deleted) if their object-key is  
+	// mutated (or removed)... here we use the node's class (so if the node 
+	// type disappears, then the action should go with it)
+	UClass* ActionKey = GetClass();
+	// to keep from needlessly instantiating a UBlueprintNodeSpawner, first   
+	// check to make sure that the registrar is looking for actions of this type
+	// (could be regenerating actions for a specific asset, and therefore the 
+	// registrar would only accept actions corresponding to that asset)
+	if (ActionRegistrar.IsOpenForRegistration(ActionKey))
 	{
-		TArray<UAnimStateNode*> States;
-		FBlueprintEditorUtils::GetAllNodesOfClass(AnimBlueprint, /*out*/ States);
+		// Non-sequence specific ones
+		ETransitionGetter::Type NonSpecificGetters[] =
+		{
+			ETransitionGetter::CurrentTransitionDuration,
+			ETransitionGetter::CurrentState_ElapsedTime,
+			ETransitionGetter::CurrentState_GetBlendWeight
+		};
 
-		for (auto StateIt = States.CreateIterator(); StateIt; ++StateIt)
+		for (int32 TypeIndex = 0; TypeIndex < ARRAY_COUNT(NonSpecificGetters); ++TypeIndex)
 		{
-			UAnimStateNode* StateNode = *StateIt;
-			
-			auto UiSpecOverride = [](const FBlueprintActionContext& /*Context*/, const IBlueprintNodeBinder::FBindingSet& Bindings, FBlueprintActionUiSpec* UiSpecOut, UAnimStateNode* StateNode)
+			auto UiSpecOverride = [](const FBlueprintActionContext& /*Context*/, const IBlueprintNodeBinder::FBindingSet& Bindings, FBlueprintActionUiSpec* UiSpecOut, TEnumAsByte<ETransitionGetter::Type> GetterType)
 			{
-				const FString OwnerName = StateNode->GetOuter()->GetName();
-				UiSpecOut->MenuName = FText::Format(LOCTEXT("TransitionRuleGetterTitle", "Current {0} for state '{1}.{2}'"), 
-														UK2Node_TransitionRuleGetter::GetFriendlyName(ETransitionGetter::ArbitraryState_GetBlendWeight), 
-														FText::FromString(OwnerName), 
-														FText::FromString(StateNode->GetStateName()));
+				UiSpecOut->Category = LOCTEXT("Transition", "Transition");
+				UiSpecOut->MenuName = UK2Node_TransitionRuleGetter::GetFriendlyName(GetterType);
 			};
 
-			auto PostSpawnSetupLambda = [](UEdGraphNode* NewNode, bool /*bIsTemplateNode*/, UAnimStateNode* StateNode)
+			auto PostSpawnSetupLambda = [](UEdGraphNode* NewNode, bool /*bIsTemplateNode*/, TEnumAsByte<ETransitionGetter::Type> GetterType)
 			{
 				UK2Node_TransitionRuleGetter* NewNodeTyped = CastChecked<UK2Node_TransitionRuleGetter>(NewNode);
-				NewNodeTyped->AssociatedStateNode = StateNode;
-				NewNodeTyped->GetterType = ETransitionGetter::ArbitraryState_GetBlendWeight;
+				NewNodeTyped->GetterType = GetterType;
 			};
 
-			UBlueprintNodeSpawner* Spawner = UBlueprintNodeSpawner::Create( UK2Node_TransitionRuleGetter::StaticClass(), nullptr, UBlueprintNodeSpawner::FCustomizeNodeDelegate::CreateStatic(PostSpawnSetupLambda, StateNode) );
-			Spawner->DynamicUiSignatureGetter = UBlueprintNodeSpawner::FUiSpecOverrideDelegate::CreateStatic(UiSpecOverride, StateNode);
-			ActionRegistrar.AddBlueprintAction( AnimBlueprint, Spawner );
+			TEnumAsByte<ETransitionGetter::Type> GetterType = NonSpecificGetters[TypeIndex];
+			UBlueprintNodeSpawner* Spawner = UBlueprintNodeSpawner::Create( UK2Node_TransitionRuleGetter::StaticClass(), nullptr, UBlueprintNodeSpawner::FCustomizeNodeDelegate::CreateStatic(PostSpawnSetupLambda, GetterType) );
+			Spawner->DynamicUiSignatureGetter = UBlueprintNodeSpawner::FUiSpecOverrideDelegate::CreateStatic(UiSpecOverride, GetterType);
+			ActionRegistrar.AddBlueprintAction( ActionKey, Spawner );
 		}
 	}
 }
 
+void UK2Node_TransitionRuleGetter::GetMenuActions(FBlueprintActionDatabaseRegistrar& ActionRegistrar) const
+{
+	if( const UAnimBlueprint* AnimBlueprint = Cast<UAnimBlueprint>( ActionRegistrar.GetActionKeyFilter() ) )
+	{
+		GetStateSpecificMenuActions(ActionRegistrar, AnimBlueprint);
+	}
+	else
+	{
+		GetNonStateSpecificMenuActions(ActionRegistrar);
+	}
+}
+
 FText UK2Node_TransitionRuleGetter::GetTooltipText() const
 {
 	return GetNodeTitle(ENodeTitleType::FullTitle);
@@ -155,25 +279,99 @@ FText UK2Node_TransitionRuleGetter::GetTooltipText() const
 
 bool UK2Node_TransitionRuleGetter::CanCreateUnderSpecifiedSchema(const UEdGraphSchema* Schema) const
 {
-	return Cast<UAnimationGraphSchema>(Schema) != NULL;
+	return Cast<UAnimationGraphSchema>(Schema) != NULL || Cast<UAnimationTransitionSchema>(Schema) != NULL;
+}
+
+UAnimStateTransitionNode* GetTransitionNodeFromGraph(const FAnimBlueprintDebugData& DebugData, const UEdGraph* Graph)
+{
+	if (const TWeakObjectPtr<UAnimStateTransitionNode>* TransNodePtr = DebugData.TransitionGraphToNodeMap.Find(Graph))
+	{
+		return TransNodePtr->Get();
+	}
+
+	if (const TWeakObjectPtr<UAnimStateTransitionNode>* TransNodePtr = DebugData.TransitionBlendGraphToNodeMap.Find(Graph))
+	{
+		return TransNodePtr->Get();
+	}
+
+	return NULL;
 }
 
 bool UK2Node_TransitionRuleGetter::IsActionFilteredOut(class FBlueprintActionFilter const& Filter)
 {
-	// only show the transition nodes if the associated state node is in every graph:
-	if( AssociatedStateNode )
+	if(Filter.Context.Graphs[0])
 	{
-		for( auto Blueprint : Filter.Context.Blueprints )
+		const UEdGraphSchema* Schema = Filter.Context.Graphs[0]->GetSchema();
+		if(Cast<UAnimationGraphSchema>(Schema))
 		{
-			TArray<UAnimStateNode*> States;
-			FBlueprintEditorUtils::GetAllNodesOfClass(Blueprint, /*out*/ States);
+			if(GetterType == ETransitionGetter::ArbitraryState_GetBlendWeight)
+			{
+				// only show the transition nodes if the associated state node is in every graph:
+				if( AssociatedStateNode )
+				{
+					for( auto Blueprint : Filter.Context.Blueprints )
+					{
+						TArray<UAnimStateNode*> States;
+						FBlueprintEditorUtils::GetAllNodesOfClass(Blueprint, /*out*/ States);
 
-			if( !States.Contains( AssociatedStateNode ) )
+						if( !States.Contains( AssociatedStateNode ) )
+						{
+							return true;
+						}
+					}
+					return false;
+				}
+			}
+			return true;
+		}
+		else if(Cast<UAnimationTransitionSchema>(Schema))
+		{
+			// Non-sequence specific TransitionRuleGetter nodes have no associated nodes assigned, they are always allowed in AnimationTransitionSchema graphs
+			if( AssociatedStateNode == nullptr && AssociatedAnimAssetPlayerNode == nullptr)
+			{
+				return false;
+			}
+			else if(AssociatedAnimAssetPlayerNode)
 			{
-				return true;
+				if(Filter.Context.Blueprints.Num())
+				{
+					UAnimBlueprint* AnimBlueprint = Cast<UAnimBlueprint>(Filter.Context.Blueprints[0]);
+					check(AnimBlueprint);
+
+					if (UAnimBlueprintGeneratedClass* AnimBlueprintClass = AnimBlueprint->GetAnimBlueprintSkeletonClass())
+					{
+						// Local function to find the transition node using the context graph and AnimBlueprintDebugData
+						auto GetTransitionNodeFromGraphLambda = [](const FAnimBlueprintDebugData& DebugData, const UEdGraph* Graph) -> UAnimStateTransitionNode*
+						{
+							if (const TWeakObjectPtr<UAnimStateTransitionNode>* TransNodePtr = DebugData.TransitionGraphToNodeMap.Find(Graph))
+							{
+								return TransNodePtr->Get();
+							}
+
+							if (const TWeakObjectPtr<UAnimStateTransitionNode>* TransNodePtr = DebugData.TransitionBlendGraphToNodeMap.Find(Graph))
+							{
+								return TransNodePtr->Get();
+							}
+
+							return NULL;
+						};
+
+						// Check if the TransitionNode can be found in the AninBlueprint's debug data
+						if (UAnimStateTransitionNode* TransNode = GetTransitionNodeFromGraph(AnimBlueprintClass->GetAnimBlueprintDebugData(), Filter.Context.Graphs[0]))
+						{
+							if (UAnimStateNode* SourceStateNode = Cast<UAnimStateNode>(TransNode->GetPreviousState()))
+							{
+								// Check if the AnimAssetPlayerNode's graph is the state machine's bound graph
+								if(AssociatedAnimAssetPlayerNode->GetGraph() == SourceStateNode->BoundGraph)
+								{
+									return false;
+								}
+							}
+						}
+					}
+				}
 			}
 		}
-		return false;
 	}
 	return true;
 }
