Fix for crashing attempting to invalidly access array at index 0.

diff --git a/Engine/Source/Editor/AnimGraph/Private/K2Node_TransitionRuleGetter.cpp b/Engine/Source/Editor/AnimGraph/Private/K2Node_TransitionRuleGetter.cpp
--- a/Engine/Source/Editor/AnimGraph/Private/K2Node_TransitionRuleGetter.cpp
+++ b/Engine/Source/Editor/AnimGraph/Private/K2Node_TransitionRuleGetter.cpp
@@ -299,7 +299,7 @@ UAnimStateTransitionNode* GetTransitionNodeFromGraph(const FAnimBlueprintDebugDa
 
 bool UK2Node_TransitionRuleGetter::IsActionFilteredOut(class FBlueprintActionFilter const& Filter)
 {
-	if(Filter.Context.Graphs[0])
+	if(Filter.Context.Graphs.Num())
 	{
 		const UEdGraphSchema* Schema = Filter.Context.Graphs[0]->GetSchema();
 		if(Cast<UAnimationGraphSchema>(Schema))
