AnimGraph: Fixed non-unity builds

diff --git a/Engine/Source/Editor/AnimGraph/Private/K2Node_TransitionRuleGetter.cpp b/Engine/Source/Editor/AnimGraph/Private/K2Node_TransitionRuleGetter.cpp
--- a/Engine/Source/Editor/AnimGraph/Private/K2Node_TransitionRuleGetter.cpp
+++ b/Engine/Source/Editor/AnimGraph/Private/K2Node_TransitionRuleGetter.cpp
@@ -1,16 +1,18 @@
 // Copyright 1998-2014 Epic Games, Inc. All Rights Reserved.
 
 #include "AnimGraphPrivatePCH.h"
-
 #include "AnimationGraphSchema.h"
 #include "AnimationTransitionSchema.h"
 #include "AnimGraphNode_Base.h"
 #include "AnimStateNode.h"
+#include "AnimStateTransitionNode.h"
 #include "BlueprintActionDatabaseRegistrar.h"
 #include "BlueprintActionFilter.h"
 #include "BlueprintNodeSpawner.h"
 #include "BlueprintUtilities.h"
 #include "K2Node_TransitionRuleGetter.h"
+
+
 /////////////////////////////////////////////////////
 // UK2Node_TransitionRuleGetter
 
