Fix previewing of root motion in montages

diff --git a/Engine/Source/Editor/AnimGraph/Classes/AnimPreviewInstance.h b/Engine/Source/Editor/AnimGraph/Classes/AnimPreviewInstance.h
--- a/Engine/Source/Editor/AnimGraph/Classes/AnimPreviewInstance.h
+++ b/Engine/Source/Editor/AnimGraph/Classes/AnimPreviewInstance.h
@@ -57,6 +57,7 @@ class ANIMGRAPH_API UAnimPreviewInstance : public UAnimSingleNodeInstance
 
 	// Begin UAnimSingleNodeInstance interface
 	virtual void RestartMontage(UAnimMontage* Montage, FName FromSection = FName()) override;
+	virtual void SetAnimationAsset(UAnimationAsset* NewAsset, bool bIsLooping = true, float InPlayRate = 1.f) override;
 	// End UAnimSingleNodeInstance interface
 
 	/** Montage preview functions */
diff --git a/Engine/Source/Editor/AnimGraph/Private/AnimPreviewInstance.cpp b/Engine/Source/Editor/AnimGraph/Private/AnimPreviewInstance.cpp
--- a/Engine/Source/Editor/AnimGraph/Private/AnimPreviewInstance.cpp
+++ b/Engine/Source/Editor/AnimGraph/Private/AnimPreviewInstance.cpp
@@ -161,6 +161,12 @@ void UAnimPreviewInstance::RestartMontage(UAnimMontage* Montage, FName FromSecti
 	}
 }
 
+void UAnimPreviewInstance::SetAnimationAsset(UAnimationAsset* NewAsset, bool bIsLooping, float InPlayRate)
+{
+	Super::SetAnimationAsset(NewAsset, bIsLooping, InPlayRate);
+	RootMotionMode = Cast<UAnimMontage>(CurrentAsset) != NULL ? ERootMotionMode::RootMotionFromMontagesOnly : ERootMotionMode::RootMotionFromEverything;
+}
+
 void UAnimPreviewInstance::MontagePreview_SetLooping(bool bIsLooping)
 {
 	bLooping = bIsLooping;
