Bone Controller Preview Feature - Applied to Two Bone IK, ModifyBone and Fabrik nodes for now

diff --git a/Engine/Source/Editor/AnimGraph/Classes/AnimGraphNode_Fabrik.h b/Engine/Source/Editor/AnimGraph/Classes/AnimGraphNode_Fabrik.h
--- a/Engine/Source/Editor/AnimGraph/Classes/AnimGraphNode_Fabrik.h
+++ b/Engine/Source/Editor/AnimGraph/Classes/AnimGraphNode_Fabrik.h
@@ -19,6 +19,17 @@ public:
 	virtual FText GetNodeTitle(ENodeTitleType::Type TitleType) const override;
 	// End of UEdGraphNode interface
 
+	// UAnimGraphNode_SkeletalControlBase interface
+	ANIMGRAPH_API virtual FVector GetWidgetLocation(const USkeletalMeshComponent* SkelComp, struct FAnimNode_SkeletalControlBase* AnimNode) override;
+	ANIMGRAPH_API virtual FWidget::EWidgetMode GetWidgetMode(const USkeletalMeshComponent* SkelComp) override;
+	ANIMGRAPH_API virtual FName FindSelectedBone() override;
+	ANIMGRAPH_API virtual void DoTranslation(const USkeletalMeshComponent* SkelComp, FVector& Drag, FAnimNode_Base* InOutAnimNode) override;
+
+	ANIMGRAPH_API virtual void	CopyNodeDataTo(FAnimNode_Base* AnimNode) override;
+	ANIMGRAPH_API virtual void	CopyNodeDataFrom(const FAnimNode_Base* NewAnimNode) override;
+
+	// End of UAnimGraphNode_SkeletalControlBase interface
+
 protected:
 	// UAnimGraphNode_SkeletalControlBase interface
 	virtual FText GetControllerDescription() const override;
diff --git a/Engine/Source/Editor/AnimGraph/Classes/AnimGraphNode_ModifyBone.h b/Engine/Source/Editor/AnimGraph/Classes/AnimGraphNode_ModifyBone.h
--- a/Engine/Source/Editor/AnimGraph/Classes/AnimGraphNode_ModifyBone.h
+++ b/Engine/Source/Editor/AnimGraph/Classes/AnimGraphNode_ModifyBone.h
@@ -25,10 +25,33 @@ public:
 protected:
 	// UAnimGraphNode_SkeletalControlBase interface
 	virtual FText GetControllerDescription() const override;
+
+	ANIMGRAPH_API virtual FVector GetWidgetLocation(const USkeletalMeshComponent* SkelComp, struct FAnimNode_SkeletalControlBase* AnimNode) override;
+	ANIMGRAPH_API virtual FWidget::EWidgetMode GetWidgetMode(const USkeletalMeshComponent* SkelComp) override;
+	ANIMGRAPH_API virtual FWidget::EWidgetMode ChangeWidgetMode(const USkeletalMeshComponent* SkelComp, FWidget::EWidgetMode InCurWidgetMode) override;
+
+	ANIMGRAPH_API virtual FName FindSelectedBone() override;
+
+	ANIMGRAPH_API virtual void UpdateDefaultValues(const FAnimNode_Base* AnimNode) override;
+	ANIMGRAPH_API virtual void UpdateAllDefaultValues(const FAnimNode_Base* AnimNode) override;
+
+	ANIMGRAPH_API virtual void DoTranslation(const USkeletalMeshComponent* SkelComp, FVector& Drag, FAnimNode_Base* InOutAnimNode) override;
+	ANIMGRAPH_API virtual void DoRotation(const USkeletalMeshComponent* SkelComp, FRotator& Rotation, FAnimNode_Base* InOutAnimNode) override;
+	ANIMGRAPH_API virtual void DoScale(const USkeletalMeshComponent* SkelComp, FVector& Drag, FAnimNode_Base* InOutAnimNode) override;
+	ANIMGRAPH_API virtual void	CopyNodeDataTo(FAnimNode_Base* OutAnimNode) override;
+	ANIMGRAPH_API virtual void	CopyNodeDataFrom(const FAnimNode_Base* InNewAnimNode) override;
+
 	// End of UAnimGraphNode_SkeletalControlBase interface
 
+	// methods to find a valid widget mode for gizmo because doesn't need to show gizmo when the mode is "Ignore"
+	FWidget::EWidgetMode FindValidWidgetMode(FWidget::EWidgetMode InMode);
+	EBoneModificationMode GetBoneModificationMode(FWidget::EWidgetMode InMode);
+	FWidget::EWidgetMode GetNextWidgetMode(FWidget::EWidgetMode InMode);
+
 private:
 	/** Constructing FText strings can be costly, so we cache the node's title */
 	FNodeTitleTextTable CachedNodeTitles;
+
+	FWidget::EWidgetMode CurWidgetMode;
 };
 
diff --git a/Engine/Source/Editor/AnimGraph/Classes/AnimGraphNode_SkeletalControlBase.h b/Engine/Source/Editor/AnimGraph/Classes/AnimGraphNode_SkeletalControlBase.h
--- a/Engine/Source/Editor/AnimGraph/Classes/AnimGraphNode_SkeletalControlBase.h
+++ b/Engine/Source/Editor/AnimGraph/Classes/AnimGraphNode_SkeletalControlBase.h
@@ -23,8 +23,69 @@ public:
 	// End of UAnimGraphNode_Base interface
 
 	// Draw function for supporting visualization
-	ANIMGRAPH_API virtual void Draw( FPrimitiveDrawInterface* PDI, USkeletalMeshComponent * PreviewSkelMeshComp ) const {};
+	ANIMGRAPH_API virtual void Draw(FPrimitiveDrawInterface* PDI, USkeletalMeshComponent * PreviewSkelMeshComp) const {};
+
+	/**
+	 * methods related to widget control
+	 */
+	ANIMGRAPH_API virtual FVector GetWidgetLocation(const USkeletalMeshComponent* SkelComp, struct FAnimNode_SkeletalControlBase* AnimNode)
+	{
+		return FVector::ZeroVector;
+	}
+	// to keep data consistency between anim nodes 
+	ANIMGRAPH_API virtual void	CopyNodeDataTo(FAnimNode_Base* OutAnimNode){}
+	ANIMGRAPH_API virtual void	CopyNodeDataFrom(const FAnimNode_Base* NewAnimNode){}
+	// update literal values in Pins
+	ANIMGRAPH_API virtual void	UpdateDefaultValues(const FAnimNode_Base* AnimNode){}
+	// update all literal values for Undo/Redo
+	ANIMGRAPH_API virtual void UpdateAllDefaultValues(const FAnimNode_Base* AnimNode){}
+
+	// return current widget mode this anim graph node supports
+	ANIMGRAPH_API virtual FWidget::EWidgetMode GetWidgetMode(const USkeletalMeshComponent* SkelComp)
+	{ 
+		return FWidget::EWidgetMode::WM_None; 
+	}
+	// called when the user changed widget mode
+	ANIMGRAPH_API virtual FWidget::EWidgetMode ChangeWidgetMode(const USkeletalMeshComponent* SkelComp, FWidget::EWidgetMode CurWidgetMode)
+	{ 
+		return GetWidgetMode(SkelComp); 
+	}
+	// 
+	ANIMGRAPH_API virtual FName FindSelectedBone()
+	{ 
+		return NAME_None; 
+	}
+
+	// if anim graph node needs other actors to select other bones, move actor's positions when this is called
+	ANIMGRAPH_API virtual void MoveSelectActorLocation(const USkeletalMeshComponent* SkelComp, FAnimNode_SkeletalControlBase* InAnimNode){}
+
+	ANIMGRAPH_API virtual bool IsActorClicked(HActor* ActorHitProxy){ return false; }
+	ANIMGRAPH_API virtual void ProcessActorClick(HActor* ActorHitProxy){}
+	// if it has select-actors, should hide all actors when de-select is called  
+	ANIMGRAPH_API virtual void	DeselectActor(USkeletalMeshComponent* SkelComp){}
+
+	// called when the widget is dragged in translation mode
+	ANIMGRAPH_API virtual void DoTranslation(const USkeletalMeshComponent* SkelComp, FVector& Drag, FAnimNode_Base* InOutAnimNode){}
+	// called when the widget is dragged in rotation mode
+	ANIMGRAPH_API virtual void DoRotation(const USkeletalMeshComponent* SkelComp, FRotator& Rotation, FAnimNode_Base* InOutAnimNode){}
+	// called when the widget is dragged in scale mode
+	ANIMGRAPH_API virtual void DoScale(const USkeletalMeshComponent* SkelComp, FVector& Scale, FAnimNode_Base* InOutAnimNode){}
+
 protected:
 	// Returns the short descriptive name of the controller
 	ANIMGRAPH_API virtual FText GetControllerDescription() const;
+
+	/**
+	* helper functions for bone control preview
+	*/
+	// local conversion function for drawing
+	void ConvertToComponentSpaceTransform(const USkeletalMeshComponent* SkelComp, const FTransform & InTransform, FTransform & OutCSTransform, int32 BoneIndex, EBoneControlSpace Space) const;
+	// convert drag vector in component space to bone space 
+	FVector ConvertCSVectorToBoneSpace(const USkeletalMeshComponent* SkelComp, FVector& InCSVector, FA2CSPose& MeshBases, const FName& BoneName, const EBoneControlSpace Space);
+	// convert rotator in component space to bone space 
+	FQuat ConvertCSRotationToBoneSpace(const USkeletalMeshComponent* SkelComp, FRotator& InCSRotator, FA2CSPose& MeshBases, const FName& BoneName, const EBoneControlSpace Space);
+	// convert widget location according to bone control space
+	FVector ConvertWidgetLocation(const USkeletalMeshComponent* InSkelComp, FA2CSPose& InMeshBases, const FName& BoneName, const FVector& InLocation, const EBoneControlSpace Space);
+	// set literal value for FVector
+	void SetDefaultValue(FString& InDefaultValueName, FVector& InValue);
 };
diff --git a/Engine/Source/Editor/AnimGraph/Classes/AnimGraphNode_TwoBoneIK.h b/Engine/Source/Editor/AnimGraph/Classes/AnimGraphNode_TwoBoneIK.h
--- a/Engine/Source/Editor/AnimGraph/Classes/AnimGraphNode_TwoBoneIK.h
+++ b/Engine/Source/Editor/AnimGraph/Classes/AnimGraphNode_TwoBoneIK.h
@@ -6,6 +6,9 @@
 #include "EdGraph/EdGraphNodeUtils.h" // for FNodeTitleTextTable
 #include "AnimGraphNode_TwoBoneIK.generated.h"
 
+// actor class used for bone selector
+#define ABoneSelectActor ATargetPoint
+
 UCLASS(MinimalAPI)
 class UAnimGraphNode_TwoBoneIK : public UAnimGraphNode_SkeletalControlBase
 {
@@ -14,31 +17,65 @@ class UAnimGraphNode_TwoBoneIK : public UAnimGraphNode_SkeletalControlBase
 	UPROPERTY(EditAnywhere, Category=Settings)
 	FAnimNode_TwoBoneIK Node;
 
+	// just for refreshing UIs when bone space was changed
+	static TSharedPtr<class FTwoBoneIKDelegate> TwoBoneIKDelegate;
+
 public:
 	// UEdGraphNode interface
 	virtual FText GetNodeTitle(ENodeTitleType::Type TitleType) const override;
 	virtual FText GetTooltipText() const override;
 	// End of UEdGraphNode interface
 
+	void UpdateEffectorLocationSpace(class IDetailLayoutBuilder& DetailBuilder);
 	// UAnimGraphNode_Base interface
 	virtual void CustomizeDetails(class IDetailLayoutBuilder& DetailBuilder) override;
 	// End of UAnimGraphNode_Base interface
 
 	// UAnimGraphNode_SkeletalControlBase interface
 	ANIMGRAPH_API virtual void Draw( FPrimitiveDrawInterface* PDI, USkeletalMeshComponent* SkelMeshComp) const override;
+
+	ANIMGRAPH_API virtual FVector GetWidgetLocation(const USkeletalMeshComponent* SkelComp, struct FAnimNode_SkeletalControlBase* AnimNode) override;
+
+	ANIMGRAPH_API virtual void	CopyNodeDataFrom(const FAnimNode_Base* NewAnimNode) override;
+
+	ANIMGRAPH_API virtual void	UpdateDefaultValues(const FAnimNode_Base* InAnimNode) override;
+	ANIMGRAPH_API virtual void	UpdateAllDefaultValues(const FAnimNode_Base* InAnimNode) override;
+
+	ANIMGRAPH_API virtual FWidget::EWidgetMode GetWidgetMode(const USkeletalMeshComponent* SkelComp) override;
+
+	ANIMGRAPH_API virtual void MoveSelectActorLocation(const USkeletalMeshComponent* SkelComp, FAnimNode_SkeletalControlBase* AnimNode) override;
+	ANIMGRAPH_API virtual FName FindSelectedBone() override;
+	ANIMGRAPH_API virtual bool IsActorClicked(HActor* ActorHitProxy) override;
+	ANIMGRAPH_API virtual void ProcessActorClick(HActor* ActorHitProxy) override;
+	ANIMGRAPH_API virtual void DoTranslation(const USkeletalMeshComponent* SkelComp, FVector& Drag, FAnimNode_Base* InOutAnimNode) override;
+	ANIMGRAPH_API virtual void	CopyNodeDataTo(FAnimNode_Base* AnimNode) override;
+	ANIMGRAPH_API virtual void	DeselectActor(USkeletalMeshComponent* SkelComp) override;
 	// End of UAnimGraphNode_SkeletalControlBase interface
 
+	enum BoneSelectModeType
+	{
+		BSM_EndEffector,
+		BSM_JointTarget,
+		BSM_Max
+	};
+
+	// to switch select mode between end effector and joint target
+	BoneSelectModeType BoneSelectMode;
+
+	TWeakObjectPtr<class ABoneSelectActor> BoneSelectActor;
+
+	IDetailLayoutBuilder* DetailLayout;
+
 protected:
 	// UAnimGraphNode_SkeletalControlBase interface
 	virtual FText GetControllerDescription() const override;
 	// End of UAnimGraphNode_SkeletalControlBase interface
 
 	// local conversion function for drawing
-	void ConvertToComponentSpaceTransform(USkeletalMeshComponent* SkelComp, USkeleton* Skeleton, const FTransform& InTransform, FTransform& OutCSTransform, int32 BoneIndex, uint8 Space) const;
-	void DrawTargetLocation(FPrimitiveDrawInterface* PDI, USkeletalMeshComponent* SkelComp, USkeleton* Skeleton, uint8 SpaceBase, FName SpaceBoneName, const FVector& TargetLocation, const FColor& TargetColor, const FColor& BoneColor) const;
+	void DrawTargetLocation(FPrimitiveDrawInterface* PDI, USkeletalMeshComponent* SkelComp, USkeleton * Skeleton, EBoneControlSpace SpaceBase, FName SpaceBoneName, const FVector & TargetLocation, const FColor & TargetColor, const FColor & BoneColor) const;
 
 	// make Pins showed / hidden by options
-	void SetPinsVisibility(bool bShow);
+	bool SetPinsVisibility(bool bShow);
 
 private:
 	/** Constructing FText strings can be costly, so we cache the node's title */
diff --git a/Engine/Source/Editor/AnimGraph/Private/AnimGraphNode_Fabrik.cpp b/Engine/Source/Editor/AnimGraph/Private/AnimGraphNode_Fabrik.cpp
--- a/Engine/Source/Editor/AnimGraph/Private/AnimGraphNode_Fabrik.cpp
+++ b/Engine/Source/Editor/AnimGraph/Private/AnimGraphNode_Fabrik.cpp
@@ -23,4 +23,60 @@ FText UAnimGraphNode_Fabrik::GetNodeTitle(ENodeTitleType::Type TitleType) const
 	return GetControllerDescription();
 }
 
+FVector UAnimGraphNode_Fabrik::GetWidgetLocation(const USkeletalMeshComponent* SkelComp, FAnimNode_SkeletalControlBase* AnimNode)
+{
+	FName& BoneName = Node.EffectorTransformBone.BoneName;
+	FVector Location = Node.EffectorTransform.GetLocation();
+	EBoneControlSpace Space = Node.EffectorTransformSpace;
+	FVector WidgetLoc = ConvertWidgetLocation(SkelComp, AnimNode->ForwardedPose, BoneName, Location, Space);
+	return WidgetLoc;
+}
+
+FWidget::EWidgetMode UAnimGraphNode_Fabrik::GetWidgetMode(const USkeletalMeshComponent* SkelComp)
+{
+	uint32 BoneIndex = SkelComp->GetBoneIndex(Node.EffectorTransformBone.BoneName);
+
+	if (BoneIndex != INDEX_NONE)
+	{
+		return FWidget::WM_Translate;
+	}
+
+	return FWidget::WM_None;
+}
+
+FName UAnimGraphNode_Fabrik::FindSelectedBone()
+{
+	return Node.EffectorTransformBone.BoneName;
+}
+
+void UAnimGraphNode_Fabrik::DoTranslation(const USkeletalMeshComponent* SkelComp, FVector& Drag, FAnimNode_Base* InOutAnimNode)
+{
+	FAnimNode_Fabrik* Fabrik = static_cast<FAnimNode_Fabrik*>(InOutAnimNode);
+
+	FVector Offset = ConvertCSVectorToBoneSpace(SkelComp, Drag, Fabrik->ForwardedPose, Node.EffectorTransformBone.BoneName, Node.EffectorTransformSpace);
+
+	if (Fabrik)
+	{
+		Fabrik->EffectorTransform.AddToTranslation(Offset);
+		// copy same value to internal node for data consistency
+		Node.EffectorTransform.SetTranslation(Fabrik->EffectorTransform.GetTranslation());
+	}
+}
+
+void UAnimGraphNode_Fabrik::CopyNodeDataTo(FAnimNode_Base* AnimNode)
+{
+	FAnimNode_Fabrik* Fabrik = static_cast<FAnimNode_Fabrik*>(AnimNode);
+
+	// copies Pin values from the internal node to get data which are not compiled yet
+	Fabrik->EffectorTransform = Node.EffectorTransform;
+}
+
+void UAnimGraphNode_Fabrik::CopyNodeDataFrom(const FAnimNode_Base* InNewAnimNode)
+{
+	const FAnimNode_Fabrik* Fabrik = static_cast<const FAnimNode_Fabrik*>(InNewAnimNode);
+
+	// copies Pin data from updated values
+	Node.EffectorTransform = Fabrik->EffectorTransform;
+}
+
 #undef LOCTEXT_NAMESPACE
\ No newline at end of file
diff --git a/Engine/Source/Editor/AnimGraph/Private/AnimGraphNode_ModifyBone.cpp b/Engine/Source/Editor/AnimGraph/Private/AnimGraphNode_ModifyBone.cpp
--- a/Engine/Source/Editor/AnimGraph/Private/AnimGraphNode_ModifyBone.cpp
+++ b/Engine/Source/Editor/AnimGraph/Private/AnimGraphNode_ModifyBone.cpp
@@ -11,6 +11,7 @@
 UAnimGraphNode_ModifyBone::UAnimGraphNode_ModifyBone(const FObjectInitializer& ObjectInitializer)
 	: Super(ObjectInitializer)
 {
+	CurWidgetMode = FWidget::WM_Rotate;
 }
 
 FText UAnimGraphNode_ModifyBone::GetControllerDescription() const
@@ -55,4 +56,257 @@ FText UAnimGraphNode_ModifyBone::GetNodeTitle(ENodeTitleType::Type TitleType) co
 	return CachedNodeTitles[TitleType];
 }
 
+FVector UAnimGraphNode_ModifyBone::GetWidgetLocation(const USkeletalMeshComponent* SkelComp, FAnimNode_SkeletalControlBase* AnimNode)
+{
+	USkeleton * Skeleton = SkelComp->SkeletalMesh->Skeleton;
+	FVector WidgetLoc = FVector::ZeroVector;
+
+	// if the current widget mode is translate, then shows the widget according to translation space
+	if (CurWidgetMode == FWidget::WM_Translate)
+	{
+		FA2CSPose& MeshBases = AnimNode->ForwardedPose;
+		WidgetLoc = ConvertWidgetLocation(SkelComp, MeshBases, Node.BoneToModify.BoneName, Node.Translation, Node.TranslationSpace);
+
+		if(Node.TranslationMode == BMM_Additive)
+		{
+			if(Node.TranslationSpace == EBoneControlSpace::BCS_WorldSpace ||
+				Node.TranslationSpace == EBoneControlSpace::BCS_ComponentSpace)
+			{
+				int32 MeshBoneIndex = SkelComp->GetBoneIndex(Node.BoneToModify.BoneName);
+				if(MeshBoneIndex != INDEX_NONE)
+				{
+					FTransform BoneTM = MeshBases.GetComponentSpaceTransform(MeshBoneIndex);
+					WidgetLoc += BoneTM.GetLocation();
+				}
+			}
+		}
+	}
+	else // if the current widget mode is not translate mode, then show the widget on the bone to modify
+	{
+		int32 MeshBoneIndex = SkelComp->GetBoneIndex(Node.BoneToModify.BoneName);
+
+		if (MeshBoneIndex != INDEX_NONE)
+		{
+			const FTransform BoneTM = SkelComp->GetBoneTransform(MeshBoneIndex);
+			WidgetLoc = BoneTM.GetLocation();
+		}
+	}
+	
+	return WidgetLoc;
+}
+
+EBoneModificationMode UAnimGraphNode_ModifyBone::GetBoneModificationMode(FWidget::EWidgetMode InMode)
+{
+	switch (InMode)
+	{
+	case FWidget::WM_Translate:
+		return Node.TranslationMode;
+		break;
+	case FWidget::WM_Rotate:
+		return Node.RotationMode;
+		break;
+	case FWidget::WM_Scale:
+		return Node.ScaleMode;
+		break;
+	}
+
+	return EBoneModificationMode::BMM_Ignore;
+}
+
+FWidget::EWidgetMode UAnimGraphNode_ModifyBone::GetNextWidgetMode(FWidget::EWidgetMode InMode)
+{
+	switch (InMode)
+	{
+	case FWidget::WM_Translate:
+		return FWidget::WM_Rotate;
+	case FWidget::WM_Rotate:
+		return FWidget::WM_Scale;
+	case FWidget::WM_Scale:
+		return FWidget::WM_Translate;
+	}
+
+	return FWidget::WM_None;
+}
+
+FWidget::EWidgetMode UAnimGraphNode_ModifyBone::FindValidWidgetMode(FWidget::EWidgetMode InMode)
+{	
+	FWidget::EWidgetMode ValidMode = InMode;
+	if (InMode == FWidget::WM_None)
+	{	// starts from Rotate mode
+		ValidMode = FWidget::WM_Rotate;
+	}
+
+	// find from current widget mode and loop 1 cycle until finding a valid mode
+	for (int32 Index = 0; Index < 3; Index++)
+	{
+		if (GetBoneModificationMode(ValidMode) != BMM_Ignore)
+		{
+			return ValidMode;
+		}
+
+		ValidMode = GetNextWidgetMode(ValidMode);
+	}
+	
+	// if couldn't find a valid mode, returns None
+	ValidMode = FWidget::WM_None;
+
+	return ValidMode;
+}
+
+FWidget::EWidgetMode UAnimGraphNode_ModifyBone::GetWidgetMode(const USkeletalMeshComponent* SkelComp)
+{
+	int32 BoneIndex = SkelComp->GetBoneIndex(Node.BoneToModify.BoneName);
+	if (BoneIndex != INDEX_NONE)
+	{
+		CurWidgetMode = FindValidWidgetMode(CurWidgetMode);
+		return CurWidgetMode;
+	}
+
+	return FWidget::WM_None;
+}
+
+FWidget::EWidgetMode UAnimGraphNode_ModifyBone::ChangeWidgetMode(const USkeletalMeshComponent* SkelComp, FWidget::EWidgetMode InCurWidgetMode)
+{
+	FWidget::EWidgetMode NextWidgetMode = GetNextWidgetMode(InCurWidgetMode);
+	CurWidgetMode = FindValidWidgetMode(NextWidgetMode);
+
+	return CurWidgetMode;
+}
+
+FName UAnimGraphNode_ModifyBone::FindSelectedBone()
+{
+	return Node.BoneToModify.BoneName;
+}
+
+void UAnimGraphNode_ModifyBone::UpdateDefaultValues(const FAnimNode_Base* AnimNode)
+{
+	FString UpdateDefaultValueName;
+	FVector UpdateValue;
+
+	const FAnimNode_ModifyBone* ModifyBone = static_cast<const FAnimNode_ModifyBone*>(AnimNode);
+
+	if (!ModifyBone)
+	{
+		return;
+	}
+
+	switch (CurWidgetMode)
+	{
+	case FWidget::WM_Translate:
+
+		UpdateDefaultValueName = FString("Translation");
+		UpdateValue = ModifyBone->Translation;
+		break;
+	case FWidget::WM_Rotate:
+		UpdateDefaultValueName = FString("Rotation");
+		UpdateValue.X = ModifyBone->Rotation.Pitch;
+		UpdateValue.Y = ModifyBone->Rotation.Yaw;
+		UpdateValue.Z = ModifyBone->Rotation.Roll;
+		break;
+	case FWidget::WM_Scale:
+		UpdateDefaultValueName = FString("Scale");
+		UpdateValue = ModifyBone->Scale;
+		break;
+	}
+
+	SetDefaultValue(UpdateDefaultValueName, UpdateValue);
+}
+
+void UAnimGraphNode_ModifyBone::UpdateAllDefaultValues(const FAnimNode_Base* AnimNode)
+{
+	FString UpdateDefaultValueName;
+	FVector UpdateValue;
+
+	const FAnimNode_ModifyBone* ModifyBone = static_cast<const FAnimNode_ModifyBone*>(AnimNode);
+
+	if (!ModifyBone)
+	{
+		return;
+	}
+
+	UpdateDefaultValueName = FString("Translation");
+	UpdateValue = ModifyBone->Translation;
+	SetDefaultValue(UpdateDefaultValueName, UpdateValue);
+
+	UpdateDefaultValueName = FString("Rotation");
+	UpdateValue.X = ModifyBone->Rotation.Pitch;
+	UpdateValue.Y = ModifyBone->Rotation.Yaw;
+	UpdateValue.Z = ModifyBone->Rotation.Roll;
+	SetDefaultValue(UpdateDefaultValueName, UpdateValue);
+
+	UpdateDefaultValueName = FString("Scale");
+	UpdateValue = ModifyBone->Scale;
+	SetDefaultValue(UpdateDefaultValueName, UpdateValue);
+}
+
+void UAnimGraphNode_ModifyBone::DoTranslation(const USkeletalMeshComponent* SkelComp, FVector& Drag, FAnimNode_Base* InOutAnimNode)
+{
+	FAnimNode_ModifyBone* ModifyBone = static_cast<FAnimNode_ModifyBone*>(InOutAnimNode);
+
+	if (ModifyBone->TranslationMode != EBoneModificationMode::BMM_Ignore)
+	{
+		FVector Offset = ConvertCSVectorToBoneSpace(SkelComp, Drag, ModifyBone->ForwardedPose, Node.BoneToModify.BoneName, Node.TranslationSpace);
+
+		ModifyBone->Translation += Offset;
+
+		Node.Translation = ModifyBone->Translation;
+	}
+}
+
+void UAnimGraphNode_ModifyBone::DoRotation(const USkeletalMeshComponent* SkelComp, FRotator& Rotation, FAnimNode_Base* InOutAnimNode)
+{
+	FAnimNode_ModifyBone* ModifyBone = static_cast<FAnimNode_ModifyBone*>(InOutAnimNode);
+
+	if (Node.RotationMode != EBoneModificationMode::BMM_Ignore)
+	{
+		FQuat DeltaQuat = ConvertCSRotationToBoneSpace(SkelComp, Rotation, ModifyBone->ForwardedPose, Node.BoneToModify.BoneName, Node.RotationSpace);
+
+		FQuat PrevQuat(ModifyBone->Rotation);
+		FQuat NewQuat = DeltaQuat * PrevQuat;
+		ModifyBone->Rotation = NewQuat.Rotator();
+		Node.Rotation = ModifyBone->Rotation;
+	}
+}
+
+void UAnimGraphNode_ModifyBone::DoScale(const USkeletalMeshComponent* SkelComp, FVector& Scale, FAnimNode_Base* InOutAnimNode)
+{
+	FAnimNode_ModifyBone* ModifyBone = static_cast<FAnimNode_ModifyBone*>(InOutAnimNode);
+
+	if (Node.ScaleMode != EBoneModificationMode::BMM_Ignore)
+	{
+		FVector Offset = Scale;
+		ModifyBone->Scale += Offset;
+		Node.Scale = ModifyBone->Scale;
+	}
+}
+
+void UAnimGraphNode_ModifyBone::CopyNodeDataTo(FAnimNode_Base* OutAnimNode)
+{
+	FAnimNode_ModifyBone* ModifyBone = static_cast<FAnimNode_ModifyBone*>(OutAnimNode);
+
+	// copies Pin values from the internal node to get data which are not compiled yet
+	ModifyBone->Translation = Node.Translation;
+	ModifyBone->Rotation = Node.Rotation;
+	ModifyBone->Scale = Node.Scale;
+
+	// copies Modes
+	ModifyBone->TranslationMode = Node.TranslationMode;
+	ModifyBone->RotationMode = Node.RotationMode;
+	ModifyBone->ScaleMode = Node.ScaleMode;
+
+	// copies Spaces
+	ModifyBone->TranslationSpace = Node.TranslationSpace;
+	ModifyBone->RotationSpace = Node.RotationSpace;
+	ModifyBone->ScaleSpace = Node.ScaleSpace;
+}
+
+void UAnimGraphNode_ModifyBone::CopyNodeDataFrom(const FAnimNode_Base* InNewAnimNode)
+{
+	const FAnimNode_ModifyBone* ModifyBone = static_cast<const FAnimNode_ModifyBone*>(InNewAnimNode);
+
+	// copies Pin data from updated values
+	Node.Translation = ModifyBone->Translation;
+	Node.Rotation = ModifyBone->Rotation;
+	Node.Scale = ModifyBone->Scale;
+}
 #undef LOCTEXT_NAMESPACE
diff --git a/Engine/Source/Editor/AnimGraph/Private/AnimGraphNode_SkeletalControlBase.cpp b/Engine/Source/Editor/AnimGraph/Private/AnimGraphNode_SkeletalControlBase.cpp
--- a/Engine/Source/Editor/AnimGraph/Private/AnimGraphNode_SkeletalControlBase.cpp
+++ b/Engine/Source/Editor/AnimGraph/Private/AnimGraphNode_SkeletalControlBase.cpp
@@ -54,4 +54,235 @@ void UAnimGraphNode_SkeletalControlBase::GetMenuEntries(FGraphContextMenuBuilder
 	NodeAction->NodeTemplate = TemplateNode;
 }
 
+void UAnimGraphNode_SkeletalControlBase::ConvertToComponentSpaceTransform(const USkeletalMeshComponent* SkelComp, const FTransform & InTransform, FTransform & OutCSTransform, int32 BoneIndex, EBoneControlSpace Space) const
+{
+	USkeleton * Skeleton = SkelComp->SkeletalMesh->Skeleton;
+
+	switch (Space)
+	{
+	case BCS_WorldSpace:
+	{
+		OutCSTransform = InTransform;
+		OutCSTransform.SetToRelativeTransform(SkelComp->ComponentToWorld);
+	}
+		break;
+
+	case BCS_ComponentSpace:
+	{
+		// Component Space, no change.
+		OutCSTransform = InTransform;
+	}
+		break;
+
+	case BCS_ParentBoneSpace:
+		if (BoneIndex != INDEX_NONE)
+		{
+			const int32 ParentIndex = Skeleton->GetReferenceSkeleton().GetParentIndex(BoneIndex);
+			if (ParentIndex != INDEX_NONE)																																																																																																																															
+			{
+				const int32 MeshParentIndex = Skeleton->GetMeshBoneIndexFromSkeletonBoneIndex(SkelComp->SkeletalMesh, ParentIndex);
+				if (MeshParentIndex != INDEX_NONE)
+				{
+					const FTransform ParentTM = SkelComp->GetBoneTransform(MeshParentIndex);
+					OutCSTransform = InTransform * ParentTM;
+				}
+				else
+				{
+					OutCSTransform = InTransform;
+				}
+			}
+		}
+		break;
+
+	case BCS_BoneSpace:
+		if (BoneIndex != INDEX_NONE)
+		{
+			const int32 MeshBoneIndex = Skeleton->GetMeshBoneIndexFromSkeletonBoneIndex(SkelComp->SkeletalMesh, BoneIndex);
+			if (MeshBoneIndex != INDEX_NONE)
+			{
+				const FTransform BoneTM = SkelComp->GetBoneTransform(MeshBoneIndex);
+				OutCSTransform = InTransform * BoneTM;
+			}			
+			else
+			{
+				OutCSTransform = InTransform;
+			}
+		}
+		break;
+
+	default:
+		if (SkelComp->SkeletalMesh)
+		{
+			UE_LOG(LogAnimation, Warning, TEXT("ConvertToComponentSpaceTransform: Unknown BoneSpace %d  for Mesh: %s"), (uint8)Space, *SkelComp->SkeletalMesh->GetFName().ToString());
+		}
+		else
+		{
+			UE_LOG(LogAnimation, Warning, TEXT("ConvertToComponentSpaceTransform: Unknown BoneSpace %d  for Skeleton: %s"), (uint8)Space, *Skeleton->GetFName().ToString());
+		}
+		break;
+	}
+}
+
+
+FVector UAnimGraphNode_SkeletalControlBase::ConvertCSVectorToBoneSpace(const USkeletalMeshComponent* SkelComp, FVector& InCSVector, FA2CSPose& MeshBases, const FName& BoneName, const EBoneControlSpace Space)
+{
+	FVector OutVector = InCSVector;
+
+	if (MeshBases.IsValid())
+	{
+		int32 MeshBoneIndex = SkelComp->GetBoneIndex(BoneName);
+
+		switch (Space)
+		{
+			// World Space, no change in preview window
+		case BCS_WorldSpace:
+		case BCS_ComponentSpace:
+			// Component Space, no change.
+			break;
+
+		case BCS_ParentBoneSpace:
+		{
+			const int32 ParentIndex = MeshBases.GetParentBoneIndex(MeshBoneIndex);
+			if (ParentIndex != INDEX_NONE)
+			{
+				FTransform ParentTM = MeshBases.GetComponentSpaceTransform(ParentIndex);
+				OutVector = ParentTM.InverseTransformVector(InCSVector);
+			}
+		}
+			break;
+
+		case BCS_BoneSpace:
+		{
+			FTransform BoneTM = MeshBases.GetComponentSpaceTransform(MeshBoneIndex);
+			OutVector = BoneTM.InverseTransformVector(InCSVector);
+		}
+			break;
+		}
+	}
+
+	return OutVector;
+}
+
+FQuat UAnimGraphNode_SkeletalControlBase::ConvertCSRotationToBoneSpace(const USkeletalMeshComponent* SkelComp, FRotator& InCSRotator, FA2CSPose& MeshBases, const FName& BoneName, const EBoneControlSpace Space)
+{
+	FQuat OutQuat = FQuat::Identity;
+
+	if (MeshBases.IsValid())
+	{
+		int32 MeshBoneIndex = SkelComp->GetBoneIndex(BoneName);
+
+		FVector RotAxis;
+		float RotAngle;
+		InCSRotator.Quaternion().ToAxisAndAngle(RotAxis, RotAngle);
+
+		switch (Space)
+		{
+			// World Space, no change in preview window
+		case BCS_WorldSpace:
+		case BCS_ComponentSpace:
+			// Component Space, no change.
+			OutQuat = InCSRotator.Quaternion();
+			break;
+
+		case BCS_ParentBoneSpace:
+		{
+			const int32 ParentIndex = MeshBases.GetParentBoneIndex(MeshBoneIndex);
+			if (ParentIndex != INDEX_NONE)
+			{
+				FTransform ParentTM = MeshBases.GetComponentSpaceTransform(ParentIndex);
+				ParentTM = ParentTM.Inverse();
+				//Calculate the new delta rotation
+				FVector4 BoneSpaceAxis = ParentTM.TransformVector(RotAxis);
+				FQuat DeltaQuat(BoneSpaceAxis, RotAngle);
+				DeltaQuat.Normalize();
+				OutQuat = DeltaQuat;
+			}
+		}
+			break;
+
+		case BCS_BoneSpace:
+		{
+			FTransform BoneTM = MeshBases.GetComponentSpaceTransform(MeshBoneIndex);
+			BoneTM = BoneTM.Inverse();
+			FVector4 BoneSpaceAxis = BoneTM.TransformVector(RotAxis);
+			//Calculate the new delta rotation
+			FQuat DeltaQuat(BoneSpaceAxis, RotAngle);
+			DeltaQuat.Normalize();
+			OutQuat = DeltaQuat;
+		}
+			break;
+		}
+	}
+
+	return OutQuat;
+}
+
+FVector UAnimGraphNode_SkeletalControlBase::ConvertWidgetLocation(const USkeletalMeshComponent* SkelComp, FA2CSPose& MeshBases, const FName& BoneName, const FVector& Location, const EBoneControlSpace Space)
+{
+	FVector WidgetLoc = FVector::ZeroVector;
+
+	if (MeshBases.IsValid())
+	{
+		USkeleton * Skeleton = SkelComp->SkeletalMesh->Skeleton;
+		int32 MeshBoneIndex = SkelComp->GetBoneIndex(BoneName);
+
+		switch (Space)
+		{
+			// ComponentToWorld must be Identity in preview window so same as ComponentSpace
+		case BCS_WorldSpace:
+		case BCS_ComponentSpace:
+		{
+			// Component Space, no change.
+			WidgetLoc = Location;
+		}
+			break;
+
+		case BCS_ParentBoneSpace:
+
+			if (MeshBoneIndex != INDEX_NONE)
+			{
+				const int32 MeshParentIndex = MeshBases.GetParentBoneIndex(MeshBoneIndex);
+				if (MeshParentIndex != INDEX_NONE)
+				{
+					const FTransform ParentTM = MeshBases.GetComponentSpaceTransform(MeshParentIndex);
+					WidgetLoc = ParentTM.TransformPosition(Location);
+				}
+			}
+			break;
+
+		case BCS_BoneSpace:
+
+			if (MeshBoneIndex != INDEX_NONE)
+			{
+				FTransform BoneTM = MeshBases.GetComponentSpaceTransform(MeshBoneIndex);
+				WidgetLoc = BoneTM.TransformPosition(Location);
+			}
+		}
+	}
+
+	return WidgetLoc;
+}
+
+void UAnimGraphNode_SkeletalControlBase::SetDefaultValue(FString& UpdateDefaultValueName, FVector& Value)
+{
+	for (UEdGraphPin* Pin : Pins)
+	{
+		if (Pin->PinName == UpdateDefaultValueName)
+		{
+			if (GetSchema()->IsCurrentPinDefaultValid(Pin).IsEmpty())
+			{
+				FVector& V = Value;
+				FString Str = FString::Printf(TEXT("%.3f,%.3f,%.3f"), V.X, V.Y, V.Z);
+				if (Pin->DefaultValue != Str)
+				{
+					PreEditChange(NULL);
+					GetSchema()->TrySetDefaultValue(*Pin, Str);
+					PostEditChange();
+					break;
+				}
+			}
+		}
+	}
+}
+
 #undef LOCTEXT_NAMESPACE
\ No newline at end of file
diff --git a/Engine/Source/Editor/AnimGraph/Private/AnimGraphNode_TwoBoneIK.cpp b/Engine/Source/Editor/AnimGraph/Private/AnimGraphNode_TwoBoneIK.cpp
--- a/Engine/Source/Editor/AnimGraph/Private/AnimGraphNode_TwoBoneIK.cpp
+++ b/Engine/Source/Editor/AnimGraph/Private/AnimGraphNode_TwoBoneIK.cpp
@@ -2,6 +2,7 @@
 
 #include "AnimGraphPrivatePCH.h"
 #include "AnimGraphNode_TwoBoneIK.h"
+#include "DebugRenderSceneProxy.h"
 
 // for customization details
 #include "../../PropertyEditor/Public/PropertyHandle.h"
@@ -27,6 +28,8 @@ public:
 	}
 };
 
+TSharedPtr<FTwoBoneIKDelegate> UAnimGraphNode_TwoBoneIK::TwoBoneIKDelegate = NULL;
+
 /////////////////////////////////////////////////////
 // UAnimGraphNode_TwoBoneIK
 
@@ -34,6 +37,7 @@ public:
 UAnimGraphNode_TwoBoneIK::UAnimGraphNode_TwoBoneIK(const FObjectInitializer& ObjectInitializer)
 	: Super(ObjectInitializer)
 {
+	BoneSelectMode = BSM_EndEffector;
 }
 
 FText UAnimGraphNode_TwoBoneIK::GetControllerDescription() const
@@ -77,105 +81,286 @@ void UAnimGraphNode_TwoBoneIK::Draw( FPrimitiveDrawInterface* PDI, USkeletalMesh
 {
 	if (SkelMeshComp && SkelMeshComp->SkeletalMesh && SkelMeshComp->SkeletalMesh->Skeleton)
 	{
-		USkeleton* Skeleton = SkelMeshComp->SkeletalMesh->Skeleton;
+		USkeleton * Skeleton = SkelMeshComp->SkeletalMesh->Skeleton;
 
 		DrawTargetLocation(PDI, SkelMeshComp, Skeleton, Node.EffectorLocationSpace, Node.EffectorSpaceBoneName, Node.EffectorLocation, FColor(255, 128, 128), FColor(180, 128, 128));
 		DrawTargetLocation(PDI, SkelMeshComp, Skeleton, Node.JointTargetLocationSpace, Node.JointTargetSpaceBoneName, Node.JointTargetLocation, FColor(128, 255, 128), FColor(128, 180, 128));
 	}
 }
 
-void UAnimGraphNode_TwoBoneIK::DrawTargetLocation(FPrimitiveDrawInterface* PDI, USkeletalMeshComponent* SkelComp, USkeleton* Skeleton, uint8 SpaceBase, FName SpaceBoneName, const FVector& TargetLocation, const FColor& TargetColor, const FColor& BoneColor) const
+void UAnimGraphNode_TwoBoneIK::DrawTargetLocation(FPrimitiveDrawInterface* PDI, USkeletalMeshComponent* SkelComp, USkeleton * Skeleton, EBoneControlSpace SpaceBase, FName SpaceBoneName, const FVector & TargetLocation, const FColor & TargetColor, const FColor & BoneColor) const
 {
 	const bool bInBoneSpace = (SpaceBase == BCS_ParentBoneSpace) || (SpaceBase == BCS_BoneSpace);
 	const int32 SpaceBoneIndex = bInBoneSpace ? Skeleton->GetReferenceSkeleton().FindBoneIndex(SpaceBoneName) : INDEX_NONE;
 	// Transform EffectorLocation from EffectorLocationSpace to ComponentSpace.
 	FTransform TargetTransform = FTransform (TargetLocation);
 	FTransform CSTransform;
-	ConvertToComponentSpaceTransform(SkelComp, Skeleton, TargetTransform, CSTransform, SpaceBoneIndex, SpaceBase);
+
+	ConvertToComponentSpaceTransform(SkelComp, TargetTransform, CSTransform, SpaceBoneIndex, SpaceBase);
 
 	FTransform WorldTransform = CSTransform * SkelComp->ComponentToWorld;
 
+#if 0 // @TODO : remove this code because this doesn't show correct target location
 	DrawCoordinateSystem( PDI, WorldTransform.GetLocation(), WorldTransform.GetRotation().Rotator(), 20.f, SDPG_Foreground );
 	DrawWireDiamond( PDI, WorldTransform.ToMatrixWithScale(), 2.f, TargetColor, SDPG_Foreground );
+#endif
 
 	if (bInBoneSpace)
 	{
-		ConvertToComponentSpaceTransform(SkelComp, Skeleton, FTransform::Identity, CSTransform, SpaceBoneIndex, SpaceBase);
+		ConvertToComponentSpaceTransform(SkelComp, FTransform::Identity, CSTransform, SpaceBoneIndex, SpaceBase);
 		WorldTransform = CSTransform * SkelComp->ComponentToWorld;
 		DrawCoordinateSystem( PDI, WorldTransform.GetLocation(), WorldTransform.GetRotation().Rotator(), 20.f, SDPG_Foreground );
 		DrawWireDiamond( PDI, WorldTransform.ToMatrixWithScale(), 2.f, BoneColor, SDPG_Foreground );
 	}
 }
 
-void UAnimGraphNode_TwoBoneIK::ConvertToComponentSpaceTransform(USkeletalMeshComponent* SkelComp, USkeleton* Skeleton, const FTransform& InTransform, FTransform& OutCSTransform, int32 BoneIndex, uint8 Space) const
+void UAnimGraphNode_TwoBoneIK::CopyNodeDataTo(FAnimNode_Base* AnimNode)
+{
+	FAnimNode_TwoBoneIK* TwoBoneIK = static_cast<FAnimNode_TwoBoneIK*>(AnimNode);
+
+	// copies Pin values from the internal node to get data which are not compiled yet
+	TwoBoneIK->EffectorLocation = Node.EffectorLocation;
+	TwoBoneIK->JointTargetLocation = Node.JointTargetLocation;
+}
+
+void UAnimGraphNode_TwoBoneIK::CopyNodeDataFrom(const FAnimNode_Base* InNewAnimNode)
 {
-	switch( Space )
+	const FAnimNode_TwoBoneIK* TwoBoneIK = static_cast<const FAnimNode_TwoBoneIK*>(InNewAnimNode);
+	
+	// copies Pin data from updated values
+	Node.EffectorLocation = TwoBoneIK->EffectorLocation;
+	Node.JointTargetLocation = TwoBoneIK->JointTargetLocation;
+}
+
+
+FVector UAnimGraphNode_TwoBoneIK::GetWidgetLocation(const USkeletalMeshComponent* SkelComp, FAnimNode_SkeletalControlBase* AnimNode)
+{
+	EBoneControlSpace Space;
+	FName BoneName;
+	FVector Location;
+	 
+	if (BoneSelectMode == BSM_EndEffector)
 	{
-	case BCS_WorldSpace : 
-		{
-			OutCSTransform = InTransform;
-			OutCSTransform.SetToRelativeTransform(SkelComp->ComponentToWorld);
-		}
-		break;
+		Space = Node.EffectorLocationSpace;
+		Location = Node.EffectorLocation;
+		BoneName = Node.EffectorSpaceBoneName;
+
+	}
+	else // BSM_JointTarget
+	{
+		Space = Node.JointTargetLocationSpace;
 
-	case BCS_ComponentSpace :
+		if (Space == BCS_WorldSpace || Space == BCS_ComponentSpace)
 		{
-			// Component Space, no change.
-			OutCSTransform = InTransform;
+			return FVector::ZeroVector;
 		}
-		break;
+		Location = Node.JointTargetLocation;
+		BoneName = Node.JointTargetSpaceBoneName;
+	}
+
+	return ConvertWidgetLocation(SkelComp, AnimNode->ForwardedPose, BoneName, Location, Space);
+}
+
+void UAnimGraphNode_TwoBoneIK::UpdateDefaultValues(const FAnimNode_Base* AnimNode)
+{
+	FString UpdateDefaultValueName;
+	FVector UpdateValue;
+
+	const FAnimNode_TwoBoneIK* TwoBoneIK = static_cast<const FAnimNode_TwoBoneIK*>(AnimNode);
+
+	if (TwoBoneIK == NULL)
+	{
+		return;
+	}
+
+	if (BoneSelectMode == BSM_EndEffector)
+	{
+		UpdateDefaultValueName = FString("EffectorLocation");
+		UpdateValue = TwoBoneIK->EffectorLocation;
+	}
+	else
+	{
+		UpdateDefaultValueName = FString("JointTargetLocation");
+		UpdateValue = TwoBoneIK->JointTargetLocation;
+	}
+
+	SetDefaultValue(UpdateDefaultValueName, UpdateValue);
+}
+
+void UAnimGraphNode_TwoBoneIK::UpdateAllDefaultValues(const FAnimNode_Base* AnimNode)
+{
+	FString UpdateDefaultValueName;
+	FVector UpdateValue;
+
+	const FAnimNode_TwoBoneIK* TwoBoneIK = static_cast<const FAnimNode_TwoBoneIK*>(AnimNode);
+
+	if (TwoBoneIK == NULL)
+	{
+		return;
+	}
+
+	UpdateDefaultValueName = FString("EffectorLocation");
+	UpdateValue = TwoBoneIK->EffectorLocation;
+	SetDefaultValue(UpdateDefaultValueName, UpdateValue);
+
+	UpdateDefaultValueName = FString("JointTargetLocation");
+	UpdateValue = TwoBoneIK->JointTargetLocation;
+	SetDefaultValue(UpdateDefaultValueName, UpdateValue);
+}
+
+
+FWidget::EWidgetMode UAnimGraphNode_TwoBoneIK::GetWidgetMode(const USkeletalMeshComponent* SkelComp)
+{
+	int32 BoneIndex = SkelComp->GetBoneIndex(Node.IKBone.BoneName);
+	// Two bone IK node uses only Translate
+	if (BoneIndex != INDEX_NONE)
+	{
+		return FWidget::WM_Translate;
+	}
 
-	case BCS_ParentBoneSpace :
-		if( BoneIndex != INDEX_NONE )
+	return FWidget::WM_None;
+}
+
+void UAnimGraphNode_TwoBoneIK::MoveSelectActorLocation(const USkeletalMeshComponent* SkelComp, FAnimNode_SkeletalControlBase* AnimNode)
+{
+	USkeleton * Skeleton = SkelComp->SkeletalMesh->Skeleton;
+
+	// create a bone select actor
+
+	if (!BoneSelectActor.IsValid())
+	{
+		if (Node.JointTargetLocationSpace == EBoneControlSpace::BCS_BoneSpace ||
+			Node.JointTargetLocationSpace == EBoneControlSpace::BCS_ParentBoneSpace)
 		{
-			const int32 ParentIndex = Skeleton->GetReferenceSkeleton().GetParentIndex(BoneIndex);
-			if( ParentIndex != INDEX_NONE )
+			int32 JointTargetIndex = Skeleton->GetReferenceSkeleton().FindBoneIndex(Node.JointTargetSpaceBoneName);
+
+			if (JointTargetIndex != INDEX_NONE)
 			{
-				const int32 MeshParentIndex = Skeleton->GetMeshBoneIndexFromSkeletonBoneIndex(SkelComp->SkeletalMesh, ParentIndex);
-				if (MeshParentIndex != INDEX_NONE)
-				{
-					const FTransform ParentTM = SkelComp->GetBoneTransform(MeshParentIndex);
-					OutCSTransform = InTransform * ParentTM;
-				}
-				else
-				{
-					OutCSTransform = InTransform;
-				}
+				UWorld* World = SkelComp->GetWorld();
+				check(World);
+
+				BoneSelectActor = World->SpawnActor<ABoneSelectActor>(FVector(0), FRotator(0, 0, 0));
+				check(BoneSelectActor.IsValid());
 			}
 		}
-		break;
+	}
+
+	if (!BoneSelectActor.IsValid())
+	{
+		return;
+	}
 
-	case BCS_BoneSpace :
-		if( BoneIndex != INDEX_NONE )
+	// move the actor's position
+	if (BoneSelectMode == BSM_JointTarget)
+	{
+		FName BoneName;
+		int32 EffectorBoneIndex = SkelComp->GetBoneIndex(Node.EffectorSpaceBoneName);
+		if (EffectorBoneIndex != INDEX_NONE)
 		{
-			const int32 MeshBoneIndex = Skeleton->GetMeshBoneIndexFromSkeletonBoneIndex(SkelComp->SkeletalMesh, BoneIndex);
-			if (MeshBoneIndex != INDEX_NONE)
-			{
-				const FTransform BoneTM = SkelComp->GetBoneTransform(MeshBoneIndex);
-				OutCSTransform = InTransform * BoneTM;
-			}
-			else
-			{
-				OutCSTransform = InTransform;
-			}
+			BoneName = Node.EffectorSpaceBoneName;
+		}
+		else
+		{
+			BoneName = Node.IKBone.BoneName;
 		}
-		break;
 
-	default:
-		if( SkelComp->SkeletalMesh )
+		FVector ActorLocation = ConvertWidgetLocation(SkelComp, AnimNode->ForwardedPose, BoneName, Node.EffectorLocation, Node.EffectorLocationSpace);
+		BoneSelectActor->SetActorLocation(ActorLocation + FVector(0, 10, 0));
+
+	}
+	else if (BoneSelectMode == BSM_EndEffector)
+	{
+		int32 JointTargetIndex = SkelComp->GetBoneIndex(Node.JointTargetSpaceBoneName);
+
+		if (JointTargetIndex != INDEX_NONE)
+		{
+			FVector ActorLocation = ConvertWidgetLocation(SkelComp, AnimNode->ForwardedPose, Node.JointTargetSpaceBoneName, Node.JointTargetLocation, Node.JointTargetLocationSpace);
+			BoneSelectActor->SetActorLocation(ActorLocation + FVector(0, 10, 0));
+		}
+	}
+}
+
+FName UAnimGraphNode_TwoBoneIK::FindSelectedBone()
+{
+	FName SelectedBone;
+
+	// should return mesh bone index
+	if (BoneSelectMode == BSM_EndEffector)
+	{
+		if (Node.EffectorLocationSpace == EBoneControlSpace::BCS_BoneSpace ||
+			Node.EffectorLocationSpace == EBoneControlSpace::BCS_ParentBoneSpace)
 		{
-			UE_LOG(LogAnimation, Warning, TEXT("ConvertToComponentSpaceTransform: Unknown BoneSpace %d  for Mesh: %s"), Space, *SkelComp->SkeletalMesh->GetFName().ToString() );
+			SelectedBone = Node.EffectorSpaceBoneName;
 		}
 		else
 		{
-			UE_LOG(LogAnimation, Warning, TEXT("ConvertToComponentSpaceTransform: Unknown BoneSpace %d  for Skeleton: %s"), Space, *Skeleton->GetFName().ToString() );
+			SelectedBone = Node.IKBone.BoneName;
 		}
-		break;
+
+	}
+	else if (BoneSelectMode == BSM_JointTarget)
+	{
+		SelectedBone = Node.JointTargetSpaceBoneName;
 	}
+
+	return SelectedBone;
 }
 
-// just for refreshing UIs when bone space was changed
-TSharedPtr<FTwoBoneIKDelegate> TwoBoneIKDelegate;
+bool UAnimGraphNode_TwoBoneIK::IsActorClicked(HActor* ActorHitProxy)
+{
+	ABoneSelectActor* BoneActor = Cast<ABoneSelectActor>(ActorHitProxy->Actor);
+
+	if (BoneActor)
+	{
+		return true;
+	}
+
+	return false;
+}
+
+void UAnimGraphNode_TwoBoneIK::ProcessActorClick(HActor* ActorHitProxy)
+{
+	// toggle bone selection mode
+	if (BoneSelectMode == BSM_EndEffector)
+	{
+		BoneSelectMode = BSM_JointTarget;
+	}
+	else
+	{
+		BoneSelectMode = BSM_EndEffector;
+	}
+}
+
+void UAnimGraphNode_TwoBoneIK::DoTranslation(const USkeletalMeshComponent* SkelComp, FVector& Drag, FAnimNode_Base* InOutAnimNode)
+{
+	FAnimNode_TwoBoneIK* TwoBoneIK = static_cast<FAnimNode_TwoBoneIK*>(InOutAnimNode);
+
+	if (BoneSelectMode == BSM_EndEffector)
+	{
+		FVector Offset = ConvertCSVectorToBoneSpace(SkelComp, Drag, TwoBoneIK->ForwardedPose, FindSelectedBone(), Node.EffectorLocationSpace);
+
+		TwoBoneIK->EffectorLocation += Offset;
+		Node.EffectorLocation = TwoBoneIK->EffectorLocation;
+	}
+	else if (BoneSelectMode == BSM_JointTarget)
+	{
+		FVector Offset = ConvertCSVectorToBoneSpace(SkelComp, Drag, TwoBoneIK->ForwardedPose, FindSelectedBone(), Node.JointTargetLocationSpace);
+
+		TwoBoneIK->JointTargetLocation += Offset;
+		Node.JointTargetLocation = TwoBoneIK->JointTargetLocation;
+	}
+}
+
+void UAnimGraphNode_TwoBoneIK::DeselectActor(USkeletalMeshComponent* SkelComp)
+{
+	if(BoneSelectActor!=NULL)
+	{
+		if (SkelComp->GetWorld()->DestroyActor(BoneSelectActor.Get()))
+		{
+			BoneSelectActor = NULL;
+		}
+	}
+}
 
 void UAnimGraphNode_TwoBoneIK::CustomizeDetails(class IDetailLayoutBuilder& DetailBuilder)
 {
@@ -186,61 +371,72 @@ void UAnimGraphNode_TwoBoneIK::CustomizeDetails(class IDetailLayoutBuilder& Deta
 	}
 
 	EBoneControlSpace Space = Node.EffectorLocationSpace;
+	const FString TakeRotationPropName = FString::Printf(TEXT("Node.%s"), GET_MEMBER_NAME_STRING_CHECKED(FAnimNode_TwoBoneIK, bTakeRotationFromEffectorSpace));
+	const FString MaintainEffectorPropName = FString::Printf(TEXT("Node.%s"), GET_MEMBER_NAME_STRING_CHECKED(FAnimNode_TwoBoneIK, bMaintainEffectorRelRot));
+	const FString EffectorBoneName = FString::Printf(TEXT("Node.%s"), GET_MEMBER_NAME_STRING_CHECKED(FAnimNode_TwoBoneIK, EffectorSpaceBoneName));
+	const FString EffectorLocationPropName = FString::Printf(TEXT("Node.%s"), GET_MEMBER_NAME_STRING_CHECKED(FAnimNode_TwoBoneIK, EffectorLocation));
+
 	if (Space == BCS_BoneSpace || Space == BCS_ParentBoneSpace)
 	{
 		IDetailCategoryBuilder& IKCategory = DetailBuilder.EditCategory("IK");
 		IDetailCategoryBuilder& EffectorCategory = DetailBuilder.EditCategory("EndEffector");
 		TSharedPtr<IPropertyHandle> PropertyHandle;
-		PropertyHandle = DetailBuilder.GetProperty(FName("Node.bTakeRotationFromEffectorSpace"), GetClass());
+		PropertyHandle = DetailBuilder.GetProperty(*TakeRotationPropName, GetClass());
 		EffectorCategory.AddProperty(PropertyHandle);
-		PropertyHandle = DetailBuilder.GetProperty(FName("Node.bMaintainEffectorRelRot"), GetClass());
+		PropertyHandle = DetailBuilder.GetProperty(*MaintainEffectorPropName, GetClass());
 		EffectorCategory.AddProperty(PropertyHandle);
-		PropertyHandle = DetailBuilder.GetProperty(FName("Node.EffectorSpaceBoneName"), GetClass());
+		PropertyHandle = DetailBuilder.GetProperty(*EffectorBoneName, GetClass());
 		EffectorCategory.AddProperty(PropertyHandle);
 	}
 	else // hide all properties in EndEffector category
 	{
-		TSharedPtr<IPropertyHandle> PropertyHandle = DetailBuilder.GetProperty(FName("Node.EffectorLocation"), GetClass());
+		TSharedPtr<IPropertyHandle> PropertyHandle = DetailBuilder.GetProperty(*EffectorLocationPropName, GetClass());
 		DetailBuilder.HideProperty(PropertyHandle);
-		PropertyHandle = DetailBuilder.GetProperty(FName("Node.bTakeRotationFromEffectorSpace"), GetClass());
+		PropertyHandle = DetailBuilder.GetProperty(*TakeRotationPropName, GetClass());
 		DetailBuilder.HideProperty(PropertyHandle);
-		PropertyHandle = DetailBuilder.GetProperty(FName("Node.bMaintainEffectorRelRot"), GetClass());
+		PropertyHandle = DetailBuilder.GetProperty(*MaintainEffectorPropName, GetClass());
 		DetailBuilder.HideProperty(PropertyHandle);
-		PropertyHandle = DetailBuilder.GetProperty(FName("Node.EffectorSpaceBoneName"), GetClass());
+		PropertyHandle = DetailBuilder.GetProperty(*EffectorBoneName, GetClass());
 		DetailBuilder.HideProperty(PropertyHandle);
 	}
 
 	Space = Node.JointTargetLocationSpace;
+	bool bPinVisibilityChanged = false;
+	const FString JointTargetSpaceBoneName = FString::Printf(TEXT("Node.%s"), GET_MEMBER_NAME_STRING_CHECKED(FAnimNode_TwoBoneIK, JointTargetSpaceBoneName));
+	const FString JointTargetLocation = FString::Printf(TEXT("Node.%s"), GET_MEMBER_NAME_STRING_CHECKED(FAnimNode_TwoBoneIK, JointTargetLocation));
+
 	if (Space == BCS_BoneSpace || Space == BCS_ParentBoneSpace)
 	{
 		IDetailCategoryBuilder& IKCategory = DetailBuilder.EditCategory("IK");
 		IDetailCategoryBuilder& EffectorCategory = DetailBuilder.EditCategory("JointTarget");
 		TSharedPtr<IPropertyHandle> PropertyHandle;
-		PropertyHandle = DetailBuilder.GetProperty(FName("Node.JointTargetSpaceBoneName"), GetClass());
+		PropertyHandle = DetailBuilder.GetProperty(*JointTargetSpaceBoneName, GetClass());
 		EffectorCategory.AddProperty(PropertyHandle);
 
-		SetPinsVisibility(true);
+		bPinVisibilityChanged = SetPinsVisibility(true);
 	}
 	else // hide all properties in JointTarget category except for JointTargetLocationSpace
 	{
-		TSharedPtr<IPropertyHandle> PropertyHandle = DetailBuilder.GetProperty(FName("Node.JointTargetLocation"), GetClass());
+		TSharedPtr<IPropertyHandle> PropertyHandle = DetailBuilder.GetProperty(*JointTargetLocation, GetClass());
 		DetailBuilder.HideProperty(PropertyHandle);
-		PropertyHandle = DetailBuilder.GetProperty(FName("Node.JointTargetSpaceBoneName"), GetClass());
+		PropertyHandle = DetailBuilder.GetProperty(*JointTargetSpaceBoneName, GetClass());
 		DetailBuilder.HideProperty(PropertyHandle);
 
-		SetPinsVisibility(false);
+		bPinVisibilityChanged = SetPinsVisibility(false);
 	}
 
+	const FString EffectorLocationSpace = FString::Printf(TEXT("Node.%s"), GET_MEMBER_NAME_STRING_CHECKED(FAnimNode_TwoBoneIK, EffectorLocationSpace));
+	const FString JointTargetLocationSpace = FString::Printf(TEXT("Node.%s"), GET_MEMBER_NAME_STRING_CHECKED(FAnimNode_TwoBoneIK, JointTargetLocationSpace));
 
 	// refresh UIs when bone space is changed
-	TSharedRef<IPropertyHandle> EffectorLocHandle = DetailBuilder.GetProperty(FName("Node.EffectorLocationSpace"), GetClass());
+	TSharedRef<IPropertyHandle> EffectorLocHandle = DetailBuilder.GetProperty(*EffectorLocationSpace, GetClass());
 	if (EffectorLocHandle->IsValidHandle())
 	{
 		FSimpleDelegate UpdateEffectorSpaceDelegate = FSimpleDelegate::CreateSP(TwoBoneIKDelegate.Get(), &FTwoBoneIKDelegate::UpdateLocationSpace, &DetailBuilder);
 		EffectorLocHandle->SetOnPropertyValueChanged(UpdateEffectorSpaceDelegate);
 	}
 
-	TSharedRef<IPropertyHandle> JointTragetLocHandle = DetailBuilder.GetProperty(FName("Node.JointTargetLocationSpace"), GetClass());
+	TSharedRef<IPropertyHandle> JointTragetLocHandle = DetailBuilder.GetProperty(*JointTargetLocationSpace, GetClass());
 	if (JointTragetLocHandle->IsValidHandle())
 	{
 		FSimpleDelegate UpdateJointSpaceDelegate = FSimpleDelegate::CreateSP(TwoBoneIKDelegate.Get(), &FTwoBoneIKDelegate::UpdateLocationSpace, &DetailBuilder);
@@ -248,20 +444,31 @@ void UAnimGraphNode_TwoBoneIK::CustomizeDetails(class IDetailLayoutBuilder& Deta
 	}
 
 	// reconstruct node for showing/hiding Pins
-	ReconstructNode();
+	if (bPinVisibilityChanged)
+	{
+		ReconstructNode();
+	}
 }
 
-void UAnimGraphNode_TwoBoneIK::SetPinsVisibility(bool bShow)
+bool UAnimGraphNode_TwoBoneIK::SetPinsVisibility(bool bShow)
 {
+	bool bChanged = false;
+
 	for (FOptionalPinFromProperty& Pin : ShowPinForProperties)
 	{
-		if (Pin.PropertyName == "JointTargetLocation")
+		if (Pin.PropertyName == GET_MEMBER_NAME_CHECKED(FAnimNode_TwoBoneIK, JointTargetLocation))
 		{
 			PreEditChange(NULL);
-			Pin.bShowPin = bShow;
+			if (Pin.bShowPin != bShow)
+			{
+				Pin.bShowPin = bShow;
+				bChanged = true;
+			}
 			break;
 		}
 	}
+
+	return bChanged;
 }
 
 #undef LOCTEXT_NAMESPACE
diff --git a/Engine/Source/Editor/Persona/Private/AnimationEditorViewportClient.cpp b/Engine/Source/Editor/Persona/Private/AnimationEditorViewportClient.cpp
--- a/Engine/Source/Editor/Persona/Private/AnimationEditorViewportClient.cpp
+++ b/Engine/Source/Editor/Persona/Private/AnimationEditorViewportClient.cpp
@@ -89,6 +89,7 @@ FAnimationViewportClient::FAnimationViewportClient( FPreviewScene& InPreviewScen
 	, SelectedWindActor(NULL)
 	, bFocusOnDraw(false)
 	, BodyTraceDistance(100000.0f)
+	, bShouldUpdateDefaultValues(false)
 {
 	// load config
 	ConfigOption = UPersonaOptions::StaticClass()->GetDefaultObject<UPersonaOptions>();
@@ -604,6 +605,153 @@ void FAnimationViewportClient::DrawUVsForMesh(FViewport* InViewport, FCanvas* In
 	DrawUVs(InViewport, InCanvas, InTextYPos, LODLevel, UVChannelToDraw, SelectedEdgeTexCoords, NULL, &PreviewSkelMeshComp->GetSkeletalMeshResource()->LODModels[LODLevel] );
 }
 
+FAnimNode_SkeletalControlBase* FAnimationViewportClient::FindSkeletalControlAnimNode(TWeakObjectPtr<UAnimGraphNode_SkeletalControlBase> AnimGraphNode) const
+{
+	FAnimNode_SkeletalControlBase* AnimNode = NULL;
+
+	if (!PreviewSkelMeshComp.IsValid() || !PreviewSkelMeshComp->GetAnimInstance())
+	{
+		return NULL;
+	}
+
+	if (AnimGraphNode.IsValid())
+	{
+		// find an anim node index from debug data
+		UAnimBlueprintGeneratedClass* AnimBlueprintClass = Cast<UAnimBlueprintGeneratedClass>(PreviewSkelMeshComp->GetAnimInstance()->GetClass());
+		FAnimBlueprintDebugData& DebugData = AnimBlueprintClass->GetAnimBlueprintDebugData();
+		int32* IndexPtr = DebugData.NodePropertyToIndexMap.Find(AnimGraphNode);
+
+		if (IndexPtr)
+		{
+			int32 AnimNodeIndex = *IndexPtr;
+			// reverse node index temporarily because of a bug in NodeGuidToIndexMap
+			AnimNodeIndex = AnimBlueprintClass->AnimNodeProperties.Num() - AnimNodeIndex - 1;
+
+			AnimNode = AnimBlueprintClass->AnimNodeProperties[AnimNodeIndex]->ContainerPtrToValuePtr<FAnimNode_SkeletalControlBase>(PreviewSkelMeshComp->GetAnimInstance());
+		}
+	}
+
+	return AnimNode;
+}
+
+void FAnimationViewportClient::FindSelectedAnimGraphNode()
+{
+	bool bSelected = false;
+	// finding a selected node
+	if (PersonaPtr.IsValid() && PreviewSkelMeshComp.IsValid())
+	{
+		const FGraphPanelSelectionSet SelectedNodes = PersonaPtr.Pin()->GetSelectedNodes();
+
+		// don't support multi-selection
+		if (SelectedNodes.Num() == 1)
+		{
+			FGraphPanelSelectionSet::TConstIterator NodeIt(SelectedNodes);
+			// first element
+			UAnimGraphNode_SkeletalControlBase* AnimGraphNode = Cast<UAnimGraphNode_SkeletalControlBase>(*NodeIt);
+
+			if (AnimGraphNode)
+			{
+				FAnimNode_SkeletalControlBase* AnimNode = FindSkeletalControlAnimNode(AnimGraphNode);
+
+				if (AnimNode)
+				{
+					bSelected = true;
+
+					// when selected first after AnimGraph is opened, assign previous data to current node
+					if (SelectedSkelControlAnimGraph != AnimGraphNode)
+					{
+						if (SelectedSkelControlAnimGraph.IsValid())
+						{
+							SelectedSkelControlAnimGraph->DeselectActor(PreviewSkelMeshComp.Get());
+						}
+
+						// make same values to ensure data consistency
+						AnimGraphNode->CopyNodeDataTo(AnimNode);
+						AnimGraphNode->UpdateAllDefaultValues(AnimNode);
+
+						WidgetMode = AnimGraphNode->GetWidgetMode(PreviewSkelMeshComp.Get());
+						if (WidgetMode == FWidget::WM_Scale)
+						{
+							SetWidgetCoordSystemSpace(COORD_Local);
+						}
+						else
+						{
+							SetWidgetCoordSystemSpace(COORD_World);
+						}
+					}
+					else
+					{
+						if (bShouldUpdateDefaultValues)
+						{
+							// copy updated values into internal node to ensure data consistency
+							AnimGraphNode->CopyNodeDataFrom(AnimNode);
+
+							AnimGraphNode->UpdateDefaultValues(AnimNode);
+							bShouldUpdateDefaultValues = false;
+						}
+					}
+
+					AnimGraphNode->MoveSelectActorLocation(PreviewSkelMeshComp.Get(), AnimNode);
+					SelectedSkelControlAnimGraph = AnimGraphNode;
+				}
+			}
+		}
+	}
+
+	if (!bSelected)
+	{
+		if (SelectedSkelControlAnimGraph.IsValid())
+		{
+			SelectedSkelControlAnimGraph->DeselectActor(PreviewSkelMeshComp.Get());
+		}
+		SelectedSkelControlAnimGraph.Reset();
+	}
+}
+
+void FAnimationViewportClient::ClearSelectedAnimGraphNode()
+{
+	if (SelectedSkelControlAnimGraph.IsValid())
+	{
+		SelectedSkelControlAnimGraph->DeselectActor(PreviewSkelMeshComp.Get());
+	}
+}
+
+void FAnimationViewportClient::PostUndo()
+{
+	// undo for skeletal controllers
+	// search all UAnimGraphNode_SkeletalControlBase nodes and apply data
+	UEdGraph* FocusedGraph = PersonaPtr.Pin()->GetFocusedGraph();
+
+	// @fixme : fix this to better way than looping all nodes
+	if (FocusedGraph)
+	{
+		// find UAnimGraphNode_SkeletalControlBase
+		for (UEdGraphNode* Node : FocusedGraph->Nodes)
+		{
+			UAnimGraphNode_SkeletalControlBase* AnimGraphNode = Cast<UAnimGraphNode_SkeletalControlBase>(Node);
+
+			if (AnimGraphNode)
+			{
+				FAnimNode_SkeletalControlBase* AnimNode = FindSkeletalControlAnimNode(AnimGraphNode);
+
+				if (AnimNode)
+				{
+					// copy undo data
+					AnimGraphNode->CopyNodeDataTo(AnimNode);
+					// update all literal values in anim node
+					AnimGraphNode->UpdateAllDefaultValues(AnimNode);
+				}
+			}
+		}
+	}
+}
+
+void FAnimationViewportClient::PostCompile()
+{
+	// reset the selected anim graph to copy 
+	SelectedSkelControlAnimGraph.Reset();
+}
+
 void FAnimationViewportClient::Tick(float DeltaSeconds) 
 {
 	FEditorViewportClient::Tick(DeltaSeconds);
@@ -630,6 +778,9 @@ void FAnimationViewportClient::Tick(float DeltaSeconds)
 		FBoxSphereBounds Bounds = EditorFloorComp->CalcBounds(EditorFloorComp->GetRelativeTransform());
 		PreviewComp->ConsumeRootMotion(Bounds.GetBox().Min, Bounds.GetBox().Max);
 	}	
+
+	FindSelectedAnimGraphNode();
+
 }
 
 void FAnimationViewportClient::SetCameraTargetLocation(const FSphere &BoundSphere, float DeltaSeconds)
@@ -996,8 +1147,14 @@ void FAnimationViewportClient::DisplayInfo(FCanvas* Canvas, FSceneView* View, bo
 
 int32 FAnimationViewportClient::FindSelectedBone() const
 {
-	int32 SelectedBone = -1;
-	if (PreviewSkelMeshComp.IsValid() && (PreviewSkelMeshComp->BonesOfInterest.Num() == 1) )
+	int32 SelectedBone = INDEX_NONE;
+	if (SelectedSkelControlAnimGraph.IsValid())
+	{		
+		FName BoneName = SelectedSkelControlAnimGraph->FindSelectedBone();
+		SelectedBone = PreviewSkelMeshComp->GetBoneIndex(BoneName);
+	}
+
+	if (SelectedBone == INDEX_NONE && PreviewSkelMeshComp.IsValid() && (PreviewSkelMeshComp->BonesOfInterest.Num() == 1) )
 	{
 		SelectedBone = PreviewSkelMeshComp->BonesOfInterest[0];
 	}
@@ -1057,6 +1214,11 @@ void FAnimationViewportClient::ProcessClick(class FSceneView& View, class HHitPr
 				SharedPersona->DeselectAll();
 				SelectedWindActor = WindActor;
 			}
+			else if (SelectedSkelControlAnimGraph.IsValid() && SelectedSkelControlAnimGraph->IsActorClicked(ActorHitProxy))
+			{
+				// do some actions when hit
+				SelectedSkelControlAnimGraph->ProcessActorClick(ActorHitProxy);
+			}
 		}
 	}
 	else
@@ -1143,7 +1305,16 @@ bool FAnimationViewportClient::InputWidgetDelta( FViewport* Viewport, EAxisList:
 
 				FRotator NewRotation = ( CurrentSkelControlTM * FTransform( DeltaQuat )).Rotator();
 
-				if ( SelectedSocket )
+				if (SelectedSkelControlAnimGraph.IsValid())
+				{
+					FAnimNode_SkeletalControlBase* AnimNode = FindSkeletalControlAnimNode(SelectedSkelControlAnimGraph);
+					if (AnimNode)
+					{
+						SelectedSkelControlAnimGraph->DoRotation(PreviewSkelMeshComp.Get(), Rot, AnimNode);
+						bShouldUpdateDefaultValues = true;
+					}
+				}
+				else if ( SelectedSocket )
 				{
 					SelectedSocket->RelativeRotation = NewRotation;
 				}
@@ -1156,8 +1327,16 @@ bool FAnimationViewportClient::InputWidgetDelta( FViewport* Viewport, EAxisList:
 			if (bDoTranslation)
 			{
 				FVector4 BoneSpaceOffset = BaseTM.TransformVector(Drag);
-
-				if (SelectedSocket)
+				if (SelectedSkelControlAnimGraph.IsValid())
+				{
+					FAnimNode_SkeletalControlBase* AnimNode = FindSkeletalControlAnimNode(SelectedSkelControlAnimGraph);
+					if (AnimNode)
+					{
+						SelectedSkelControlAnimGraph->DoTranslation(PreviewSkelMeshComp.Get(), Drag, AnimNode);
+						bShouldUpdateDefaultValues = true;
+					}
+				}
+				else if (SelectedSocket)
 				{
 					SelectedSocket->RelativeLocation += BoneSpaceOffset;
 				}
@@ -1179,7 +1358,16 @@ bool FAnimationViewportClient::InputWidgetDelta( FViewport* Viewport, EAxisList:
 					BoneSpaceScaleOffset = Scale;
 				}
 
-				if(SelectedSocket)
+				if (SelectedSkelControlAnimGraph.IsValid())
+				{
+					FAnimNode_SkeletalControlBase* AnimNode = FindSkeletalControlAnimNode(SelectedSkelControlAnimGraph);
+					if (AnimNode)
+					{
+						SelectedSkelControlAnimGraph->DoScale(PreviewSkelMeshComp.Get(), Scale, AnimNode);
+						bShouldUpdateDefaultValues = true;
+					}
+				}
+				else if(SelectedSocket)
 				{
 					SelectedSocket->RelativeScale += BoneSpaceScaleOffset;
 				}
@@ -1188,8 +1376,8 @@ bool FAnimationViewportClient::InputWidgetDelta( FViewport* Viewport, EAxisList:
 					SkelControl->Scale += BoneSpaceScaleOffset;
 				}
 			}
-			
-		}
+
+			}
 		else if( WindActor.IsValid() )
 		{
 			if (WidgetMode == FWidget::WM_Rotate)
@@ -1318,12 +1506,25 @@ FWidget::EWidgetMode FAnimationViewportClient::GetWidgetMode() const
 	}
 	}
 
+	if (SelectedSkelControlAnimGraph.IsValid())
+	{
+		Mode = SelectedSkelControlAnimGraph->GetWidgetMode(PreviewSkelMeshComp.Get());
+	}
+
 	return Mode;
 }
 
 FVector FAnimationViewportClient::GetWidgetLocation() const
 {
-	if( PreviewSkelMeshComp->BonesOfInterest.Num() > 0 )
+	if (SelectedSkelControlAnimGraph.IsValid())
+	{
+		FAnimNode_SkeletalControlBase* AnimNode = FindSkeletalControlAnimNode(SelectedSkelControlAnimGraph.Get());
+		if (AnimNode)
+		{
+			return SelectedSkelControlAnimGraph->GetWidgetLocation(PreviewSkelMeshComp.Get(), AnimNode);
+		}
+	}
+	else if( PreviewSkelMeshComp->BonesOfInterest.Num() > 0 )
 	{
 		int32 BoneIndex = PreviewSkelMeshComp->BonesOfInterest.Last();
 		const FName BoneName = PreviewSkelMeshComp->SkeletalMesh->RefSkeleton.GetBoneName(BoneIndex);
@@ -1355,6 +1556,15 @@ FMatrix FAnimationViewportClient::GetWidgetCoordSystem() const
 
 	if( bIsLocal )
 	{
+		if (SelectedSkelControlAnimGraph.IsValid())
+		{
+			FName BoneName = SelectedSkelControlAnimGraph->FindSelectedBone();
+			int32 BoneIndex = PreviewSkelMeshComp->GetBoneIndex(BoneName);
+			FTransform BoneMatrix = PreviewSkelMeshComp->GetBoneTransform(BoneIndex);
+
+			return BoneMatrix.ToMatrixNoScale().RemoveTranslation();
+		}
+		else
 		if ( PreviewSkelMeshComp->BonesOfInterest.Num() > 0 )
 		{
 			int32 BoneIndex = PreviewSkelMeshComp->BonesOfInterest.Last();
@@ -1422,6 +1632,19 @@ bool FAnimationViewportClient::InputKey( FViewport* Viewport, int32 ControllerId
 
 		if (SelectedBone >= 0 || SelectedSocket || SelectedWind.IsValid())
 		{
+			if (SelectedSkelControlAnimGraph.IsValid())
+			{
+				WidgetMode = SelectedSkelControlAnimGraph->ChangeWidgetMode(PreviewSkelMeshComp.Get(), WidgetMode);
+				if (WidgetMode == FWidget::WM_Scale)
+				{
+					SetWidgetCoordSystemSpace(COORD_Local);
+				}
+				else
+				{
+					SetWidgetCoordSystemSpace(COORD_World);
+				}
+			}
+			else
 			if (WidgetMode == FWidget::WM_Rotate)
 			{
 				WidgetMode = FWidget::WM_Translate;
@@ -1575,7 +1798,7 @@ void FAnimationViewportClient::DrawMeshBones(USkeletalMeshComponent * MeshCompon
 	}
 }
 
-void FAnimationViewportClient::DrawBones(const USkeletalMeshComponent* MeshComponent, const TArray<FBoneIndexType>& RequiredBones, const TArray<FTransform>& WorldTransforms, FPrimitiveDrawInterface* PDI, const TArray<FLinearColor> BoneColours, float LineThickness/*=0.f*/) const
+void FAnimationViewportClient::DrawBones(const USkeletalMeshComponent* MeshComponent, const TArray<FBoneIndexType> & RequiredBones, const TArray<FTransform> & WorldTransforms, FPrimitiveDrawInterface* PDI, const TArray<FLinearColor> BoneColours, float LineThickness/*=0.f*/) const
 {
 	check ( MeshComponent && MeshComponent->SkeletalMesh );
 
diff --git a/Engine/Source/Editor/Persona/Private/AnimationEditorViewportClient.h b/Engine/Source/Editor/Persona/Private/AnimationEditorViewportClient.h
--- a/Engine/Source/Editor/Persona/Private/AnimationEditorViewportClient.h
+++ b/Engine/Source/Editor/Persona/Private/AnimationEditorViewportClient.h
@@ -182,6 +182,10 @@ public:
 		SelectedWindActor = NULL;
 	}
 
+	void ClearSelectedAnimGraphNode();
+	void PostUndo();
+	void PostCompile();
+
 	/* Returns the floor height offset */	
 	float GetFloorOffset() const;
 
@@ -312,4 +316,13 @@ private:
 	void DrawSockets( TArray<class USkeletalMeshSocket*>& Sockets, FPrimitiveDrawInterface* PDI, bool bUseSkeletonSocketColor ) const;
 
 	TWeakObjectPtr<AWindDirectionalSource> FindSelectedWindActor() const;
+
+	struct FAnimNode_SkeletalControlBase* FindSkeletalControlAnimNode(TWeakObjectPtr<class UAnimGraphNode_SkeletalControlBase> AnimGraphNode) const;
+
+	void FindSelectedAnimGraphNode();
+
+	// selected skeletal control anim graph node 
+	TWeakObjectPtr<class UAnimGraphNode_SkeletalControlBase> SelectedSkelControlAnimGraph;
+	// to check whether we should update literal values in selected AnimGraphNode
+	bool bShouldUpdateDefaultValues;
 };
diff --git a/Engine/Source/Editor/Persona/Private/Persona.cpp b/Engine/Source/Editor/Persona/Private/Persona.cpp
--- a/Engine/Source/Editor/Persona/Private/Persona.cpp
+++ b/Engine/Source/Editor/Persona/Private/Persona.cpp
@@ -1641,6 +1641,12 @@ void FPersona::Compile()
 			GetBlueprintObj()->SetObjectBeingDebugged(PreviewComponent->AnimScriptInstance);
 		}
 	}
+
+	// calls PostCompile to copy proper values between anim nodes
+	if (Viewport.IsValid())
+	{
+		Viewport.Pin()->GetAnimationViewportClient()->PostCompile();
+	}
 }
 
 FString FPersona::GetDefaultEditorTitle()
@@ -2044,7 +2050,18 @@ void FPersona::ClearSelectedSocket()
 
 void FPersona::ClearSelectedWindActor()
 {
-	Viewport.Pin()->GetAnimationViewportClient()->ClearSelectedWindActor();
+	if (Viewport.IsValid())
+	{
+		Viewport.Pin()->GetAnimationViewportClient()->ClearSelectedWindActor();
+	}
+}
+
+void FPersona::ClearSelectedAnimGraphNode()
+{
+	if(Viewport.IsValid())
+	{
+		Viewport.Pin()->GetAnimationViewportClient()->ClearSelectedAnimGraphNode();
+	}
 }
 
 void FPersona::RenameSocket( USkeletalMeshSocket* Socket, const FName& NewSocketName )
@@ -2387,6 +2404,7 @@ void FPersona::DeselectAll()
 	ClearSelectedBones();
 	ClearSelectedSocket();
 	ClearSelectedWindActor();
+	ClearSelectedAnimGraphNode();
 
 	OnAllDeselected.Broadcast();
 }
diff --git a/Engine/Source/Editor/Persona/Private/Persona.h b/Engine/Source/Editor/Persona/Private/Persona.h
--- a/Engine/Source/Editor/Persona/Private/Persona.h
+++ b/Engine/Source/Editor/Persona/Private/Persona.h
@@ -114,6 +114,9 @@ public:
 	/** Clears the selected wind actor */
 	void ClearSelectedWindActor();
 
+	/** Clears the selected anim graph node */
+	void ClearSelectedAnimGraphNode();
+
 	/** Clears the selection (both sockets and bones). Also broadcasts this */
 	void DeselectAll();
 
diff --git a/Engine/Source/Editor/Persona/Private/SAnimationEditorViewport.cpp b/Engine/Source/Editor/Persona/Private/SAnimationEditorViewport.cpp
--- a/Engine/Source/Editor/Persona/Private/SAnimationEditorViewport.cpp
+++ b/Engine/Source/Editor/Persona/Private/SAnimationEditorViewport.cpp
@@ -75,6 +75,7 @@ TSharedPtr<SWidget> SAnimationEditorViewport::MakeViewportToolbar()
 void SAnimationEditorViewport::OnUndoRedo()
 {
 	LevelViewportClient->Invalidate();
+	LevelViewportClient->PostUndo();
 }
 
 //////////////////////////////////////////////////////////////////////////
diff --git a/Engine/Source/Runtime/Engine/Classes/Animation/BoneControllers/AnimNode_SkeletalControlBase.h b/Engine/Source/Runtime/Engine/Classes/Animation/BoneControllers/AnimNode_SkeletalControlBase.h
--- a/Engine/Source/Runtime/Engine/Classes/Animation/BoneControllers/AnimNode_SkeletalControlBase.h
+++ b/Engine/Source/Runtime/Engine/Classes/Animation/BoneControllers/AnimNode_SkeletalControlBase.h
@@ -36,6 +36,11 @@ public:
 	}
 
 public:
+#if WITH_EDITORONLY_DATA
+	// forwarded pose data from the wired node which current node's skeletal control is not applied yet
+	FA2CSPose ForwardedPose;
+#endif //#if WITH_EDITORONLY_DATA
+
 	// FAnimNode_Base interface
 	virtual void Initialize(const FAnimationInitializeContext& Context) override;
 	virtual void CacheBones(const FAnimationCacheBonesContext& Context)  override;
diff --git a/Engine/Source/Runtime/Engine/Private/Animation/AnimNodeSkeletalControlBase.cpp b/Engine/Source/Runtime/Engine/Private/Animation/AnimNodeSkeletalControlBase.cpp
--- a/Engine/Source/Runtime/Engine/Private/Animation/AnimNodeSkeletalControlBase.cpp
+++ b/Engine/Source/Runtime/Engine/Private/Animation/AnimNodeSkeletalControlBase.cpp
@@ -49,6 +49,11 @@ void FAnimNode_SkeletalControlBase::EvaluateComponentSpace(FComponentSpacePoseCo
 	{
 		USkeletalMeshComponent* Component = Output.AnimInstance->GetSkelMeshComponent();
 
+#if WITH_EDITORONLY_DATA
+		// save current pose before applying skeletal control to compute the exact gizmo location in AnimGraphNode
+		ForwardedPose = Output.Pose;
+#endif // #if WITH_EDITORONLY_DATA
+
 		TArray<FBoneTransform> BoneTransforms;
 		EvaluateBoneTransforms(Component, Output.AnimInstance->RequiredBones, Output.Pose, BoneTransforms);
 
