included a header file to fix compile errors on Mac

diff --git a/Engine/Source/Editor/AnimGraph/Classes/AnimGraphNode_SkeletalControlBase.h b/Engine/Source/Editor/AnimGraph/Classes/AnimGraphNode_SkeletalControlBase.h
--- a/Engine/Source/Editor/AnimGraph/Classes/AnimGraphNode_SkeletalControlBase.h
+++ b/Engine/Source/Editor/AnimGraph/Classes/AnimGraphNode_SkeletalControlBase.h
@@ -3,6 +3,7 @@
 #pragma once
 
 #include "AnimGraphNode_Base.h"
+#include "UnrealWidget.h"
 #include "AnimGraphNode_SkeletalControlBase.generated.h"
 
 UCLASS(Abstract, MinimalAPI)
