Fixed compile errors on Mac

diff --git a/Engine/Source/Editor/AnimGraph/Classes/AnimGraphNode_Fabrik.h b/Engine/Source/Editor/AnimGraph/Classes/AnimGraphNode_Fabrik.h
--- a/Engine/Source/Editor/AnimGraph/Classes/AnimGraphNode_Fabrik.h
+++ b/Engine/Source/Editor/AnimGraph/Classes/AnimGraphNode_Fabrik.h
@@ -21,7 +21,7 @@ public:
 
 	// UAnimGraphNode_SkeletalControlBase interface
 	ANIMGRAPH_API virtual FVector GetWidgetLocation(const USkeletalMeshComponent* SkelComp, struct FAnimNode_SkeletalControlBase* AnimNode) override;
-	ANIMGRAPH_API virtual FWidget::EWidgetMode GetWidgetMode(const USkeletalMeshComponent* SkelComp) override;
+	ANIMGRAPH_API virtual int32 GetWidgetMode(const USkeletalMeshComponent* SkelComp) override;
 	ANIMGRAPH_API virtual FName FindSelectedBone() override;
 	ANIMGRAPH_API virtual void DoTranslation(const USkeletalMeshComponent* SkelComp, FVector& Drag, FAnimNode_Base* InOutAnimNode) override;
 
diff --git a/Engine/Source/Editor/AnimGraph/Classes/AnimGraphNode_ModifyBone.h b/Engine/Source/Editor/AnimGraph/Classes/AnimGraphNode_ModifyBone.h
--- a/Engine/Source/Editor/AnimGraph/Classes/AnimGraphNode_ModifyBone.h
+++ b/Engine/Source/Editor/AnimGraph/Classes/AnimGraphNode_ModifyBone.h
@@ -27,8 +27,8 @@ protected:
 	virtual FText GetControllerDescription() const override;
 
 	ANIMGRAPH_API virtual FVector GetWidgetLocation(const USkeletalMeshComponent* SkelComp, struct FAnimNode_SkeletalControlBase* AnimNode) override;
-	ANIMGRAPH_API virtual FWidget::EWidgetMode GetWidgetMode(const USkeletalMeshComponent* SkelComp) override;
-	ANIMGRAPH_API virtual FWidget::EWidgetMode ChangeWidgetMode(const USkeletalMeshComponent* SkelComp, FWidget::EWidgetMode InCurWidgetMode) override;
+	ANIMGRAPH_API virtual int32 GetWidgetMode(const USkeletalMeshComponent* SkelComp) override;
+	ANIMGRAPH_API virtual int32 ChangeWidgetMode(const USkeletalMeshComponent* SkelComp, int32 InCurWidgetMode) override;
 
 	ANIMGRAPH_API virtual FName FindSelectedBone() override;
 
@@ -44,14 +44,15 @@ protected:
 	// End of UAnimGraphNode_SkeletalControlBase interface
 
 	// methods to find a valid widget mode for gizmo because doesn't need to show gizmo when the mode is "Ignore"
-	FWidget::EWidgetMode FindValidWidgetMode(FWidget::EWidgetMode InMode);
-	EBoneModificationMode GetBoneModificationMode(FWidget::EWidgetMode InMode);
-	FWidget::EWidgetMode GetNextWidgetMode(FWidget::EWidgetMode InMode);
+	int32 FindValidWidgetMode(int32 InWidgetMode);
+	EBoneModificationMode GetBoneModificationMode(int32 InWidgetMode);
+	int32 GetNextWidgetMode(int32 InWidgetMode);
 
 private:
 	/** Constructing FText strings can be costly, so we cache the node's title */
 	FNodeTitleTextTable CachedNodeTitles;
 
-	FWidget::EWidgetMode CurWidgetMode;
+	// storing current widget mode 
+	int32 CurWidgetMode;
 };
 
diff --git a/Engine/Source/Editor/AnimGraph/Classes/AnimGraphNode_SkeletalControlBase.h b/Engine/Source/Editor/AnimGraph/Classes/AnimGraphNode_SkeletalControlBase.h
--- a/Engine/Source/Editor/AnimGraph/Classes/AnimGraphNode_SkeletalControlBase.h
+++ b/Engine/Source/Editor/AnimGraph/Classes/AnimGraphNode_SkeletalControlBase.h
@@ -42,20 +42,11 @@ public:
 	ANIMGRAPH_API virtual void UpdateAllDefaultValues(const FAnimNode_Base* AnimNode){}
 
 	// return current widget mode this anim graph node supports
-	ANIMGRAPH_API virtual FWidget::EWidgetMode GetWidgetMode(const USkeletalMeshComponent* SkelComp)
-	{ 
-		return FWidget::EWidgetMode::WM_None; 
-	}
+	ANIMGRAPH_API virtual int32 GetWidgetMode(const USkeletalMeshComponent* SkelComp);
 	// called when the user changed widget mode
-	ANIMGRAPH_API virtual FWidget::EWidgetMode ChangeWidgetMode(const USkeletalMeshComponent* SkelComp, FWidget::EWidgetMode CurWidgetMode)
-	{ 
-		return GetWidgetMode(SkelComp); 
-	}
+	ANIMGRAPH_API virtual int32 ChangeWidgetMode(const USkeletalMeshComponent* SkelComp, int32 CurWidgetMode);
 	// 
-	ANIMGRAPH_API virtual FName FindSelectedBone()
-	{ 
-		return NAME_None; 
-	}
+	ANIMGRAPH_API virtual FName FindSelectedBone();
 
 	// if anim graph node needs other actors to select other bones, move actor's positions when this is called
 	ANIMGRAPH_API virtual void MoveSelectActorLocation(const USkeletalMeshComponent* SkelComp, FAnimNode_SkeletalControlBase* InAnimNode){}
diff --git a/Engine/Source/Editor/AnimGraph/Classes/AnimGraphNode_TwoBoneIK.h b/Engine/Source/Editor/AnimGraph/Classes/AnimGraphNode_TwoBoneIK.h
--- a/Engine/Source/Editor/AnimGraph/Classes/AnimGraphNode_TwoBoneIK.h
+++ b/Engine/Source/Editor/AnimGraph/Classes/AnimGraphNode_TwoBoneIK.h
@@ -41,7 +41,7 @@ public:
 	ANIMGRAPH_API virtual void	UpdateDefaultValues(const FAnimNode_Base* InAnimNode) override;
 	ANIMGRAPH_API virtual void	UpdateAllDefaultValues(const FAnimNode_Base* InAnimNode) override;
 
-	ANIMGRAPH_API virtual FWidget::EWidgetMode GetWidgetMode(const USkeletalMeshComponent* SkelComp) override;
+	ANIMGRAPH_API virtual int32 GetWidgetMode(const USkeletalMeshComponent* SkelComp) override;
 
 	ANIMGRAPH_API virtual void MoveSelectActorLocation(const USkeletalMeshComponent* SkelComp, FAnimNode_SkeletalControlBase* AnimNode) override;
 	ANIMGRAPH_API virtual FName FindSelectedBone() override;
diff --git a/Engine/Source/Editor/AnimGraph/Private/AnimGraphNode_Fabrik.cpp b/Engine/Source/Editor/AnimGraph/Private/AnimGraphNode_Fabrik.cpp
--- a/Engine/Source/Editor/AnimGraph/Private/AnimGraphNode_Fabrik.cpp
+++ b/Engine/Source/Editor/AnimGraph/Private/AnimGraphNode_Fabrik.cpp
@@ -32,7 +32,7 @@ FVector UAnimGraphNode_Fabrik::GetWidgetLocation(const USkeletalMeshComponent* S
 	return WidgetLoc;
 }
 
-FWidget::EWidgetMode UAnimGraphNode_Fabrik::GetWidgetMode(const USkeletalMeshComponent* SkelComp)
+int32 UAnimGraphNode_Fabrik::GetWidgetMode(const USkeletalMeshComponent* SkelComp)
 {
 	uint32 BoneIndex = SkelComp->GetBoneIndex(Node.EffectorTransformBone.BoneName);
 
diff --git a/Engine/Source/Editor/AnimGraph/Private/AnimGraphNode_ModifyBone.cpp b/Engine/Source/Editor/AnimGraph/Private/AnimGraphNode_ModifyBone.cpp
--- a/Engine/Source/Editor/AnimGraph/Private/AnimGraphNode_ModifyBone.cpp
+++ b/Engine/Source/Editor/AnimGraph/Private/AnimGraphNode_ModifyBone.cpp
@@ -11,7 +11,7 @@
 UAnimGraphNode_ModifyBone::UAnimGraphNode_ModifyBone(const FObjectInitializer& ObjectInitializer)
 	: Super(ObjectInitializer)
 {
-	CurWidgetMode = FWidget::WM_Rotate;
+	CurWidgetMode = (int32)FWidget::WM_Rotate;
 }
 
 FText UAnimGraphNode_ModifyBone::GetControllerDescription() const
@@ -95,8 +95,9 @@ FVector UAnimGraphNode_ModifyBone::GetWidgetLocation(const USkeletalMeshComponen
 	return WidgetLoc;
 }
 
-EBoneModificationMode UAnimGraphNode_ModifyBone::GetBoneModificationMode(FWidget::EWidgetMode InMode)
+EBoneModificationMode UAnimGraphNode_ModifyBone::GetBoneModificationMode(int32 InWidgetMode)
 {
+	FWidget::EWidgetMode InMode = (FWidget::EWidgetMode)InWidgetMode;
 	switch (InMode)
 	{
 	case FWidget::WM_Translate:
@@ -113,8 +114,9 @@ EBoneModificationMode UAnimGraphNode_ModifyBone::GetBoneModificationMode(FWidget
 	return EBoneModificationMode::BMM_Ignore;
 }
 
-FWidget::EWidgetMode UAnimGraphNode_ModifyBone::GetNextWidgetMode(FWidget::EWidgetMode InMode)
+int32 UAnimGraphNode_ModifyBone::GetNextWidgetMode(int32 InWidgetMode)
 {
+	FWidget::EWidgetMode InMode = (FWidget::EWidgetMode)InWidgetMode;
 	switch (InMode)
 	{
 	case FWidget::WM_Translate:
@@ -125,11 +127,12 @@ FWidget::EWidgetMode UAnimGraphNode_ModifyBone::GetNextWidgetMode(FWidget::EWidg
 		return FWidget::WM_Translate;
 	}
 
-	return FWidget::WM_None;
+	return (int32)FWidget::WM_None;
 }
 
-FWidget::EWidgetMode UAnimGraphNode_ModifyBone::FindValidWidgetMode(FWidget::EWidgetMode InMode)
+int32 UAnimGraphNode_ModifyBone::FindValidWidgetMode(int32 InWidgetMode)
 {	
+	FWidget::EWidgetMode InMode = (FWidget::EWidgetMode)InWidgetMode;
 	FWidget::EWidgetMode ValidMode = InMode;
 	if (InMode == FWidget::WM_None)
 	{	// starts from Rotate mode
@@ -144,16 +147,16 @@ FWidget::EWidgetMode UAnimGraphNode_ModifyBone::FindValidWidgetMode(FWidget::EWi
 			return ValidMode;
 		}
 
-		ValidMode = GetNextWidgetMode(ValidMode);
+		ValidMode = (FWidget::EWidgetMode)GetNextWidgetMode(ValidMode);
 	}
 	
 	// if couldn't find a valid mode, returns None
 	ValidMode = FWidget::WM_None;
 
-	return ValidMode;
+	return (int32)ValidMode;
 }
 
-FWidget::EWidgetMode UAnimGraphNode_ModifyBone::GetWidgetMode(const USkeletalMeshComponent* SkelComp)
+int32 UAnimGraphNode_ModifyBone::GetWidgetMode(const USkeletalMeshComponent* SkelComp)
 {
 	int32 BoneIndex = SkelComp->GetBoneIndex(Node.BoneToModify.BoneName);
 	if (BoneIndex != INDEX_NONE)
@@ -162,12 +165,12 @@ FWidget::EWidgetMode UAnimGraphNode_ModifyBone::GetWidgetMode(const USkeletalMes
 		return CurWidgetMode;
 	}
 
-	return FWidget::WM_None;
+	return (int32)FWidget::WM_None;
 }
 
-FWidget::EWidgetMode UAnimGraphNode_ModifyBone::ChangeWidgetMode(const USkeletalMeshComponent* SkelComp, FWidget::EWidgetMode InCurWidgetMode)
+int32 UAnimGraphNode_ModifyBone::ChangeWidgetMode(const USkeletalMeshComponent* SkelComp, int32 InCurWidgetMode)
 {
-	FWidget::EWidgetMode NextWidgetMode = GetNextWidgetMode(InCurWidgetMode);
+	int32 NextWidgetMode = GetNextWidgetMode(InCurWidgetMode);
 	CurWidgetMode = FindValidWidgetMode(NextWidgetMode);
 
 	return CurWidgetMode;
diff --git a/Engine/Source/Editor/AnimGraph/Private/AnimGraphNode_SkeletalControlBase.cpp b/Engine/Source/Editor/AnimGraph/Private/AnimGraphNode_SkeletalControlBase.cpp
--- a/Engine/Source/Editor/AnimGraph/Private/AnimGraphNode_SkeletalControlBase.cpp
+++ b/Engine/Source/Editor/AnimGraph/Private/AnimGraphNode_SkeletalControlBase.cpp
@@ -15,6 +15,22 @@ UAnimGraphNode_SkeletalControlBase::UAnimGraphNode_SkeletalControlBase(const FOb
 {
 }
 
+// returns int32 instead of EWidgetMode because of compiling issue on Mac
+int32 UAnimGraphNode_SkeletalControlBase::GetWidgetMode(const USkeletalMeshComponent* SkelComp)
+{
+	return  (int32)FWidget::EWidgetMode::WM_None;
+}
+
+int32 UAnimGraphNode_SkeletalControlBase::ChangeWidgetMode(const USkeletalMeshComponent* SkelComp, int32 CurWidgetMode)
+{
+	return GetWidgetMode(SkelComp);
+}
+// 
+FName UAnimGraphNode_SkeletalControlBase::FindSelectedBone()
+{
+	return NAME_None;
+}
+
 FLinearColor UAnimGraphNode_SkeletalControlBase::GetNodeTitleColor() const
 {
 	return FLinearColor(0.75f, 0.75f, 0.10f);
diff --git a/Engine/Source/Editor/AnimGraph/Private/AnimGraphNode_TwoBoneIK.cpp b/Engine/Source/Editor/AnimGraph/Private/AnimGraphNode_TwoBoneIK.cpp
--- a/Engine/Source/Editor/AnimGraph/Private/AnimGraphNode_TwoBoneIK.cpp
+++ b/Engine/Source/Editor/AnimGraph/Private/AnimGraphNode_TwoBoneIK.cpp
@@ -209,7 +209,7 @@ void UAnimGraphNode_TwoBoneIK::UpdateAllDefaultValues(const FAnimNode_Base* Anim
 }
 
 
-FWidget::EWidgetMode UAnimGraphNode_TwoBoneIK::GetWidgetMode(const USkeletalMeshComponent* SkelComp)
+int32 UAnimGraphNode_TwoBoneIK::GetWidgetMode(const USkeletalMeshComponent* SkelComp)
 {
 	int32 BoneIndex = SkelComp->GetBoneIndex(Node.IKBone.BoneName);
 	// Two bone IK node uses only Translate
diff --git a/Engine/Source/Editor/Persona/Private/AnimationEditorViewportClient.cpp b/Engine/Source/Editor/Persona/Private/AnimationEditorViewportClient.cpp
--- a/Engine/Source/Editor/Persona/Private/AnimationEditorViewportClient.cpp
+++ b/Engine/Source/Editor/Persona/Private/AnimationEditorViewportClient.cpp
@@ -669,7 +669,7 @@ void FAnimationViewportClient::FindSelectedAnimGraphNode()
 						AnimGraphNode->CopyNodeDataTo(AnimNode);
 						AnimGraphNode->UpdateAllDefaultValues(AnimNode);
 
-						WidgetMode = AnimGraphNode->GetWidgetMode(PreviewSkelMeshComp.Get());
+						WidgetMode = (FWidget::EWidgetMode)AnimGraphNode->GetWidgetMode(PreviewSkelMeshComp.Get());
 						if (WidgetMode == FWidget::WM_Scale)
 						{
 							SetWidgetCoordSystemSpace(COORD_Local);
@@ -1508,7 +1508,7 @@ FWidget::EWidgetMode FAnimationViewportClient::GetWidgetMode() const
 
 	if (SelectedSkelControlAnimGraph.IsValid())
 	{
-		Mode = SelectedSkelControlAnimGraph->GetWidgetMode(PreviewSkelMeshComp.Get());
+		Mode = (FWidget::EWidgetMode)SelectedSkelControlAnimGraph->GetWidgetMode(PreviewSkelMeshComp.Get());
 	}
 
 	return Mode;
@@ -1634,7 +1634,7 @@ bool FAnimationViewportClient::InputKey( FViewport* Viewport, int32 ControllerId
 		{
 			if (SelectedSkelControlAnimGraph.IsValid())
 			{
-				WidgetMode = SelectedSkelControlAnimGraph->ChangeWidgetMode(PreviewSkelMeshComp.Get(), WidgetMode);
+				WidgetMode = (FWidget::EWidgetMode)SelectedSkelControlAnimGraph->ChangeWidgetMode(PreviewSkelMeshComp.Get(), WidgetMode);
 				if (WidgetMode == FWidget::WM_Scale)
 				{
 					SetWidgetCoordSystemSpace(COORD_Local);
