Removed an unnecessary header file

diff --git a/Engine/Source/Editor/AnimGraph/Classes/AnimGraphNode_SkeletalControlBase.h b/Engine/Source/Editor/AnimGraph/Classes/AnimGraphNode_SkeletalControlBase.h
--- a/Engine/Source/Editor/AnimGraph/Classes/AnimGraphNode_SkeletalControlBase.h
+++ b/Engine/Source/Editor/AnimGraph/Classes/AnimGraphNode_SkeletalControlBase.h
@@ -3,7 +3,6 @@
 #pragma once
 
 #include "AnimGraphNode_Base.h"
-#include "UnrealWidget.h"
 #include "AnimGraphNode_SkeletalControlBase.generated.h"
 
 UCLASS(Abstract, MinimalAPI)
