#Anim : Fixed save cached pose node warning message

- We now don't check verification on skeleton class
- Moved Validation to EarlyValidation

diff --git a/Engine/Source/Editor/AnimGraph/Classes/AnimGraphNode_UseCachedPose.h b/Engine/Source/Editor/AnimGraph/Classes/AnimGraphNode_UseCachedPose.h
--- a/Engine/Source/Editor/AnimGraph/Classes/AnimGraphNode_UseCachedPose.h
+++ b/Engine/Source/Editor/AnimGraph/Classes/AnimGraphNode_UseCachedPose.h
@@ -27,12 +27,12 @@ public:
 	// UK2Node interface.
 	virtual void GetMenuActions(FBlueprintActionDatabaseRegistrar& ActionRegistrar) const override;
 	virtual bool IsActionFilteredOut(class FBlueprintActionFilter const& Filter) override;
-	virtual void ValidateNodeDuringCompilation(class FCompilerResultsLog& MessageLog) const override;
 	// End of UK2Node interface
 
 	// UAnimGraphNode_Base interface
 	virtual FString GetNodeCategory() const override;
 	virtual void GetMenuEntries(FGraphContextMenuBuilder& ContextMenuBuilder) const override;
+	virtual void EarlyValidation(class FCompilerResultsLog& MessageLog) const override;
 	// End of UAnimGraphNode_Base interface
 
 private:
diff --git a/Engine/Source/Editor/AnimGraph/Private/AnimGraphNode_UseCachedPose.cpp b/Engine/Source/Editor/AnimGraph/Private/AnimGraphNode_UseCachedPose.cpp
--- a/Engine/Source/Editor/AnimGraph/Private/AnimGraphNode_UseCachedPose.cpp
+++ b/Engine/Source/Editor/AnimGraph/Private/AnimGraphNode_UseCachedPose.cpp
@@ -22,8 +22,10 @@ UAnimGraphNode_UseCachedPose::UAnimGraphNode_UseCachedPose(const FObjectInitiali
 {
 }
 
-void UAnimGraphNode_UseCachedPose::ValidateNodeDuringCompilation(class FCompilerResultsLog& MessageLog) const
+void UAnimGraphNode_UseCachedPose::EarlyValidation(class FCompilerResultsLog& MessageLog) const
 {
+	Super::EarlyValidation(MessageLog);
+	
 	bool bRefreshSavCachedPoseNode = true;
 
 	// Check to see the current cached node is still valid (and not deleted, by checking pin connections)
diff --git a/Engine/Source/Editor/KismetCompiler/Private/AnimBlueprintCompiler.cpp b/Engine/Source/Editor/KismetCompiler/Private/AnimBlueprintCompiler.cpp
--- a/Engine/Source/Editor/KismetCompiler/Private/AnimBlueprintCompiler.cpp
+++ b/Engine/Source/Editor/KismetCompiler/Private/AnimBlueprintCompiler.cpp
@@ -387,6 +387,12 @@ void FAnimBlueprintCompiler::ProcessUseCachedPose(UAnimGraphNode_UseCachedPose*
 {
 	bool bSuccessful = false;
 
+	// if compiling only skeleton, we don't have to worry about linking save node
+	if (CompileOptions.CompileType == EKismetCompileType::SkeletonOnly)
+	{
+		return;
+	}
+
 	// Link to the saved cached pose
 	if(UseCachedPose->SaveCachedPoseNode.IsValid())
 	{
