PR #725: exposing class can be properly subclassed in non-engine modules. (Contributed by [REDACTED_PERSON])

https://github.com/EpicGames/UnrealEngine/pull/725

diff --git a/Engine/Source/Editor/AnimGraph/Classes/AnimGraphNode_BlendListBase.h b/Engine/Source/Editor/AnimGraph/Classes/AnimGraphNode_BlendListBase.h
--- a/Engine/Source/Editor/AnimGraph/Classes/AnimGraphNode_BlendListBase.h
+++ b/Engine/Source/Editor/AnimGraph/Classes/AnimGraphNode_BlendListBase.h
@@ -10,9 +10,9 @@ class UAnimGraphNode_BlendListBase : public UAnimGraphNode_Base
 	GENERATED_UCLASS_BODY()
 
 	// UEdGraphNode interface
-	virtual FLinearColor GetNodeTitleColor() const override;
-	virtual void PostEditChangeProperty(struct FPropertyChangedEvent& PropertyChangedEvent) override;
-	virtual void ReallocatePinsDuringReconstruction(TArray<UEdGraphPin*>& OldPins) override;
+	ANIMGRAPH_API virtual FLinearColor GetNodeTitleColor() const override;
+	ANIMGRAPH_API virtual void PostEditChangeProperty(struct FPropertyChangedEvent& PropertyChangedEvent) override;
+	ANIMGRAPH_API virtual void ReallocatePinsDuringReconstruction(TArray<UEdGraphPin*>& OldPins) override;
 	// End of UEdGraphNode interface
 
 	// UAnimGraphNode_Base interface
