#ANIM: Blueprint creation

- Make sure it doens't refresh parent classes based on skeleton - very confusing
- Exclude SingleNodeInstance and PreviewInstance to be not BP-table - it won't work.

diff --git a/Engine/Source/Editor/AnimGraph/Classes/AnimPreviewInstance.h b/Engine/Source/Editor/AnimGraph/Classes/AnimPreviewInstance.h
--- a/Engine/Source/Editor/AnimGraph/Classes/AnimPreviewInstance.h
+++ b/Engine/Source/Editor/AnimGraph/Classes/AnimPreviewInstance.h
@@ -24,7 +24,7 @@ enum EMontagePreviewType
  * Used by Preview in AnimGraph, Playing single animation in Kismet2 and etc
  */
 
-UCLASS(transient, noteditinlinenew)
+UCLASS(transient, NotBlueprintable, noteditinlinenew)
 class ANIMGRAPH_API UAnimPreviewInstance : public UAnimSingleNodeInstance
 {
 	GENERATED_UCLASS_BODY()
diff --git a/Engine/Source/Editor/UnrealEd/Private/Factories/AnimBlueprintFactory.cpp b/Engine/Source/Editor/UnrealEd/Private/Factories/AnimBlueprintFactory.cpp
--- a/Engine/Source/Editor/UnrealEd/Private/Factories/AnimBlueprintFactory.cpp
+++ b/Engine/Source/Editor/UnrealEd/Private/Factories/AnimBlueprintFactory.cpp
@@ -253,7 +253,6 @@ private:
 	void OnSkeletonSelected(const FAssetData& AssetData)
 	{
 		TargetSkeleton = AssetData;
-		MakeParentClassPicker();
 	}
 
 	/** Handler for when ok is clicked */
diff --git a/Engine/Source/Runtime/Engine/Classes/Animation/AnimSingleNodeInstance.h b/Engine/Source/Runtime/Engine/Classes/Animation/AnimSingleNodeInstance.h
--- a/Engine/Source/Runtime/Engine/Classes/Animation/AnimSingleNodeInstance.h
+++ b/Engine/Source/Runtime/Engine/Classes/Animation/AnimSingleNodeInstance.h
@@ -12,7 +12,7 @@
 
 DECLARE_DYNAMIC_DELEGATE(FPostEvaluateAnimEvent);
 
-UCLASS(transient)
+UCLASS(transient, NotBlueprintable)
 class ENGINE_API UAnimSingleNodeInstance : public UAnimInstance
 {
 	GENERATED_UCLASS_BODY()
