Fix crash in bone control widget when blueprint needs recompiling

diff --git a/Engine/Source/Editor/AnimGraph/Private/AnimGraphNode_ModifyBone.cpp b/Engine/Source/Editor/AnimGraph/Private/AnimGraphNode_ModifyBone.cpp
--- a/Engine/Source/Editor/AnimGraph/Private/AnimGraphNode_ModifyBone.cpp
+++ b/Engine/Source/Editor/AnimGraph/Private/AnimGraphNode_ModifyBone.cpp
@@ -67,7 +67,7 @@ FVector UAnimGraphNode_ModifyBone::GetWidgetLocation(const USkeletalMeshComponen
 		FA2CSPose& MeshBases = AnimNode->ForwardedPose;
 		WidgetLoc = ConvertWidgetLocation(SkelComp, MeshBases, Node.BoneToModify.BoneName, Node.Translation, Node.TranslationSpace);
 
-		if(Node.TranslationMode == BMM_Additive)
+		if (MeshBases.IsValid() && Node.TranslationMode == BMM_Additive)
 		{
 			if(Node.TranslationSpace == EBoneControlSpace::BCS_WorldSpace ||
 				Node.TranslationSpace == EBoneControlSpace::BCS_ComponentSpace)
