[INTEGRATE] Change 2440548 by [REDACTED_PERSON][REDACTED_HANDLE] on 2015/02/10 17:50:56

	[UE-9334] Fix for crash that occured when working with anim blueprints and the old menuing system. Origin was 2340250

diff --git a/Engine/Source/Editor/AnimGraph/Private/AnimGraphNode_Slot.cpp b/Engine/Source/Editor/AnimGraph/Private/AnimGraphNode_Slot.cpp
--- a/Engine/Source/Editor/AnimGraph/Private/AnimGraphNode_Slot.cpp
+++ b/Engine/Source/Editor/AnimGraph/Private/AnimGraphNode_Slot.cpp
@@ -28,7 +28,7 @@ FText UAnimGraphNode_Slot::GetTooltipText() const
 
 FText UAnimGraphNode_Slot::GetNodeTitle(ENodeTitleType::Type TitleType) const
 {
-	if (Node.SlotName == NAME_None)
+	if (Node.SlotName == NAME_None || !HasValidBlueprint() )
 	{
 		if (TitleType == ENodeTitleType::ListView || TitleType == ENodeTitleType::MenuTitle)
 		{
