[INTEGRATE] Change 2443244 by [REDACTED_PERSON][REDACTED_HANDLE] on 2015/02/12 07:06:56

	Fix manipulation of transforms on bone control nodes so that selecting a node doesn't overwrite its conents

diff --git a/Engine/Source/Editor/AnimGraph/Classes/AnimGraphNode_ModifyBone.h b/Engine/Source/Editor/AnimGraph/Classes/AnimGraphNode_ModifyBone.h
--- a/Engine/Source/Editor/AnimGraph/Classes/AnimGraphNode_ModifyBone.h
+++ b/Engine/Source/Editor/AnimGraph/Classes/AnimGraphNode_ModifyBone.h
@@ -33,9 +33,6 @@ protected:
 
 	ANIMGRAPH_API virtual FName FindSelectedBone() override;
 
-	ANIMGRAPH_API virtual void UpdateDefaultValues(const FAnimNode_Base* AnimNode) override;
-	ANIMGRAPH_API virtual void UpdateAllDefaultValues(const FAnimNode_Base* AnimNode) override;
-
 	ANIMGRAPH_API virtual void DoTranslation(const USkeletalMeshComponent* SkelComp, FVector& Drag, FAnimNode_Base* InOutAnimNode) override;
 	ANIMGRAPH_API virtual void DoRotation(const USkeletalMeshComponent* SkelComp, FRotator& Rotation, FAnimNode_Base* InOutAnimNode) override;
 	ANIMGRAPH_API virtual void DoScale(const USkeletalMeshComponent* SkelComp, FVector& Drag, FAnimNode_Base* InOutAnimNode) override;
diff --git a/Engine/Source/Editor/AnimGraph/Classes/AnimGraphNode_SkeletalControlBase.h b/Engine/Source/Editor/AnimGraph/Classes/AnimGraphNode_SkeletalControlBase.h
--- a/Engine/Source/Editor/AnimGraph/Classes/AnimGraphNode_SkeletalControlBase.h
+++ b/Engine/Source/Editor/AnimGraph/Classes/AnimGraphNode_SkeletalControlBase.h
@@ -35,10 +35,9 @@ public:
 	// to keep data consistency between anim nodes 
 	ANIMGRAPH_API virtual void	CopyNodeDataTo(FAnimNode_Base* OutAnimNode){}
 	ANIMGRAPH_API virtual void	CopyNodeDataFrom(const FAnimNode_Base* NewAnimNode){}
-	// update literal values in Pins
-	ANIMGRAPH_API virtual void	UpdateDefaultValues(const FAnimNode_Base* AnimNode){}
-	// update all literal values for Undo/Redo
-	ANIMGRAPH_API virtual void UpdateAllDefaultValues(const FAnimNode_Base* AnimNode){}
+
+	/** Are we currently showing this pin */
+	ANIMGRAPH_API bool IsPinShown(const FString& PinName) const;
 
 	// return current widget mode this anim graph node supports
 	ANIMGRAPH_API virtual int32 GetWidgetMode(const USkeletalMeshComponent* SkelComp);
@@ -81,5 +80,44 @@ protected:
 	// convert widget location according to bone control space
 	FVector ConvertWidgetLocation(const USkeletalMeshComponent* InSkelComp, FA2CSPose& InMeshBases, const FName& BoneName, const FVector& InLocation, const EBoneControlSpace Space);
 	// set literal value for FVector
-	void SetDefaultValue(FString& InDefaultValueName, FVector& InValue);
+	void SetDefaultValue(const FString& InDefaultValueName, const FVector& InValue);
+	// get literal value for vector
+	void GetDefaultValue(const FString& UpdateDefaultValueName, FVector& OutVec);
+
+	void GetDefaultValue(const FString& PropName, FRotator& OutValue)
+	{
+		FVector Value;
+		GetDefaultValue(PropName, Value);
+		OutValue.Pitch = Value.X;
+		OutValue.Yaw = Value.Y;
+		OutValue.Roll = Value.Z;
+	}
+
+	template<class ValueType>
+	ValueType GetNodeValue(const FString& PropName, const ValueType& CompileNodeValue)
+	{
+		if (IsPinShown(PropName))
+		{
+			ValueType Val;
+			GetDefaultValue(PropName, Val);
+			return Val;
+		}
+		return CompileNodeValue;
+	}
+
+	void SetDefaultValue(const FString& PropName, const FRotator& InValue)
+	{
+		FVector VecValue(InValue.Pitch, InValue.Yaw, InValue.Roll);
+		SetDefaultValue(PropName, VecValue);
+	}
+
+	template<class ValueType>
+	void SetNodeValue(const FString& PropName, ValueType& CompileNodeValue, const ValueType& InValue)
+	{
+		if (IsPinShown(PropName))
+		{
+			SetDefaultValue(PropName, InValue);
+		}
+		CompileNodeValue = InValue;
+	}
 };
diff --git a/Engine/Source/Editor/AnimGraph/Classes/AnimGraphNode_TwoBoneIK.h b/Engine/Source/Editor/AnimGraph/Classes/AnimGraphNode_TwoBoneIK.h
--- a/Engine/Source/Editor/AnimGraph/Classes/AnimGraphNode_TwoBoneIK.h
+++ b/Engine/Source/Editor/AnimGraph/Classes/AnimGraphNode_TwoBoneIK.h
@@ -38,9 +38,6 @@ public:
 
 	ANIMGRAPH_API virtual void	CopyNodeDataFrom(const FAnimNode_Base* NewAnimNode) override;
 
-	ANIMGRAPH_API virtual void	UpdateDefaultValues(const FAnimNode_Base* InAnimNode) override;
-	ANIMGRAPH_API virtual void	UpdateAllDefaultValues(const FAnimNode_Base* InAnimNode) override;
-
 	ANIMGRAPH_API virtual int32 GetWidgetMode(const USkeletalMeshComponent* SkelComp) override;
 
 	ANIMGRAPH_API virtual void MoveSelectActorLocation(const USkeletalMeshComponent* SkelComp, FAnimNode_SkeletalControlBase* AnimNode) override;
diff --git a/Engine/Source/Editor/AnimGraph/Private/AnimGraphNode_ModifyBone.cpp b/Engine/Source/Editor/AnimGraph/Private/AnimGraphNode_ModifyBone.cpp
--- a/Engine/Source/Editor/AnimGraph/Private/AnimGraphNode_ModifyBone.cpp
+++ b/Engine/Source/Editor/AnimGraph/Private/AnimGraphNode_ModifyBone.cpp
@@ -65,7 +65,7 @@ FVector UAnimGraphNode_ModifyBone::GetWidgetLocation(const USkeletalMeshComponen
 	if (CurWidgetMode == FWidget::WM_Translate)
 	{
 		FA2CSPose& MeshBases = AnimNode->ForwardedPose;
-		WidgetLoc = ConvertWidgetLocation(SkelComp, MeshBases, Node.BoneToModify.BoneName, Node.Translation, Node.TranslationSpace);
+		WidgetLoc = ConvertWidgetLocation(SkelComp, MeshBases, Node.BoneToModify.BoneName, GetNodeValue(FString("Translation"), Node.Translation), Node.TranslationSpace);
 
 		if (MeshBases.IsValid() && Node.TranslationMode == BMM_Additive)
 		{
@@ -193,67 +193,6 @@ FName UAnimGraphNode_ModifyBone::FindSelectedBone()
 	return Node.BoneToModify.BoneName;
 }
 
-void UAnimGraphNode_ModifyBone::UpdateDefaultValues(const FAnimNode_Base* AnimNode)
-{
-	FString UpdateDefaultValueName;
-	FVector UpdateValue;
-
-	const FAnimNode_ModifyBone* ModifyBone = static_cast<const FAnimNode_ModifyBone*>(AnimNode);
-
-	if (!ModifyBone)
-	{
-		return;
-	}
-
-	switch (CurWidgetMode)
-	{
-	case FWidget::WM_Translate:
-
-		UpdateDefaultValueName = FString("Translation");
-		UpdateValue = ModifyBone->Translation;
-		break;
-	case FWidget::WM_Rotate:
-		UpdateDefaultValueName = FString("Rotation");
-		UpdateValue.X = ModifyBone->Rotation.Pitch;
-		UpdateValue.Y = ModifyBone->Rotation.Yaw;
-		UpdateValue.Z = ModifyBone->Rotation.Roll;
-		break;
-	case FWidget::WM_Scale:
-		UpdateDefaultValueName = FString("Scale");
-		UpdateValue = ModifyBone->Scale;
-		break;
-	}
-
-	SetDefaultValue(UpdateDefaultValueName, UpdateValue);
-}
-
-void UAnimGraphNode_ModifyBone::UpdateAllDefaultValues(const FAnimNode_Base* AnimNode)
-{
-	FString UpdateDefaultValueName;
-	FVector UpdateValue;
-
-	const FAnimNode_ModifyBone* ModifyBone = static_cast<const FAnimNode_ModifyBone*>(AnimNode);
-
-	if (!ModifyBone)
-	{
-		return;
-	}
-
-	UpdateDefaultValueName = FString("Translation");
-	UpdateValue = ModifyBone->Translation;
-	SetDefaultValue(UpdateDefaultValueName, UpdateValue);
-
-	UpdateDefaultValueName = FString("Rotation");
-	UpdateValue.X = ModifyBone->Rotation.Pitch;
-	UpdateValue.Y = ModifyBone->Rotation.Yaw;
-	UpdateValue.Z = ModifyBone->Rotation.Roll;
-	SetDefaultValue(UpdateDefaultValueName, UpdateValue);
-
-	UpdateDefaultValueName = FString("Scale");
-	UpdateValue = ModifyBone->Scale;
-	SetDefaultValue(UpdateDefaultValueName, UpdateValue);
-}
-
 void UAnimGraphNode_ModifyBone::DoTranslation(const USkeletalMeshComponent* SkelComp, FVector& Drag, FAnimNode_Base* InOutAnimNode)
 {
 	FAnimNode_ModifyBone* ModifyBone = static_cast<FAnimNode_ModifyBone*>(InOutAnimNode);
@@ -300,9 +239,9 @@ void UAnimGraphNode_ModifyBone::CopyNodeDataTo(FAnimNode_Base* OutAnimNode)
 	FAnimNode_ModifyBone* ModifyBone = static_cast<FAnimNode_ModifyBone*>(OutAnimNode);
 
 	// copies Pin values from the internal node to get data which are not compiled yet
-	ModifyBone->Translation = Node.Translation;
-	ModifyBone->Rotation = Node.Rotation;
-	ModifyBone->Scale = Node.Scale;
+	ModifyBone->Translation = GetNodeValue(FString("Translation"), Node.Translation);
+	ModifyBone->Rotation = GetNodeValue(FString("Rotation"), Node.Rotation);
+	ModifyBone->Scale = GetNodeValue(FString("Scale"), Node.Scale);
 
 	// copies Modes
 	ModifyBone->TranslationMode = Node.TranslationMode;
@@ -319,9 +258,17 @@ void UAnimGraphNode_ModifyBone::CopyNodeDataFrom(const FAnimNode_Base* InNewAnim
 {
 	const FAnimNode_ModifyBone* ModifyBone = static_cast<const FAnimNode_ModifyBone*>(InNewAnimNode);
 
-	// copies Pin data from updated values
-	Node.Translation = ModifyBone->Translation;
-	Node.Rotation = ModifyBone->Rotation;
-	Node.Scale = ModifyBone->Scale;
+	switch (CurWidgetMode)
+	{
+	case FWidget::WM_Translate:
+		SetNodeValue(FString("Translation"), Node.Translation, ModifyBone->Translation);
+		break;
+	case FWidget::WM_Rotate:
+		SetNodeValue(FString("Rotation"), Node.Rotation, ModifyBone->Rotation);
+		break;
+	case FWidget::WM_Scale:
+		SetNodeValue(FString("Scale"), Node.Scale, ModifyBone->Scale);
+		break;
+	}
 }
 #undef LOCTEXT_NAMESPACE
diff --git a/Engine/Source/Editor/AnimGraph/Private/AnimGraphNode_SkeletalControlBase.cpp b/Engine/Source/Editor/AnimGraph/Private/AnimGraphNode_SkeletalControlBase.cpp
--- a/Engine/Source/Editor/AnimGraph/Private/AnimGraphNode_SkeletalControlBase.cpp
+++ b/Engine/Source/Editor/AnimGraph/Private/AnimGraphNode_SkeletalControlBase.cpp
@@ -279,7 +279,7 @@ FVector UAnimGraphNode_SkeletalControlBase::ConvertWidgetLocation(const USkeleta
 	return WidgetLoc;
 }
 
-void UAnimGraphNode_SkeletalControlBase::SetDefaultValue(FString& UpdateDefaultValueName, FVector& Value)
+void UAnimGraphNode_SkeletalControlBase::GetDefaultValue(const FString& UpdateDefaultValueName, FVector& OutVec)
 {
 	for (UEdGraphPin* Pin : Pins)
 	{
@@ -287,8 +287,37 @@ void UAnimGraphNode_SkeletalControlBase::SetDefaultValue(FString& UpdateDefaultV
 		{
 			if (GetSchema()->IsCurrentPinDefaultValid(Pin).IsEmpty())
 			{
-				FVector& V = Value;
-				FString Str = FString::Printf(TEXT("%.3f,%.3f,%.3f"), V.X, V.Y, V.Z);
+				FString DefaultString = Pin->GetDefaultAsString();
+				TArray<FString> ResultString;
+
+				//Parse string to split its contents separated by ','
+				DefaultString.Trim();
+				DefaultString.TrimTrailing();
+				DefaultString.ParseIntoArray(&ResultString, TEXT(","), true);
+
+				check(ResultString.Num() == 3);
+
+				OutVec.Set(
+							FCString::Atof(*ResultString[0]),
+							FCString::Atof(*ResultString[1]),
+							FCString::Atof(*ResultString[2])
+							);
+				return;
+			}
+		}
+	}
+	OutVec = FVector::ZeroVector;
+}
+
+void UAnimGraphNode_SkeletalControlBase::SetDefaultValue(const FString& UpdateDefaultValueName, const FVector& Value)
+{
+	for (UEdGraphPin* Pin : Pins)
+	{
+		if (Pin->PinName == UpdateDefaultValueName)
+		{
+			if (GetSchema()->IsCurrentPinDefaultValid(Pin).IsEmpty())
+			{
+				FString Str = FString::Printf(TEXT("%.3f,%.3f,%.3f"), Value.X, Value.Y, Value.Z);
 				if (Pin->DefaultValue != Str)
 				{
 					PreEditChange(NULL);
@@ -301,4 +330,16 @@ void UAnimGraphNode_SkeletalControlBase::SetDefaultValue(FString& UpdateDefaultV
 	}
 }
 
+bool UAnimGraphNode_SkeletalControlBase::IsPinShown(const FString& PinName) const
+{
+	for (const FOptionalPinFromProperty& Pin : ShowPinForProperties)
+	{
+		if (Pin.PropertyName.ToString() == PinName)
+		{
+			return Pin.bShowPin;
+		}
+	}
+	return false;
+}
+
 #undef LOCTEXT_NAMESPACE
\ No newline at end of file
diff --git a/Engine/Source/Editor/AnimGraph/Private/AnimGraphNode_TwoBoneIK.cpp b/Engine/Source/Editor/AnimGraph/Private/AnimGraphNode_TwoBoneIK.cpp
--- a/Engine/Source/Editor/AnimGraph/Private/AnimGraphNode_TwoBoneIK.cpp
+++ b/Engine/Source/Editor/AnimGraph/Private/AnimGraphNode_TwoBoneIK.cpp
@@ -119,17 +119,22 @@ void UAnimGraphNode_TwoBoneIK::CopyNodeDataTo(FAnimNode_Base* AnimNode)
 	FAnimNode_TwoBoneIK* TwoBoneIK = static_cast<FAnimNode_TwoBoneIK*>(AnimNode);
 
 	// copies Pin values from the internal node to get data which are not compiled yet
-	TwoBoneIK->EffectorLocation = Node.EffectorLocation;
-	TwoBoneIK->JointTargetLocation = Node.JointTargetLocation;
+	TwoBoneIK->EffectorLocation = GetNodeValue(FString("EffectorLocation"), Node.EffectorLocation);
+	TwoBoneIK->JointTargetLocation = GetNodeValue(FString("JointTargetLocation"), Node.JointTargetLocation);
 }
 
 void UAnimGraphNode_TwoBoneIK::CopyNodeDataFrom(const FAnimNode_Base* InNewAnimNode)
 {
 	const FAnimNode_TwoBoneIK* TwoBoneIK = static_cast<const FAnimNode_TwoBoneIK*>(InNewAnimNode);
 	
-	// copies Pin data from updated values
-	Node.EffectorLocation = TwoBoneIK->EffectorLocation;
-	Node.JointTargetLocation = TwoBoneIK->JointTargetLocation;
+	if (BoneSelectMode == BSM_EndEffector)
+	{
+		SetNodeValue(FString("EffectorLocation"), Node.EffectorLocation, TwoBoneIK->EffectorLocation);
+	}
+	else
+	{
+		SetNodeValue(FString("JointTargetLocation"), Node.JointTargetLocation, TwoBoneIK->JointTargetLocation);
+	}
 }
 
 
@@ -142,7 +147,7 @@ FVector UAnimGraphNode_TwoBoneIK::GetWidgetLocation(const USkeletalMeshComponent
 	if (BoneSelectMode == BSM_EndEffector)
 	{
 		Space = Node.EffectorLocationSpace;
-		Location = Node.EffectorLocation;
+		Location = GetNodeValue(FString("EffectorLocation"), Node.EffectorLocation);
 		BoneName = Node.EffectorSpaceBoneName;
 
 	}
@@ -154,61 +159,13 @@ FVector UAnimGraphNode_TwoBoneIK::GetWidgetLocation(const USkeletalMeshComponent
 		{
 			return FVector::ZeroVector;
 		}
-		Location = Node.JointTargetLocation;
+		Location = GetNodeValue(FString("JointTargetLocation"), Node.JointTargetLocation);
 		BoneName = Node.JointTargetSpaceBoneName;
 	}
 
 	return ConvertWidgetLocation(SkelComp, AnimNode->ForwardedPose, BoneName, Location, Space);
 }
 
-void UAnimGraphNode_TwoBoneIK::UpdateDefaultValues(const FAnimNode_Base* AnimNode)
-{
-	FString UpdateDefaultValueName;
-	FVector UpdateValue;
-
-	const FAnimNode_TwoBoneIK* TwoBoneIK = static_cast<const FAnimNode_TwoBoneIK*>(AnimNode);
-
-	if (TwoBoneIK == NULL)
-	{
-		return;
-	}
-
-	if (BoneSelectMode == BSM_EndEffector)
-	{
-		UpdateDefaultValueName = FString("EffectorLocation");
-		UpdateValue = TwoBoneIK->EffectorLocation;
-	}
-	else
-	{
-		UpdateDefaultValueName = FString("JointTargetLocation");
-		UpdateValue = TwoBoneIK->JointTargetLocation;
-	}
-
-	SetDefaultValue(UpdateDefaultValueName, UpdateValue);
-}
-
-void UAnimGraphNode_TwoBoneIK::UpdateAllDefaultValues(const FAnimNode_Base* AnimNode)
-{
-	FString UpdateDefaultValueName;
-	FVector UpdateValue;
-
-	const FAnimNode_TwoBoneIK* TwoBoneIK = static_cast<const FAnimNode_TwoBoneIK*>(AnimNode);
-
-	if (TwoBoneIK == NULL)
-	{
-		return;
-	}
-
-	UpdateDefaultValueName = FString("EffectorLocation");
-	UpdateValue = TwoBoneIK->EffectorLocation;
-	SetDefaultValue(UpdateDefaultValueName, UpdateValue);
-
-	UpdateDefaultValueName = FString("JointTargetLocation");
-	UpdateValue = TwoBoneIK->JointTargetLocation;
-	SetDefaultValue(UpdateDefaultValueName, UpdateValue);
-}
-
-
 int32 UAnimGraphNode_TwoBoneIK::GetWidgetMode(const USkeletalMeshComponent* SkelComp)
 {
 	int32 BoneIndex = SkelComp->GetBoneIndex(Node.IKBone.BoneName);
diff --git a/Engine/Source/Editor/Persona/Private/AnimationEditorViewportClient.cpp b/Engine/Source/Editor/Persona/Private/AnimationEditorViewportClient.cpp
--- a/Engine/Source/Editor/Persona/Private/AnimationEditorViewportClient.cpp
+++ b/Engine/Source/Editor/Persona/Private/AnimationEditorViewportClient.cpp
@@ -602,7 +602,6 @@ void FAnimationViewportClient::FindSelectedAnimGraphNode()
 
 						// make same values to ensure data consistency
 						AnimGraphNode->CopyNodeDataTo(AnimNode);
-						AnimGraphNode->UpdateAllDefaultValues(AnimNode);
 
 						WidgetMode = (FWidget::EWidgetMode)AnimGraphNode->GetWidgetMode(PreviewSkelMeshComp.Get());
 						if (WidgetMode == FWidget::WM_Scale)
@@ -620,8 +619,6 @@ void FAnimationViewportClient::FindSelectedAnimGraphNode()
 						{
 							// copy updated values into internal node to ensure data consistency
 							AnimGraphNode->CopyNodeDataFrom(AnimNode);
-
-							AnimGraphNode->UpdateDefaultValues(AnimNode);
 							bShouldUpdateDefaultValues = false;
 						}
 					}
@@ -673,8 +670,6 @@ void FAnimationViewportClient::PostUndo()
 				{
 					// copy undo data
 					AnimGraphNode->CopyNodeDataTo(AnimNode);
-					// update all literal values in anim node
-					AnimGraphNode->UpdateAllDefaultValues(AnimNode);
 				}
 			}
 		}
