Fix shadowed variables

diff --git a/Engine/Source/Developer/ShaderCompilerCommon/Private/HlslParser.cpp b/Engine/Source/Developer/ShaderCompilerCommon/Private/HlslParser.cpp
--- a/Engine/Source/Developer/ShaderCompilerCommon/Private/HlslParser.cpp
+++ b/Engine/Source/Developer/ShaderCompilerCommon/Private/HlslParser.cpp
@@ -434,12 +434,12 @@ namespace CrossCompiler
 		auto OriginalToken = Scanner.GetCurrentTokenIndex();
 		bool bPrimitiveFound = false;
 		auto* FullType = new(Allocator) AST::FFullySpecifiedType(Allocator, Scanner.GetCurrentToken()->SourceInfo);
-		auto Result = ParseDeclarationStorageQualifiers(Scanner, Flags, bPrimitiveFound, &FullType->Qualifier);
-		if (Result == EParseResult::Error)
+		auto ParseResult = ParseDeclarationStorageQualifiers(Scanner, Flags, bPrimitiveFound, &FullType->Qualifier);
+		if (ParseResult == EParseResult::Error)
 		{
 			return EParseResult::Error;
 		}
-		bool bCanBeUnmatched = (Result == EParseResult::NotMatched);
+		bool bCanBeUnmatched = (ParseResult == EParseResult::NotMatched);
 
 		auto* DeclaratorList = new(Allocator) AST::FDeclaratorList(Allocator, FullType->SourceInfo);
 		DeclaratorList->Type = FullType;
@@ -1302,7 +1302,7 @@ check(0);
 		bool bDefaultFound = false;
 		while (Parser.Scanner.HasMoreTokens())
 		{
-			auto* Token = Parser.Scanner.GetCurrentToken();
+			Token = Parser.Scanner.GetCurrentToken();
 			if (Parser.Scanner.MatchToken(EHlslToken::RightBrace))
 			{
 				break;
diff --git a/Engine/Source/Developer/ShaderFormatOpenGL/Private/GlslBackend.cpp b/Engine/Source/Developer/ShaderFormatOpenGL/Private/GlslBackend.cpp
--- a/Engine/Source/Developer/ShaderFormatOpenGL/Private/GlslBackend.cpp
+++ b/Engine/Source/Developer/ShaderFormatOpenGL/Private/GlslBackend.cpp
@@ -3651,7 +3651,7 @@ static ir_rvalue* GenShaderInputSemantic(
 	else
 	{
 		// Array variable, not first pass. It already exists, get it.
-		ir_variable* Variable = ParseState->symbols->get_variable(ralloc_asprintf(ParseState, "in_%s", Semantic));
+		Variable = ParseState->symbols->get_variable(ralloc_asprintf(ParseState, "in_%s", Semantic));
 		check(Variable);
 
 		ir_rvalue* VariableDeref = new(ParseState)ir_dereference_variable(Variable);
diff --git a/Engine/Source/Editor/AnimGraph/Private/K2Node_TransitionRuleGetter.cpp b/Engine/Source/Editor/AnimGraph/Private/K2Node_TransitionRuleGetter.cpp
--- a/Engine/Source/Editor/AnimGraph/Private/K2Node_TransitionRuleGetter.cpp
+++ b/Engine/Source/Editor/AnimGraph/Private/K2Node_TransitionRuleGetter.cpp
@@ -119,19 +119,19 @@ FText UK2Node_TransitionRuleGetter::GetNodeTitle(ENodeTitleType::Type TitleType)
 
 void UK2Node_TransitionRuleGetter::GetStateSpecificAnimGraphSchemaMenuActions(FBlueprintActionDatabaseRegistrar& ActionRegistrar, const UAnimBlueprint* AnimBlueprint, UAnimStateNode* StateNode) const
 {
-	auto UiSpecOverride = [](const FBlueprintActionContext& /*Context*/, const IBlueprintNodeBinder::FBindingSet& Bindings, FBlueprintActionUiSpec* UiSpecOut, UAnimStateNode* StateNode)
+	auto UiSpecOverride = [](const FBlueprintActionContext& /*Context*/, const IBlueprintNodeBinder::FBindingSet& Bindings, FBlueprintActionUiSpec* UiSpecOut, UAnimStateNode* AnimStateNode)
 	{
-		const FString OwnerName = StateNode->GetOuter()->GetName();
+		const FString OwnerName = AnimStateNode->GetOuter()->GetName();
 		UiSpecOut->MenuName = FText::Format(LOCTEXT("TransitionRuleGetterTitle", "Current {0} for state '{1}.{2}'"), 
 			UK2Node_TransitionRuleGetter::GetFriendlyName(ETransitionGetter::ArbitraryState_GetBlendWeight), 
 			FText::FromString(OwnerName), 
-			FText::FromString(StateNode->GetStateName()));
+			FText::FromString(AnimStateNode->GetStateName()));
 	};
 
-	auto PostSpawnSetupLambda = [](UEdGraphNode* NewNode, bool /*bIsTemplateNode*/, UAnimStateNode* StateNode)
+	auto PostSpawnSetupLambda = [](UEdGraphNode* NewNode, bool /*bIsTemplateNode*/, UAnimStateNode* AnimStateNode)
 	{
 		UK2Node_TransitionRuleGetter* NewNodeTyped = CastChecked<UK2Node_TransitionRuleGetter>(NewNode);
-		NewNodeTyped->AssociatedStateNode = StateNode;
+		NewNodeTyped->AssociatedStateNode = AnimStateNode;
 		NewNodeTyped->GetterType = ETransitionGetter::ArbitraryState_GetBlendWeight;
 	};
 
@@ -166,32 +166,30 @@ void UK2Node_TransitionRuleGetter::GetStateSpecificAnimTransitionSchemaMenuActio
 
 			if (AnimNode->DoesSupportTimeForTransitionGetter())
 			{
-				FString AssetName;
 				UAnimationAsset * AnimAsset = AnimNode->GetAnimationAsset();
-				UAnimGraphNode_Base* AssociatedAnimAssetPlayerNode = nullptr;
 				if (AnimAsset)
 				{
-					auto UiSpecOverride = [](const FBlueprintActionContext& /*Context*/, const IBlueprintNodeBinder::FBindingSet& Bindings, FBlueprintActionUiSpec* UiSpecOut, FString AssetName, TEnumAsByte<ETransitionGetter::Type> GetterType)
+					auto UiSpecOverride = [](const FBlueprintActionContext& /*Context*/, const IBlueprintNodeBinder::FBindingSet& Bindings, FBlueprintActionUiSpec* UiSpecOut, FString AssetName, TEnumAsByte<ETransitionGetter::Type> InGetterType)
 					{
 						UiSpecOut->Category = LOCTEXT("AssetPlayer", "Asset Player");
 
 						FFormatNamedArguments Args;
-						Args.Add(TEXT("NodeName"), UK2Node_TransitionRuleGetter::GetFriendlyName(GetterType));
+						Args.Add(TEXT("NodeName"), UK2Node_TransitionRuleGetter::GetFriendlyName(InGetterType));
 						Args.Add(TEXT("AssetName"), FText::FromString(AssetName));
 						FText Title = FText::Format(LOCTEXT("TransitionFor", "{NodeName} for '{AssetName}'"), Args);
 						UiSpecOut->MenuName = Title;
 					};
 
-					auto PostSpawnSetupLambda = [](UEdGraphNode* NewNode, bool /*bIsTemplateNode*/, UAnimGraphNode_Base* AssociatedAnimAssetPlayerNode, TEnumAsByte<ETransitionGetter::Type> GetterType)
+					auto PostSpawnSetupLambda = [](UEdGraphNode* NewNode, bool /*bIsTemplateNode*/, UAnimGraphNode_Base* InAssociatedAnimAssetPlayerNode, TEnumAsByte<ETransitionGetter::Type> InGetterType)
 					{
 						UK2Node_TransitionRuleGetter* NewNodeTyped = CastChecked<UK2Node_TransitionRuleGetter>(NewNode);
-						NewNodeTyped->AssociatedAnimAssetPlayerNode = AssociatedAnimAssetPlayerNode;
-						NewNodeTyped->GetterType = GetterType;
+						NewNodeTyped->AssociatedAnimAssetPlayerNode = InAssociatedAnimAssetPlayerNode;
+						NewNodeTyped->GetterType = InGetterType;
 					};
 
 					// Prepare the node spawner
-					AssociatedAnimAssetPlayerNode = AnimNode;
-					AssetName = AnimAsset->GetName();
+					UAnimGraphNode_Base* AssociatedAnimAssetPlayerNode = AnimNode;
+					const FString AssetName = AnimAsset->GetName();
 
 					TEnumAsByte<ETransitionGetter::Type> GetterType = SequenceSpecificGetters[TypeIndex];
 
diff --git a/Engine/Source/Runtime/Analytics/AnalyticsET/Private/AnalyticsET.cpp b/Engine/Source/Runtime/Analytics/AnalyticsET/Private/AnalyticsET.cpp
--- a/Engine/Source/Runtime/Analytics/AnalyticsET/Private/AnalyticsET.cpp
+++ b/Engine/Source/Runtime/Analytics/AnalyticsET/Private/AnalyticsET.cpp
@@ -366,7 +366,7 @@ void FAnalyticsProviderET::FlushEvents()
 			// NOTE - This branch is temp, and we will eventually use the DataRouter path exclusively
 
 			// Create/send Http request for an event
-			TSharedRef<IHttpRequest> HttpRequest = FHttpModule::Get().CreateRequest();
+			HttpRequest = FHttpModule::Get().CreateRequest();
 			HttpRequest->SetHeader(TEXT("Content-Type"), TEXT("application/json; charset=utf-8"));
 
 			// TODO need agent here??
diff --git a/Engine/Source/Runtime/AppFramework/Private/Framework/Testing/STableViewTesting.cpp b/Engine/Source/Runtime/AppFramework/Private/Framework/Testing/STableViewTesting.cpp
--- a/Engine/Source/Runtime/AppFramework/Private/Framework/Testing/STableViewTesting.cpp
+++ b/Engine/Source/Runtime/AppFramework/Private/Framework/Testing/STableViewTesting.cpp
@@ -48,7 +48,7 @@ class FTestData
 		 */
 		static bool RemoveRecursive(TArray< TSharedPtr< FTestData > >& RemoveFrom, const TSharedPtr<FTestData>& ItemToRemove)
 		{
-			const int32 ItemIndex = RemoveFrom.Find(ItemToRemove);
+			int32 ItemIndex = RemoveFrom.Find(ItemToRemove);
 			if (ItemIndex != INDEX_NONE)
 			{
 				RemoveFrom.RemoveAt(ItemIndex);
@@ -56,7 +56,7 @@ class FTestData
 			}
 
 			// Did not successfully remove an item. Try all the children.
-			for (int32 ItemIndex = 0; ItemIndex < RemoveFrom.Num(); ++ItemIndex)
+			for (ItemIndex = 0; ItemIndex < RemoveFrom.Num(); ++ItemIndex)
 			{
 				if (RemoveRecursive(RemoveFrom[ItemIndex]->Children, ItemToRemove))
 				{
diff --git a/Engine/Source/Runtime/Core/Public/Mac/MacPlatformMisc.h b/Engine/Source/Runtime/Core/Public/Mac/MacPlatformMisc.h
--- a/Engine/Source/Runtime/Core/Public/Mac/MacPlatformMisc.h
+++ b/Engine/Source/Runtime/Core/Public/Mac/MacPlatformMisc.h
@@ -196,7 +196,7 @@ private:
 	NSAutoreleasePool*	Pool;
 };
 
-#define SCOPED_AUTORELEASE_POOL const FScopeAutoreleasePool Pool;
+#define SCOPED_AUTORELEASE_POOL const FScopeAutoreleasePool Pool##__LINE__;
 
 #endif // __OBJC__
 
diff --git a/Engine/Source/Runtime/CoreUObject/Private/UObject/Linker.cpp b/Engine/Source/Runtime/CoreUObject/Private/UObject/Linker.cpp
--- a/Engine/Source/Runtime/CoreUObject/Private/UObject/Linker.cpp
+++ b/Engine/Source/Runtime/CoreUObject/Private/UObject/Linker.cpp
@@ -452,7 +452,7 @@ static void LogGetPackageLinkerError(FArchiveUObject* LinkerArchive, const TCHAR
 	struct Local
 	{
 		/** Helper function to output more detailed error info if available */
-		static void OutputErrorDetail(FArchiveUObject* LinkerArchive, const FName& LogName)
+		static void OutputErrorDetail(FArchiveUObject* InLinkerArchive, const FName& LogName)
 		{
 			if ( GSerializedObject && GSerializedImportLinker )
 			{
@@ -463,7 +463,7 @@ static void LogGetPackageLinkerError(FArchiveUObject* LinkerArchive, const TCHAR
 				Message->AddToken(FAssetNameToken::Create(GSerializedImportLinker->GetImportPathName(GSerializedImportIndex)));
 				Message->AddToken(FTextToken::Create(LOCTEXT("FailedLoad_Referenced", "Referenced by")));
 				Message->AddToken(FUObjectToken::Create(GSerializedObject));
-				auto SerializedProperty = LinkerArchive ? LinkerArchive->GetSerializedProperty() : nullptr;
+				auto SerializedProperty = InLinkerArchive ? InLinkerArchive->GetSerializedProperty() : nullptr;
 				if (SerializedProperty != nullptr)
 				{
 					FString PropertyPathName = SerializedProperty->GetPathName();
diff --git a/Engine/Source/Runtime/CoreUObject/Private/UObject/Obj.cpp b/Engine/Source/Runtime/CoreUObject/Private/UObject/Obj.cpp
--- a/Engine/Source/Runtime/CoreUObject/Private/UObject/Obj.cpp
+++ b/Engine/Source/Runtime/CoreUObject/Private/UObject/Obj.cpp
@@ -2976,8 +2976,7 @@ bool StaticExec( UWorld* InWorld, const TCHAR* Cmd, FOutputDevice& Ar )
 
 				if (Errors.Num() > 0)
 				{
-					const FString ErrorStr = FString::Printf(TEXT("Errors for %s"), *Target->GetName());
-					Ar.Logf(*ErrorStr);
+					Ar.Logf(*FString::Printf(TEXT("Errors for %s"), *Target->GetName()));
 
 					for (auto ErrorStr : Errors)
 					{
diff --git a/Engine/Source/Runtime/CoreUObject/Private/UObject/SavePackage.cpp b/Engine/Source/Runtime/CoreUObject/Private/UObject/SavePackage.cpp
--- a/Engine/Source/Runtime/CoreUObject/Private/UObject/SavePackage.cpp
+++ b/Engine/Source/Runtime/CoreUObject/Private/UObject/SavePackage.cpp
@@ -522,8 +522,8 @@ public:
 	TArray<UObject*> Dependencies;
 	TArray<FString>& StringAssetReferencesMap;
 
-	FArchiveSaveTagImports(ULinkerSave* InLinker, TArray<FString>& StringAssetReferencesMap)
-		: Linker(InLinker), StringAssetReferencesMap(StringAssetReferencesMap)
+	FArchiveSaveTagImports(ULinkerSave* InLinker, TArray<FString>& InStringAssetReferencesMap)
+		: Linker(InLinker), StringAssetReferencesMap(InStringAssetReferencesMap)
 	{
 		ArIsSaving				= true;
 		ArIsPersistent			= true;
diff --git a/Engine/Source/Runtime/NetworkFile/Private/NetworkPlatformFile.cpp b/Engine/Source/Runtime/NetworkFile/Private/NetworkPlatformFile.cpp
--- a/Engine/Source/Runtime/NetworkFile/Private/NetworkPlatformFile.cpp
+++ b/Engine/Source/Runtime/NetworkFile/Private/NetworkPlatformFile.cpp
@@ -857,14 +857,14 @@ void AsyncReadUnsolicitedFiles(int32 InNumUnsolictedFiles, FNetworkPlatformFile&
 		FScopedEvent* NetworkDoneEvent; // finished using the network
 		FScopedEvent* WritingDoneEvent; // finished writing the files to disk
 
-		FAsyncReadUnsolicitedFile(int32 In_NumUnsolictedFiles, FNetworkPlatformFile* In_NetworkFile, IPlatformFile* In_InnerPlatformFile, FString& In_ServerEngineDir, FString& In_ServerGameDir, FScopedEvent *InNetworkDoneEvent, FScopedEvent *InWritingDoneEvent )
+		FAsyncReadUnsolicitedFile(int32 In_NumUnsolictedFiles, FNetworkPlatformFile* In_NetworkFile, IPlatformFile* In_InnerPlatformFile, FString& In_ServerEngineDir, FString& In_ServerGameDir, FScopedEvent *In_NetworkDoneEvent, FScopedEvent *In_WritingDoneEvent )
 			: NumUnsolictedFiles(In_NumUnsolictedFiles)
 			, NetworkFile(*In_NetworkFile)
 			, InnerPlatformFile(*In_InnerPlatformFile)
 			, ServerEngineDir(In_ServerEngineDir)
 			, ServerGameDir(In_ServerGameDir)
-			, NetworkDoneEvent(InNetworkDoneEvent)
-			, WritingDoneEvent(InWritingDoneEvent)
+			, NetworkDoneEvent(In_NetworkDoneEvent)
+			, WritingDoneEvent(In_WritingDoneEvent)
 		{
 		}
 		
diff --git a/Engine/Source/Runtime/NetworkReplayStreaming/HttpNetworkReplayStreaming/Private/HttpNetworkReplayStreaming.cpp b/Engine/Source/Runtime/NetworkReplayStreaming/HttpNetworkReplayStreaming/Private/HttpNetworkReplayStreaming.cpp
--- a/Engine/Source/Runtime/NetworkReplayStreaming/HttpNetworkReplayStreaming/Private/HttpNetworkReplayStreaming.cpp
+++ b/Engine/Source/Runtime/NetworkReplayStreaming/HttpNetworkReplayStreaming/Private/HttpNetworkReplayStreaming.cpp
@@ -437,7 +437,7 @@ void FHttpNetworkReplayStreamer::HttpHeaderUploadFinished( FHttpRequestPtr HttpR
 	}
 }
 
-void FHttpNetworkReplayStreamer::HttpUploadFinished( FHttpRequestPtr HttpRequest, FHttpResponsePtr HttpResponse, bool bSucceeded )
+void FHttpNetworkReplayStreamer::HttpUploadFinished( FHttpRequestPtr , FHttpResponsePtr HttpResponse, bool bSucceeded )
 {
 	check( HttpState == EHttptate::UploadingStream );
 	check( StreamerState == EStreamerState::StreamingUp || StreamerState == EStreamerState::StreamingUpFinal );
diff --git a/Engine/Source/Runtime/Projects/Private/PluginDescriptor.cpp b/Engine/Source/Runtime/Projects/Private/PluginDescriptor.cpp
--- a/Engine/Source/Runtime/Projects/Private/PluginDescriptor.cpp
+++ b/Engine/Source/Runtime/Projects/Private/PluginDescriptor.cpp
@@ -56,10 +56,10 @@ bool FPluginDescriptor::Read( const FJsonObject& Object, FText& OutFailReason )
 	}
 
 	// Check that it's within range
-	EPluginDescriptorVersion::Type FileVersion = (EPluginDescriptorVersion::Type)FileVersionInt32;
-	if ((FileVersion <= EPluginDescriptorVersion::Invalid) || (FileVersion > EPluginDescriptorVersion::Latest))
+	EPluginDescriptorVersion::Type PluginFileVersion = (EPluginDescriptorVersion::Type)FileVersionInt32;
+	if ((PluginFileVersion <= EPluginDescriptorVersion::Invalid) || (PluginFileVersion > EPluginDescriptorVersion::Latest))
 	{
-		FText ReadVersionText = FText::FromString(FString::Printf(TEXT("%d"), (int32)FileVersion));
+		FText ReadVersionText = FText::FromString(FString::Printf(TEXT("%d"), (int32)PluginFileVersion));
 		FText LatestVersionText = FText::FromString(FString::Printf(TEXT("%d"), (int32)EPluginDescriptorVersion::Latest));
 		OutFailReason = FText::Format( LOCTEXT("ProjectFileVersionTooLarge", "File appears to be in a newer version ({0}) of the file format that we can load (max version: {1})."), ReadVersionText, LatestVersionText);
 		return false;
@@ -166,13 +166,13 @@ bool FPluginReferenceDescriptor::ReadArray( const FJsonObject& Object, const TCH
 	{
 		for (const TSharedPtr<FJsonValue> &Item : *Array)
 		{
-			const TSharedPtr<FJsonObject> *Object;
+			const TSharedPtr<FJsonObject> *ObjectPtr;
 
-			if (Item.IsValid() && Item->TryGetObject(Object))
+			if (Item.IsValid() && Item->TryGetObject(ObjectPtr))
 			{
 				FPluginReferenceDescriptor Plugin;
 
-				if (!Plugin.Read(*Object->Get(), OutFailReason))
+				if (!Plugin.Read(*ObjectPtr->Get(), OutFailReason))
 				{
 					return false;
 				}
