Fix shadowed variables

diff --git a/Engine/Source/Developer/DesktopPlatform/Private/Mac/DesktopPlatformMac.cpp b/Engine/Source/Developer/DesktopPlatform/Private/Mac/DesktopPlatformMac.cpp
--- a/Engine/Source/Developer/DesktopPlatform/Private/Mac/DesktopPlatformMac.cpp
+++ b/Engine/Source/Developer/DesktopPlatform/Private/Mac/DesktopPlatformMac.cpp
@@ -273,7 +273,7 @@ bool FDesktopPlatformMac::OpenDirectoryDialog(const void* ParentWindowHandle, co
 			[Panel setDirectoryURL: DefaultPathURL];
 			CFRelease(DefaultPathCFString);
 
-			bool bSuccess = false;
+			bool bResult = false;
 
 			NSInteger Result = [Panel runModal];
 
@@ -285,12 +285,12 @@ bool FDesktopPlatformMac::OpenDirectoryDialog(const void* ParentWindowHandle, co
 				OutFolderName = FolderPath;
 				FPaths::NormalizeFilename(OutFolderName);
 
-				bSuccess = true;
+				bResult = true;
 			}
 
 			[Panel close];
 
-			return bSuccess;
+			return bResult;
 		});
 	}
 	MacApplication->ResetModifierKeys();
@@ -319,13 +319,13 @@ bool FDesktopPlatformMac::OpenFontDialog(const void* ParentWindowHandle, FString
 			
 			[Panel close];
 
-			bool bSuccess = [AccessoryView result];
+			bool bResult = [AccessoryView result];
 
 			[Panel setAccessoryView: NULL];
 			[AccessoryView release];
 			[[Panel standardWindowButton: NSWindowCloseButton] setEnabled: true];
 
-			if( bSuccess )
+			if( bResult )
 			{
 				NSFont* Font = [Panel panelConvertFont: [NSFont userFontOfSize: 0]];
 
@@ -345,7 +345,7 @@ bool FDesktopPlatformMac::OpenFontDialog(const void* ParentWindowHandle, FString
 				OutFlags = FontFlags;
 			}
 
-			return bSuccess;
+			return bResult;
 		});
 	}
 	
@@ -489,7 +489,7 @@ bool FDesktopPlatformMac::FileDialogShared(bool bSave, const void* ParentWindowH
 
 			[AccessoryView AddAllowedFileTypes:AllowedFileTypes];
 
-			bool bSuccess = false;
+			bool bOkPressed = false;
 			NSWindow* FocusWindow = [[NSApplication sharedApplication] keyWindow];
 
 			NSInteger Result = [Panel runModal];
@@ -521,7 +521,7 @@ bool FDesktopPlatformMac::FileDialogShared(bool bSave, const void* ParentWindowH
 					FPaths::NormalizeFilename(*FilenameIt);
 				}
 
-				bSuccess = true;
+				bOkPressed = true;
 			}
 
 			[Panel close];
@@ -531,7 +531,7 @@ bool FDesktopPlatformMac::FileDialogShared(bool bSave, const void* ParentWindowH
 				[FocusWindow makeKeyWindow];
 			}
 
-			return bSuccess;
+			return bOkPressed;
 		});
 	}
 	
diff --git a/Engine/Source/Developer/ShaderCompilerCommon/Private/ShaderCompilerCommon.cpp b/Engine/Source/Developer/ShaderCompilerCommon/Private/ShaderCompilerCommon.cpp
--- a/Engine/Source/Developer/ShaderCompilerCommon/Private/ShaderCompilerCommon.cpp
+++ b/Engine/Source/Developer/ShaderCompilerCommon/Private/ShaderCompilerCommon.cpp
@@ -273,7 +273,7 @@ FString CreateCrossCompilerBatchFileContents(const FString& ShaderFile, const FS
 	}
 	else if (PLATFORM_WINDOWS)
 	{
-		FString BatchFile = TEXT("@echo off");
+		BatchFile = TEXT("@echo off");
 		BatchFile += TEXT("\nif defined ue.hlslcc GOTO DONE\nset ue.hlslcc=");
 		BatchFile += FPaths::RootDir() / TEXT("Engine\\Binaries\\Win64\\CrossCompilerTool.exe");
 		BatchFile += TEXT("\n\n:DONE\n%ue.hlslcc% ");
diff --git a/Engine/Source/Editor/AnimGraph/Private/K2Node_TransitionRuleGetter.cpp b/Engine/Source/Editor/AnimGraph/Private/K2Node_TransitionRuleGetter.cpp
--- a/Engine/Source/Editor/AnimGraph/Private/K2Node_TransitionRuleGetter.cpp
+++ b/Engine/Source/Editor/AnimGraph/Private/K2Node_TransitionRuleGetter.cpp
@@ -188,13 +188,13 @@ void UK2Node_TransitionRuleGetter::GetStateSpecificAnimTransitionSchemaMenuActio
 					};
 
 					// Prepare the node spawner
-					UAnimGraphNode_Base* AssociatedAnimAssetPlayerNode = AnimNode;
+					UAnimGraphNode_Base* AssociatedAnimNode = AnimNode;
 					const FString AssetName = AnimAsset->GetName();
 
-					TEnumAsByte<ETransitionGetter::Type> GetterType = SequenceSpecificGetters[TypeIndex];
+					TEnumAsByte<ETransitionGetter::Type> TransitionGetterType = SequenceSpecificGetters[TypeIndex];
 
-					UBlueprintNodeSpawner* Spawner = UBlueprintNodeSpawner::Create( UK2Node_TransitionRuleGetter::StaticClass(), nullptr, UBlueprintNodeSpawner::FCustomizeNodeDelegate::CreateStatic(PostSpawnSetupLambda, AssociatedAnimAssetPlayerNode, GetterType) );
-					Spawner->DynamicUiSignatureGetter = UBlueprintNodeSpawner::FUiSpecOverrideDelegate::CreateStatic(UiSpecOverride, AssetName, GetterType);
+					UBlueprintNodeSpawner* Spawner = UBlueprintNodeSpawner::Create( UK2Node_TransitionRuleGetter::StaticClass(), nullptr, UBlueprintNodeSpawner::FCustomizeNodeDelegate::CreateStatic(PostSpawnSetupLambda, AssociatedAnimNode, TransitionGetterType) );
+					Spawner->DynamicUiSignatureGetter = UBlueprintNodeSpawner::FUiSpecOverrideDelegate::CreateStatic(UiSpecOverride, AssetName, TransitionGetterType);
 					ActionRegistrar.AddBlueprintAction( AnimBlueprint, Spawner );
 				}
 			}
@@ -240,21 +240,21 @@ void UK2Node_TransitionRuleGetter::GetNonStateSpecificMenuActions(FBlueprintActi
 
 		for (int32 TypeIndex = 0; TypeIndex < ARRAY_COUNT(NonSpecificGetters); ++TypeIndex)
 		{
-			auto UiSpecOverride = [](const FBlueprintActionContext& /*Context*/, const IBlueprintNodeBinder::FBindingSet& Bindings, FBlueprintActionUiSpec* UiSpecOut, TEnumAsByte<ETransitionGetter::Type> GetterType)
+			auto UiSpecOverride = [](const FBlueprintActionContext& /*Context*/, const IBlueprintNodeBinder::FBindingSet& Bindings, FBlueprintActionUiSpec* UiSpecOut, TEnumAsByte<ETransitionGetter::Type> InGetterType)
 			{
 				UiSpecOut->Category = LOCTEXT("Transition", "Transition");
-				UiSpecOut->MenuName = UK2Node_TransitionRuleGetter::GetFriendlyName(GetterType);
+				UiSpecOut->MenuName = UK2Node_TransitionRuleGetter::GetFriendlyName(InGetterType);
 			};
 
-			auto PostSpawnSetupLambda = [](UEdGraphNode* NewNode, bool /*bIsTemplateNode*/, TEnumAsByte<ETransitionGetter::Type> GetterType)
+			auto PostSpawnSetupLambda = [](UEdGraphNode* NewNode, bool /*bIsTemplateNode*/, TEnumAsByte<ETransitionGetter::Type> InGetterType)
 			{
 				UK2Node_TransitionRuleGetter* NewNodeTyped = CastChecked<UK2Node_TransitionRuleGetter>(NewNode);
-				NewNodeTyped->GetterType = GetterType;
+				NewNodeTyped->GetterType = InGetterType;
 			};
 
-			TEnumAsByte<ETransitionGetter::Type> GetterType = NonSpecificGetters[TypeIndex];
-			UBlueprintNodeSpawner* Spawner = UBlueprintNodeSpawner::Create( UK2Node_TransitionRuleGetter::StaticClass(), nullptr, UBlueprintNodeSpawner::FCustomizeNodeDelegate::CreateStatic(PostSpawnSetupLambda, GetterType) );
-			Spawner->DynamicUiSignatureGetter = UBlueprintNodeSpawner::FUiSpecOverrideDelegate::CreateStatic(UiSpecOverride, GetterType);
+			TEnumAsByte<ETransitionGetter::Type> TransitionGetterType = NonSpecificGetters[TypeIndex];
+			UBlueprintNodeSpawner* Spawner = UBlueprintNodeSpawner::Create( UK2Node_TransitionRuleGetter::StaticClass(), nullptr, UBlueprintNodeSpawner::FCustomizeNodeDelegate::CreateStatic(PostSpawnSetupLambda, TransitionGetterType) );
+			Spawner->DynamicUiSignatureGetter = UBlueprintNodeSpawner::FUiSpecOverrideDelegate::CreateStatic(UiSpecOverride, TransitionGetterType);
 			ActionRegistrar.AddBlueprintAction( ActionKey, Spawner );
 		}
 	}
diff --git a/Engine/Source/Editor/UnrealEd/Private/Factories/PackFactory.cpp b/Engine/Source/Editor/UnrealEd/Private/Factories/PackFactory.cpp
--- a/Engine/Source/Editor/UnrealEd/Private/Factories/PackFactory.cpp
+++ b/Engine/Source/Editor/UnrealEd/Private/Factories/PackFactory.cpp
@@ -264,9 +264,9 @@ UObject* UPackFactory::FactoryCreateBinary
 			ContentDestinationRoot /= PakFile.GetMountPoint().RightChop(ChopIndex + ContentFolder.Len());
 		}
 
-		TArray<uint8> Buffer;
+		TArray<uint8> CopyBuffer;
 		TArray<uint8> PersistentCompressionBuffer;
-		Buffer.AddUninitialized(8 * 1024 * 1024); // 8MB buffer for extracting
+		CopyBuffer.AddUninitialized(8 * 1024 * 1024); // 8MB buffer for extracting
 		int32 ErrorCount = 0;
 		int32 FileCount = 0;
 
@@ -290,7 +290,7 @@ UObject* UPackFactory::FactoryCreateBinary
 				if (EntryInfo == Entry)
 				{
 					FString ConfigString;
-					PackFactoryHelper::ExtractFileToString(Entry, PakReader, Buffer, PersistentCompressionBuffer, ConfigString);
+					PackFactoryHelper::ExtractFileToString(Entry, PakReader, CopyBuffer, PersistentCompressionBuffer, ConfigString);
 					PackFactoryHelper::ProcessPackConfig(ConfigString, ConfigParameters);
 				}
 				else
@@ -418,7 +418,7 @@ UObject* UPackFactory::FactoryCreateBinary
 					UE_LOG(LogPackFactory, Log, TEXT("%s (%ld) -> %s"), *It.Filename(), Entry.Size, *DestFilename);
 
 					FString SourceContents;
-					PackFactoryHelper::ExtractFileToString(Entry, PakReader, Buffer, PersistentCompressionBuffer, SourceContents);
+					PackFactoryHelper::ExtractFileToString(Entry, PakReader, CopyBuffer, PersistentCompressionBuffer, SourceContents);
 
 					FGameProjectGenerationModule& GameProjectModule = FModuleManager::LoadModuleChecked<FGameProjectGenerationModule>(TEXT("GameProjectGeneration"));
 					FString StringToReplace = *ConfigParameters.Find("GameName");
@@ -458,7 +458,7 @@ UObject* UPackFactory::FactoryCreateBinary
 
 					if (FileHandle.IsValid())
 					{
-						PackFactoryHelper::ExtractFile(Entry, PakReader, Buffer, PersistentCompressionBuffer, *FileHandle);
+						PackFactoryHelper::ExtractFile(Entry, PakReader, CopyBuffer, PersistentCompressionBuffer, *FileHandle);
 						WrittenFiles.Add(*DestFilename);
 					}
 					else
diff --git a/Engine/Source/Editor/UnrealEd/Private/FileHelpers.cpp b/Engine/Source/Editor/UnrealEd/Private/FileHelpers.cpp
--- a/Engine/Source/Editor/UnrealEd/Private/FileHelpers.cpp
+++ b/Engine/Source/Editor/UnrealEd/Private/FileHelpers.cpp
@@ -1606,7 +1606,7 @@ void FEditorFileUtils::OpenLevelPickingDialog(const FOnLevelsChosen& OnLevelsCho
 {
 	struct FLocal
 	{
-		static void OnLevelsSelected(const TArray<FAssetData>& SelectedLevels, FOnLevelsChosen OnLevelsChosen)
+		static void OnLevelsSelected(const TArray<FAssetData>& SelectedLevels, FOnLevelsChosen OnLevelsChosenDelegate)
 		{
 			if ( SelectedLevels.Num() > 0 )
 			{
@@ -1625,13 +1625,13 @@ void FEditorFileUtils::OpenLevelPickingDialog(const FOnLevelsChosen& OnLevelsCho
 
 				FEditorDirectories::Get().SetLastDirectory(ELastDirectory::LEVEL, FilesystemPath);
 
-				OnLevelsChosen.ExecuteIfBound(SelectedLevels);
+				OnLevelsChosenDelegate.ExecuteIfBound(SelectedLevels);
 			}
 		}
 
-		static void OnDialogCancelled(FOnLevelPickingCancelled OnLevelPickingCancelled)
+		static void OnDialogCancelled(FOnLevelPickingCancelled OnLevelPickingCancelledDelegate)
 		{
-			OnLevelPickingCancelled.ExecuteIfBound();
+			OnLevelPickingCancelledDelegate.ExecuteIfBound();
 		}
 	};
 
@@ -2326,16 +2326,16 @@ static int32 InternalSavePackage( UPackage* PackageToSave, bool& bOutPackageLoca
 				FString DefaultPackagePath;
 				FPackageName::TryConvertFilenameToLongPackageName(DefaultLocation, DefaultPackagePath);
 
-				FString PackageName;
+				FString SaveAsPackageName;
 				bSaveFile = OpenLevelSaveAsDialog(
 					DefaultPackagePath,
 					FPaths::GetBaseFilename(FinalPackageFilename),
-					PackageName);
+					SaveAsPackageName);
 
 				if (bSaveFile)
 				{
 					// Leave out the extension. It will be added below.
-					FinalPackageFilename = FPackageName::LongPackageNameToFilename(PackageName);
+					FinalPackageFilename = FPackageName::LongPackageNameToFilename(SaveAsPackageName);
 				}
 			}
 			else
diff --git a/Engine/Source/Editor/UnrealEd/Private/Kismet2/DebuggerCommands.cpp b/Engine/Source/Editor/UnrealEd/Private/Kismet2/DebuggerCommands.cpp
--- a/Engine/Source/Editor/UnrealEd/Private/Kismet2/DebuggerCommands.cpp
+++ b/Engine/Source/Editor/UnrealEd/Private/Kismet2/DebuggerCommands.cpp
@@ -1566,9 +1566,6 @@ bool FInternalPlayWorldCommandCallbacks::IsReadyToLaunchOnDevice(FString DeviceI
 	}
 	else
 	{
-		const PlatformInfo::FPlatformInfo* PlatformInfo = PlatformInfo::FindPlatformInfo(FName(*PlatformName));
-		check(PlatformInfo);
-
 		IMainFrameModule& MainFrameModule = FModuleManager::GetModuleChecked<IMainFrameModule>(TEXT("MainFrame"));
 		MainFrameModule.BroadcastMainFrameSDKNotInstalled(PlatformInfo->TargetPlatformName.ToString(), PlatformInfo->SDKTutorial);
 		return false;
diff --git a/Engine/Source/Editor/UnrealEd/Private/LevelEditorViewport.cpp b/Engine/Source/Editor/UnrealEd/Private/LevelEditorViewport.cpp
--- a/Engine/Source/Editor/UnrealEd/Private/LevelEditorViewport.cpp
+++ b/Engine/Source/Editor/UnrealEd/Private/LevelEditorViewport.cpp
@@ -2159,7 +2159,7 @@ void FLevelEditorViewportClient::ProjectActorsIntoWorld(const TArray<AActor*>& A
 	}
 }
 
-bool FLevelEditorViewportClient::InputWidgetDelta(FViewport* Viewport, EAxisList::Type CurrentAxis, FVector& Drag, FRotator& Rot, FVector& Scale)
+bool FLevelEditorViewportClient::InputWidgetDelta(FViewport* Viewport, EAxisList::Type InCurrentAxis, FVector& Drag, FRotator& Rot, FVector& Scale)
 {
 	if (GUnrealEd->ComponentVisManager.HandleInputDelta(this, Viewport, Drag, Rot, Scale))
 	{
@@ -2169,14 +2169,14 @@ bool FLevelEditorViewportClient::InputWidgetDelta(FViewport* Viewport, EAxisList
 	bool bHandled = false;
 
 	// Give the current editor mode a chance to use the input first.  If it does, don't apply it to anything else.
-	if (FEditorViewportClient::InputWidgetDelta(Viewport, CurrentAxis, Drag, Rot, Scale))
+	if (FEditorViewportClient::InputWidgetDelta(Viewport, InCurrentAxis, Drag, Rot, Scale))
 	{
 		bHandled = true;
 	}
 	else
 	{
 		//@TODO: MODETOOLS: Much of this needs to get pushed to Super, but not all of it can be...
-		if( CurrentAxis != EAxisList::None )
+		if( InCurrentAxis != EAxisList::None )
 		{
 			// Skip actors transformation routine in case if any of the selected actors locked
 			// but still pretend that we have handled the input
@@ -3006,9 +3006,9 @@ void FLevelEditorViewportClient::ApplyDeltaToActors(const FVector& InDrag,
 	TArray<AGroupActor*> ActorGroups;
 
 	// Apply the deltas to any selected actors.
-	for ( FSelectionIterator It( GEditor->GetSelectedActorIterator() ) ; It ; ++It )
+	for ( FSelectionIterator SelectedActorIt( GEditor->GetSelectedActorIterator() ) ; SelectedActorIt ; ++SelectedActorIt )
 	{
-		AActor* Actor = static_cast<AActor*>( *It );
+		AActor* Actor = static_cast<AActor*>( *SelectedActorIt );
 		checkSlow( Actor->IsA(AActor::StaticClass()) );
 
 		// Verify that the actor is in the same world as the viewport before moving it.
@@ -3032,17 +3032,17 @@ void FLevelEditorViewportClient::ApplyDeltaToActors(const FVector& InDrag,
 		{
 			if (GEditor->GetSelectedComponentCount() > 0)
 			{
-				auto ComponentSelection = GEditor->GetSelectedComponents();
+				USelection* ComponentSelection = GEditor->GetSelectedComponents();
 
 				// Only move the parent-most component(s) that are selected 
 				// Otherwise, if both a parent and child are selected and the delta is applied to both, the child will actually move 2x delta
 				TInlineComponentArray<USceneComponent*> ComponentsToMove;
-				for (FSelectedEditableComponentIterator It(GEditor->GetSelectedEditableComponentIterator()); It; ++It)
+				for (FSelectedEditableComponentIterator EditableComponentIt(GEditor->GetSelectedEditableComponentIterator()); EditableComponentIt; ++EditableComponentIt)
 				{
-					USceneComponent* SceneComponent = CastChecked<USceneComponent>(*It);
+					USceneComponent* SceneComponent = CastChecked<USceneComponent>(*EditableComponentIt);
 					if (SceneComponent)
 					{
-						auto SelectedComponent = Cast<USceneComponent>(*It);
+						USceneComponent* SelectedComponent = Cast<USceneComponent>(*EditableComponentIt);
 
 						// Check to see if any parent is selected
 						bool bParentAlsoSelected = false;
diff --git a/Engine/Source/Editor/UnrealEd/Private/ObjectTools.cpp b/Engine/Source/Editor/UnrealEd/Private/ObjectTools.cpp
--- a/Engine/Source/Editor/UnrealEd/Private/ObjectTools.cpp
+++ b/Engine/Source/Editor/UnrealEd/Private/ObjectTools.cpp
@@ -828,7 +828,6 @@ namespace ObjectTools
 					RedirectorToObjectNameMap.Add(Redirector, CurObjName);
 
 					// If consolidating blueprints, make sure redirectors are created for the consolidated blueprint class and CDO
-					UBlueprint* BlueprintToConsolidate = Cast<UBlueprint>(CurObjToConsolidate);
 					if ( BlueprintToConsolidateTo != NULL && BlueprintToConsolidate != NULL )
 					{
 						// One redirector for the class
diff --git a/Engine/Source/Runtime/AIModule/Classes/Actions/PawnAction.h b/Engine/Source/Runtime/AIModule/Classes/Actions/PawnAction.h
--- a/Engine/Source/Runtime/AIModule/Classes/Actions/PawnAction.h
+++ b/Engine/Source/Runtime/AIModule/Classes/Actions/PawnAction.h
@@ -188,7 +188,7 @@ public:
 	// @note this function will change its signature once AI messaging is rewritten @todo
 	virtual void HandleAIMessage(UBrainComponent*, const FAIMessage&){};
 
-	void SetActionObserver(const FPawnActionEventDelegate& ActionObserver) { this->ActionObserver = ActionObserver; }
+	void SetActionObserver(const FPawnActionEventDelegate& InActionObserver) { ActionObserver = InActionObserver; }
 	bool HasActionObserver() const { return ActionObserver.IsBound(); }
 
 	//----------------------------------------------------------------------//
diff --git a/Engine/Source/Runtime/AIModule/Private/EnvironmentQuery/Generators/EnvQueryGenerator_ProjectedPoints.cpp b/Engine/Source/Runtime/AIModule/Private/EnvironmentQuery/Generators/EnvQueryGenerator_ProjectedPoints.cpp
--- a/Engine/Source/Runtime/AIModule/Private/EnvironmentQuery/Generators/EnvQueryGenerator_ProjectedPoints.cpp
+++ b/Engine/Source/Runtime/AIModule/Private/EnvironmentQuery/Generators/EnvQueryGenerator_ProjectedPoints.cpp
@@ -31,10 +31,10 @@ void UEnvQueryGenerator_ProjectedPoints::ProjectAndFilterNavPoints(TArray<FVecto
 		}
 		else
 		{
-			const FVector VerticalOffset = FVector(0, 0, (ProjectionData.ProjectUp - ProjectionData.ProjectDown) / 2);
+			const FVector ProjectionVerticalOffset = FVector(0, 0, (ProjectionData.ProjectUp - ProjectionData.ProjectDown) / 2);
 			for (const auto& Point : Points)
 			{
-				Workload.Add(FNavigationProjectionWork(Point + VerticalOffset));
+				Workload.Add(FNavigationProjectionWork(Point + ProjectionVerticalOffset));
 			}
 		}
 
diff --git a/Engine/Source/Runtime/Core/Private/Mac/MacPlatformMisc.cpp b/Engine/Source/Runtime/Core/Private/Mac/MacPlatformMisc.cpp
--- a/Engine/Source/Runtime/Core/Private/Mac/MacPlatformMisc.cpp
+++ b/Engine/Source/Runtime/Core/Private/Mac/MacPlatformMisc.cpp
@@ -656,7 +656,7 @@ EAppReturnType::Type FMacPlatformMisc::MessageBoxExt(EAppMsgType::Type MsgType,
 {
 	SCOPED_AUTORELEASE_POOL;
 
-	EAppReturnType::Type RetValue = MainThreadReturn(^{
+	EAppReturnType::Type ReturnValue = MainThreadReturn(^{
 		EAppReturnType::Type RetValue = EAppReturnType::Cancel;
 		NSInteger Result;
 
@@ -819,7 +819,7 @@ EAppReturnType::Type FMacPlatformMisc::MessageBoxExt(EAppMsgType::Type MsgType,
 		return RetValue;
 	});
 
-	return RetValue;
+	return ReturnValue;
 }
 
 static bool HandleFirstInstall()
@@ -1801,7 +1801,6 @@ void FMacCrashContext::GenerateCrashInfoAndLaunchReporter() const
 			int VideoSrc = open(GMacAppInfo.CrashReportVideo, O_RDONLY);
 			int VideoDst = open(FilePath, O_CREAT|O_WRONLY, 0766);
 			
-			int Bytes = 0;
 			while((Bytes = read(VideoSrc, Data, PATH_MAX)) > 0)
 			{
 				write(VideoDst, Data, Bytes);
diff --git a/Engine/Source/Runtime/Core/Public/Mac/MacPlatformMisc.h b/Engine/Source/Runtime/Core/Public/Mac/MacPlatformMisc.h
--- a/Engine/Source/Runtime/Core/Public/Mac/MacPlatformMisc.h
+++ b/Engine/Source/Runtime/Core/Public/Mac/MacPlatformMisc.h
@@ -196,7 +196,9 @@ private:
 	NSAutoreleasePool*	Pool;
 };
 
-#define SCOPED_AUTORELEASE_POOL const FScopeAutoreleasePool Pool##__LINE__;
+#define SCOPED_AUTORELEASE_POOL_TOKENPASTE_INNER(x,y) x##y
+#define SCOPED_AUTORELEASE_POOL_TOKENPASTE(x,y) SCOPED_AUTORELEASE_POOL_TOKENPASTE_INNER(x,y)
+#define SCOPED_AUTORELEASE_POOL const FScopeAutoreleasePool SCOPED_AUTORELEASE_POOL_TOKENPASTE(Pool,__LINE__);
 
 #endif // __OBJC__
 
diff --git a/Engine/Source/Runtime/Engine/Private/AI/Navigation/RecastNavMesh.cpp b/Engine/Source/Runtime/Engine/Private/AI/Navigation/RecastNavMesh.cpp
--- a/Engine/Source/Runtime/Engine/Private/AI/Navigation/RecastNavMesh.cpp
+++ b/Engine/Source/Runtime/Engine/Private/AI/Navigation/RecastNavMesh.cpp
@@ -674,7 +674,7 @@ void ARecastNavMesh::RestrictBuildingToActiveTiles(bool InRestrictBuildingToActi
 	}
 }
 
-void ARecastNavMesh::SerializeRecastNavMesh(FArchive& Ar, FPImplRecastNavMesh*& NavMesh, int32 NavMeshVersion)
+void ARecastNavMesh::SerializeRecastNavMesh(FArchive& Ar, FPImplRecastNavMesh*& NavMesh, int32 InNavMeshVersion)
 {
 	if (!Ar.IsLoading()	&& NavMesh == NULL)
 	{
@@ -692,7 +692,7 @@ void ARecastNavMesh::SerializeRecastNavMesh(FArchive& Ar, FPImplRecastNavMesh*&
 	
 	if (RecastNavMeshImpl)
 	{
-		RecastNavMeshImpl->Serialize(Ar, NavMeshVersion);
+		RecastNavMeshImpl->Serialize(Ar, InNavMeshVersion);
 	}	
 }
 
diff --git a/Engine/Source/Runtime/Engine/Private/ComponentInstanceDataCache.cpp b/Engine/Source/Runtime/Engine/Private/ComponentInstanceDataCache.cpp
--- a/Engine/Source/Runtime/Engine/Private/ComponentInstanceDataCache.cpp
+++ b/Engine/Source/Runtime/Engine/Private/ComponentInstanceDataCache.cpp
@@ -96,11 +96,11 @@ void FActorComponentInstanceData::ApplyToComponent(UActorComponent* Component, c
 		class FComponentPropertyReader : public FObjectReader
 		{
 		public:
-			FComponentPropertyReader(UActorComponent* Component, TArray<uint8>& InBytes)
+			FComponentPropertyReader(UActorComponent* InComponent, TArray<uint8>& InBytes)
 				: FObjectReader(InBytes)
 			{
-				UClass* Class = Component->GetClass();
-				Class->SerializeTaggedProperties(*this, (uint8*)Component, Class, nullptr);
+				UClass* Class = InComponent->GetClass();
+				Class->SerializeTaggedProperties(*this, (uint8*)InComponent, Class, nullptr);
 			}
 		} ComponentPropertyReader(Component, SavedProperties);	
 	}
diff --git a/Engine/Source/Runtime/Engine/Private/VisualLogger/VisualLogger.cpp b/Engine/Source/Runtime/Engine/Private/VisualLogger/VisualLogger.cpp
--- a/Engine/Source/Runtime/Engine/Private/VisualLogger/VisualLogger.cpp
+++ b/Engine/Source/Runtime/Engine/Private/VisualLogger/VisualLogger.cpp
@@ -75,8 +75,7 @@ void FVisualLogger::Cleanup(bool bReleaseMemory)
 		SetIsRecordingToFile(false);
 	}
 
-	auto& OutputDevices = FVisualLogger::Get().OutputDevices;
-	for (auto* Device : OutputDevices)
+	for (FVisualLogDevice* Device : FVisualLogger::Get().OutputDevices)
 	{
 		Device->Cleanup(bReleaseMemory);
 	}
diff --git a/Engine/Source/Runtime/Landscape/Public/LandscapeRender.h b/Engine/Source/Runtime/Landscape/Public/LandscapeRender.h
--- a/Engine/Source/Runtime/Landscape/Public/LandscapeRender.h
+++ b/Engine/Source/Runtime/Landscape/Public/LandscapeRender.h
@@ -373,9 +373,9 @@ protected:
 			, Guid(InGuid)
 		{}
 
-		friend inline uint32 GetTypeHash(const FLandscapeKey& LandscapeKey)
+		friend inline uint32 GetTypeHash(const FLandscapeKey& InLandscapeKey)
 		{
-			return HashCombine(GetTypeHash(LandscapeKey.World), GetTypeHash(LandscapeKey.Guid));
+			return HashCombine(GetTypeHash(InLandscapeKey.World), GetTypeHash(InLandscapeKey.Guid));
 		}
 
 		friend bool operator==(const FLandscapeKey& A, const FLandscapeKey& B)
diff --git a/Engine/Source/Runtime/Navmesh/Private/DetourCrowd/DetourSharedBoundary.cpp b/Engine/Source/Runtime/Navmesh/Private/DetourCrowd/DetourSharedBoundary.cpp
--- a/Engine/Source/Runtime/Navmesh/Private/DetourCrowd/DetourSharedBoundary.cpp
+++ b/Engine/Source/Runtime/Navmesh/Private/DetourCrowd/DetourSharedBoundary.cpp
@@ -105,7 +105,7 @@ int32 dtSharedBoundary::CacheData(float* Center, float Radius, dtPolyRef CenterP
 	return DataIdx;
 }
 
-void dtSharedBoundary::FindEdges(dtSharedBoundaryData& Data, dtPolyRef CenterPoly, dtNavMeshQuery* NavQuery, dtQueryFilter* NavFilter)
+void dtSharedBoundary::FindEdges(dtSharedBoundaryData& SharedData, dtPolyRef CenterPoly, dtNavMeshQuery* NavQuery, dtQueryFilter* NavFilter)
 {
 	const int32 MaxWalls = 64;
 	int32 NumWalls = 0;
@@ -116,7 +116,7 @@ void dtSharedBoundary::FindEdges(dtSharedBoundaryData& Data, dtPolyRef CenterPol
 	int32 NumNeis = 0;
 	dtPolyRef NeiPolys[MaxNeis] = { 0 };
 
-	NavQuery->findWallsInNeighbourhood(CenterPoly, Data.Center, Data.Radius, NavFilter,
+	NavQuery->findWallsInNeighbourhood(CenterPoly, SharedData.Center, SharedData.Radius, NavFilter,
 		NeiPolys, &NumNeis, MaxNeis, WallSegments, WallPolys, &NumWalls, MaxWalls);
 	
 	dtSharedBoundaryEdge NewEdge;
@@ -127,13 +127,13 @@ void dtSharedBoundary::FindEdges(dtSharedBoundaryData& Data, dtPolyRef CenterPol
 		NewEdge.p0 = WallPolys[Idx * 2];
 		NewEdge.p1 = WallPolys[Idx * 2 + 1];
 
-		Data.Edges.Add(NewEdge);
+		SharedData.Edges.Add(NewEdge);
 	}
 
-	Data.Polys.Reserve(NumNeis);
+	SharedData.Polys.Reserve(NumNeis);
 	for (int32 Idx = 0; Idx < NumNeis; Idx++)
 	{
-		Data.Polys.Add(NeiPolys[Idx]);
+		SharedData.Polys.Add(NeiPolys[Idx]);
 	}
 }
 
diff --git a/Engine/Source/Runtime/Navmesh/Private/DetourTileCache/DetourTileCacheRegion.cpp b/Engine/Source/Runtime/Navmesh/Private/DetourTileCache/DetourTileCacheRegion.cpp
--- a/Engine/Source/Runtime/Navmesh/Private/DetourTileCache/DetourTileCacheRegion.cpp
+++ b/Engine/Source/Runtime/Navmesh/Private/DetourTileCache/DetourTileCacheRegion.cpp
@@ -622,13 +622,16 @@ static void walkContour(int x, int y, int i, int dir,
 	const int w = (int)layer.header->width;
 	const int h = (int)layer.header->height;
 
-	const int ax = x + getDirOffsetX(dir);
-	const int ay = y + getDirOffsetY(dir);
-	const int ai = ax+ay*w;
 	unsigned short curReg = 0;
-	if (ax >= 0 && ax < w && ay >= 0 && ay < h && isConnected(layer, i, dir))
+
 	{
-		curReg = srcReg[ai];
+		const int ax = x + getDirOffsetX(dir);
+		const int ay = y + getDirOffsetY(dir);
+		const int ai = ax+ay*w;
+		if (ax >= 0 && ax < w && ay >= 0 && ay < h && isConnected(layer, i, dir))
+		{
+			curReg = srcReg[ai];
+		}
 	}
 
 	cont.push(curReg);
diff --git a/Engine/Source/Runtime/Online/OnlineSubsystem/Private/OnlineKeyValuePair.cpp b/Engine/Source/Runtime/Online/OnlineSubsystem/Private/OnlineKeyValuePair.cpp
--- a/Engine/Source/Runtime/Online/OnlineSubsystem/Private/OnlineKeyValuePair.cpp
+++ b/Engine/Source/Runtime/Online/OnlineSubsystem/Private/OnlineKeyValuePair.cpp
@@ -531,55 +531,55 @@ bool FVariantData::FromJson(const TSharedRef<FJsonObject>& JsonObject)
 	{
 		if (VariantTypeStr.Equals(EOnlineKeyValuePairDataType::ToString(EOnlineKeyValuePairDataType::Int32)))
 		{
-			int32 Value;
-			if (JsonObject->TryGetNumberField(ValueStr, Value))
+			int32 FieldValue;
+			if (JsonObject->TryGetNumberField(ValueStr, FieldValue))
 			{
-				SetValue(Value);
+				SetValue(FieldValue);
 				bResult = true;
 			}
 		}
 		else if (VariantTypeStr.Equals(EOnlineKeyValuePairDataType::ToString(EOnlineKeyValuePairDataType::Float)))
 		{
-			double Value;
-			if (JsonObject->TryGetNumberField(ValueStr, Value))
+			double FieldValue;
+			if (JsonObject->TryGetNumberField(ValueStr, FieldValue))
 			{
-				SetValue((float)Value);
+				SetValue((float)FieldValue);
 				bResult = true;
 			}
 		}
 		else if (VariantTypeStr.Equals(EOnlineKeyValuePairDataType::ToString(EOnlineKeyValuePairDataType::String)))
 		{
-			FString Value;
-			if (JsonObject->TryGetStringField(ValueStr, Value))
+			FString FieldValue;
+			if (JsonObject->TryGetStringField(ValueStr, FieldValue))
 			{
-				SetValue(Value);
+				SetValue(FieldValue);
 				bResult = true;
 			}
 		}
 		else if (VariantTypeStr.Equals(EOnlineKeyValuePairDataType::ToString(EOnlineKeyValuePairDataType::Bool)))
 		{
-			bool Value;
-			if (JsonObject->TryGetBoolField(ValueStr, Value))
+			bool FieldValue;
+			if (JsonObject->TryGetBoolField(ValueStr, FieldValue))
 			{
-				SetValue(Value);
+				SetValue(FieldValue);
 				bResult = true;
 			}
 		}
 		else if (VariantTypeStr.Equals(EOnlineKeyValuePairDataType::ToString(EOnlineKeyValuePairDataType::Int64)))
 		{
-			FString Value;
-			if (JsonObject->TryGetStringField(ValueStr, Value))
+			FString FieldValue;
+			if (JsonObject->TryGetStringField(ValueStr, FieldValue))
 			{
 				Type = EOnlineKeyValuePairDataType::Int64;
-				bResult = FromString(Value);
+				bResult = FromString(FieldValue);
 			}
 		}
 		else if (VariantTypeStr.Equals(EOnlineKeyValuePairDataType::ToString(EOnlineKeyValuePairDataType::Double)))
 		{
-			double Value;
-			if (JsonObject->TryGetNumberField(ValueStr, Value))
+			double FieldValue;
+			if (JsonObject->TryGetNumberField(ValueStr, FieldValue))
 			{
-				SetValue(Value);
+				SetValue(FieldValue);
 				bResult = true;
 			}
 		}
