[UE-11369] Fix potential crash on select of existing "Two Bone IK" nodes in an animation graph.

- introduced in CL# 2443244 (Fix manipulation of transforms on bone control nodes so that selecting a node doesn't overwrite its conents)
- seems like existing (older) nodes could potentially have serialized an empty default value string for the vector location pins; fixed by checking for that and allowing it to fall back to the zero vector in that case

diff --git a/Engine/Source/Editor/AnimGraph/Private/AnimGraphNode_SkeletalControlBase.cpp b/Engine/Source/Editor/AnimGraph/Private/AnimGraphNode_SkeletalControlBase.cpp
--- a/Engine/Source/Editor/AnimGraph/Private/AnimGraphNode_SkeletalControlBase.cpp
+++ b/Engine/Source/Editor/AnimGraph/Private/AnimGraphNode_SkeletalControlBase.cpp
@@ -288,21 +288,26 @@ void UAnimGraphNode_SkeletalControlBase::GetDefaultValue(const FString& UpdateDe
 			if (GetSchema()->IsCurrentPinDefaultValid(Pin).IsEmpty())
 			{
 				FString DefaultString = Pin->GetDefaultAsString();
-				TArray<FString> ResultString;
 
-				//Parse string to split its contents separated by ','
-				DefaultString.Trim();
-				DefaultString.TrimTrailing();
-				DefaultString.ParseIntoArray(ResultString, TEXT(","), true);
+				// Existing nodes (from older versions) might have an empty default value string; in that case we just fall through and return the zero vector below (which is the default value in that case).
+				if(!DefaultString.IsEmpty())
+				{
+					TArray<FString> ResultString;
+
+					//Parse string to split its contents separated by ','
+					DefaultString.Trim();
+					DefaultString.TrimTrailing();
+					DefaultString.ParseIntoArray(ResultString, TEXT(","), true);
 
-				check(ResultString.Num() == 3);
+					check(ResultString.Num() == 3);
 
-				OutVec.Set(
-							FCString::Atof(*ResultString[0]),
-							FCString::Atof(*ResultString[1]),
-							FCString::Atof(*ResultString[2])
-							);
-				return;
+					OutVec.Set(
+						FCString::Atof(*ResultString[0]),
+						FCString::Atof(*ResultString[1]),
+						FCString::Atof(*ResultString[2])
+						);
+					return;
+				}
 			}
 		}
 	}
