Merging

//depot/UE4-Orion/Engine/Source/Editor/AnimGraph/...

to //depot/UE4/Engine/Source/Editor/AnimGraph/...

diff --git a/Engine/Source/Editor/AnimGraph/Classes/AnimGraphNode_RotationOffsetBlendSpace.h b/Engine/Source/Editor/AnimGraph/Classes/AnimGraphNode_RotationOffsetBlendSpace.h
--- a/Engine/Source/Editor/AnimGraph/Classes/AnimGraphNode_RotationOffsetBlendSpace.h
+++ b/Engine/Source/Editor/AnimGraph/Classes/AnimGraphNode_RotationOffsetBlendSpace.h
@@ -23,6 +23,8 @@ class UAnimGraphNode_RotationOffsetBlendSpace: public UAnimGraphNode_BlendSpaceB
 	virtual void GetMenuEntries(FGraphContextMenuBuilder& ContextMenuBuilder) const override;
 	virtual void ValidateAnimNodeDuringCompilation(class USkeleton* ForSkeleton, class FCompilerResultsLog& MessageLog) override;
 	virtual UAnimationAsset* GetAnimationAsset() const override { return Node.BlendSpace; }
+	virtual void GetAllAnimationSequencesReferred(TArray<UAnimationAsset*>& ComplexAnims, TArray<UAnimSequence*>& AnimationSequences) const override;
+	virtual void ReplaceReferredAnimations(const TMap<UAnimationAsset*, UAnimationAsset*>& ComplexAnimsMap, const TMap<UAnimSequence*, UAnimSequence*>& AnimSequenceMap) override;
 	// End of UAnimGraphNode_Base interface
 
 	// UK2Node interface
diff --git a/Engine/Source/Editor/AnimGraph/Private/AnimGraphNode_RotationOffsetBlendSpace.cpp b/Engine/Source/Editor/AnimGraph/Private/AnimGraphNode_RotationOffsetBlendSpace.cpp
--- a/Engine/Source/Editor/AnimGraph/Private/AnimGraphNode_RotationOffsetBlendSpace.cpp
+++ b/Engine/Source/Editor/AnimGraph/Private/AnimGraphNode_RotationOffsetBlendSpace.cpp
@@ -141,4 +141,16 @@ void UAnimGraphNode_RotationOffsetBlendSpace::GetContextMenuActions(const FGraph
 	}
 }
 
+void UAnimGraphNode_RotationOffsetBlendSpace::GetAllAnimationSequencesReferred(TArray<UAnimationAsset*>& ComplexAnims, TArray<UAnimSequence*>& AnimationSequences) const
+{
+	if(Node.BlendSpace)
+	{
+		HandleAnimReferenceCollection(Node.BlendSpace, ComplexAnims, AnimationSequences);
+	}
+}
+
+void UAnimGraphNode_RotationOffsetBlendSpace::ReplaceReferredAnimations(const TMap<UAnimationAsset*, UAnimationAsset*>& ComplexAnimsMap, const TMap<UAnimSequence*, UAnimSequence*>& AnimSequenceMap)
+{
+	HandleAnimReferenceReplacement(Node.BlendSpace, ComplexAnimsMap, AnimSequenceMap);
+}
 #undef LOCTEXT_NAMESPACE
\ No newline at end of file
