Fix step forwards/backwards not snapping to frames in montages.

diff --git a/Engine/Source/Editor/AnimGraph/Private/AnimPreviewInstance.cpp b/Engine/Source/Editor/AnimGraph/Private/AnimPreviewInstance.cpp
--- a/Engine/Source/Editor/AnimGraph/Private/AnimPreviewInstance.cpp
+++ b/Engine/Source/Editor/AnimGraph/Private/AnimPreviewInstance.cpp
@@ -536,8 +536,15 @@ void UAnimPreviewInstance::MontagePreview_StepForward()
 		MontagePreview_SetPlaying(true);
 
 		// Advance a single frame, leaving it paused afterwards
+		int32 NumFrames = Montage->GetNumberOfFrames();
+		// Add DELTA to prefer next frame when we're close to the boundary
+		float CurrentFraction = CurrentTime / Montage->SequenceLength + DELTA;
+		float NextFrame = FMath::Clamp<float>(FMath::FloorToFloat(CurrentFraction * NumFrames) + 1.0f, 0, NumFrames);
+		float NewTime = Montage->SequenceLength * (NextFrame / NumFrames);
+
 		GetSkelMeshComponent()->GlobalAnimRateScale = 1.0f;
-		GetSkelMeshComponent()->TickAnimation(MontagePreview_CalculateStepLength());
+		GetSkelMeshComponent()->TickAnimation(NewTime - CurrentTime);
+
 		MontagePreview_SetPlaying(false);
 	}
 }
@@ -584,8 +591,15 @@ void UAnimPreviewInstance::MontagePreview_StepBackward()
 		MontagePreview_SetPlaying(true);
 
 		// Advance a single frame, leaving it paused afterwards
+		int32 NumFrames = Montage->GetNumberOfFrames();
+		// Add DELTA to prefer next frame when we're close to the boundary
+		float CurrentFraction = CurrentTime / Montage->SequenceLength + DELTA;
+		float NextFrame = FMath::Clamp<float>(FMath::FloorToFloat(CurrentFraction * NumFrames) - 1.0f, 0, NumFrames);
+		float NewTime = Montage->SequenceLength * (NextFrame / NumFrames);
+
 		GetSkelMeshComponent()->GlobalAnimRateScale = 1.0f;
-		GetSkelMeshComponent()->TickAnimation(MontagePreview_CalculateStepLength());
+		GetSkelMeshComponent()->TickAnimation(FMath::Abs(NewTime - CurrentTime));
+
 		MontagePreview_SetPlaying(false);
 	}
 }
