Animation: DLL export improvements for some skeletal control classes
- Changed some classes from MinimalAPI to full export, since all methods were already being exported
- Added some comments
- Some style tweaks (tabs/ws in function declarations)

diff --git a/Engine/Source/Editor/AnimGraph/Classes/AnimGraphNode_BoneDrivenController.h b/Engine/Source/Editor/AnimGraph/Classes/AnimGraphNode_BoneDrivenController.h
--- a/Engine/Source/Editor/AnimGraph/Classes/AnimGraphNode_BoneDrivenController.h
+++ b/Engine/Source/Editor/AnimGraph/Classes/AnimGraphNode_BoneDrivenController.h
@@ -7,8 +7,11 @@
 #include "EdGraph/EdGraphNodeUtils.h" // for FNodeTitleTextTable
 #include "AnimGraphNode_BoneDrivenController.generated.h"
 
-UCLASS(MinimalAPI)
-class UAnimGraphNode_BoneDrivenController : public UAnimGraphNode_SkeletalControlBase
+/**
+ * This is the 'source version' of of a bone driven controller, which copies part of the state from one bone to another
+ */
+UCLASS()
+class ANIMGRAPH_API UAnimGraphNode_BoneDrivenController : public UAnimGraphNode_SkeletalControlBase
 {
 	GENERATED_UCLASS_BODY()
 
@@ -19,19 +22,18 @@ public:
 
 	// UEdGraphNode interface
 	virtual FText GetNodeTitle(ENodeTitleType::Type TitleType) const override;
-	//virtual FString GetNodeNativeTitle(ENodeTitleType::Type TitleType) const override;
 	virtual FText GetTooltipText() const override;
-	//////////////////////////////////////////////////////////////////////////
+	// End of UEdGraphNode interface
 
 	// UAnimGraphNode_SkeletalControlBase interface
-	ANIMGRAPH_API virtual void Draw(FPrimitiveDrawInterface* PDI, USkeletalMeshComponent* SkelMeshComp) const override;
-	//////////////////////////////////////////////////////////////////////////
+	virtual void Draw(FPrimitiveDrawInterface* PDI, USkeletalMeshComponent* SkelMeshComp) const override;
+	// End of UAnimGraphNode_SkeletalControlBase interface
 
 protected:
 
-	// UAnimGraphNode_SkeletalControlBase interface
+	// UAnimGraphNode_SkeletalControlBase protected interface
 	virtual FText GetControllerDescription() const override;
-	//////////////////////////////////////////////////////////////////////////
+	// End of UAnimGraphNode_SkeletalControlBase protected interface
 
 private:
 	/** Constructing FText strings can be costly, so we cache the node's title */
diff --git a/Engine/Source/Editor/AnimGraph/Classes/AnimGraphNode_ModifyBone.h b/Engine/Source/Editor/AnimGraph/Classes/AnimGraphNode_ModifyBone.h
--- a/Engine/Source/Editor/AnimGraph/Classes/AnimGraphNode_ModifyBone.h
+++ b/Engine/Source/Editor/AnimGraph/Classes/AnimGraphNode_ModifyBone.h
@@ -7,8 +7,8 @@
 #include "EdGraph/EdGraphNodeUtils.h" // for FNodeTitleTextTable
 #include "AnimGraphNode_ModifyBone.generated.h"
 
-UCLASS(MinimalAPI, meta=(Keywords = "Modify Transform"))
-class UAnimGraphNode_ModifyBone : public UAnimGraphNode_SkeletalControlBase
+UCLASS(meta=(Keywords = "Modify Transform"))
+class ANIMGRAPH_API UAnimGraphNode_ModifyBone : public UAnimGraphNode_SkeletalControlBase
 {
 	GENERATED_UCLASS_BODY()
 
@@ -24,20 +24,16 @@ public:
 protected:
 	// UAnimGraphNode_SkeletalControlBase interface
 	virtual FText GetControllerDescription() const override;
-
-	ANIMGRAPH_API virtual FVector GetWidgetLocation(const USkeletalMeshComponent* SkelComp, struct FAnimNode_SkeletalControlBase* AnimNode) override;
-	ANIMGRAPH_API virtual int32 GetWidgetMode(const USkeletalMeshComponent* SkelComp) override;
-	ANIMGRAPH_API virtual int32 ChangeToNextWidgetMode(const USkeletalMeshComponent* SkelComp, int32 InCurWidgetMode) override;
-	ANIMGRAPH_API virtual bool SetWidgetMode(const USkeletalMeshComponent* SkelComp, int32 InWidgetMode) override;
-
-	ANIMGRAPH_API virtual FName FindSelectedBone() override;
-
-	ANIMGRAPH_API virtual void DoTranslation(const USkeletalMeshComponent* SkelComp, FVector& Drag, FAnimNode_Base* InOutAnimNode) override;
-	ANIMGRAPH_API virtual void DoRotation(const USkeletalMeshComponent* SkelComp, FRotator& Rotation, FAnimNode_Base* InOutAnimNode) override;
-	ANIMGRAPH_API virtual void DoScale(const USkeletalMeshComponent* SkelComp, FVector& Drag, FAnimNode_Base* InOutAnimNode) override;
-	ANIMGRAPH_API virtual void	CopyNodeDataTo(FAnimNode_Base* OutAnimNode) override;
-	ANIMGRAPH_API virtual void	CopyNodeDataFrom(const FAnimNode_Base* InNewAnimNode) override;
-
+	virtual FVector GetWidgetLocation(const USkeletalMeshComponent* SkelComp, struct FAnimNode_SkeletalControlBase* AnimNode) override;
+	virtual int32 GetWidgetMode(const USkeletalMeshComponent* SkelComp) override;
+	virtual int32 ChangeToNextWidgetMode(const USkeletalMeshComponent* SkelComp, int32 InCurWidgetMode) override;
+	virtual bool SetWidgetMode(const USkeletalMeshComponent* SkelComp, int32 InWidgetMode) override;
+	virtual FName FindSelectedBone() override;
+	virtual void DoTranslation(const USkeletalMeshComponent* SkelComp, FVector& Drag, FAnimNode_Base* InOutAnimNode) override;
+	virtual void DoRotation(const USkeletalMeshComponent* SkelComp, FRotator& Rotation, FAnimNode_Base* InOutAnimNode) override;
+	virtual void DoScale(const USkeletalMeshComponent* SkelComp, FVector& Drag, FAnimNode_Base* InOutAnimNode) override;
+	virtual void CopyNodeDataTo(FAnimNode_Base* OutAnimNode) override;
+	virtual void CopyNodeDataFrom(const FAnimNode_Base* InNewAnimNode) override;
 	// End of UAnimGraphNode_SkeletalControlBase interface
 
 	// methods to find a valid widget mode for gizmo because doesn't need to show gizmo when the mode is "Ignore"
diff --git a/Engine/Source/Editor/AnimGraph/Classes/AnimGraphNode_SkeletalControlBase.h b/Engine/Source/Editor/AnimGraph/Classes/AnimGraphNode_SkeletalControlBase.h
--- a/Engine/Source/Editor/AnimGraph/Classes/AnimGraphNode_SkeletalControlBase.h
+++ b/Engine/Source/Editor/AnimGraph/Classes/AnimGraphNode_SkeletalControlBase.h
@@ -7,67 +7,73 @@
 
 struct HActor;
 
-UCLASS(Abstract, MinimalAPI)
-class UAnimGraphNode_SkeletalControlBase : public UAnimGraphNode_Base
+/**
+ * This is the base class for the 'source version' of all skeletal control animation graph nodes
+ * (nodes that manipulate the pose rather than playing animations to create a pose or blending between poses)
+ *
+ * Concrete subclasses should contain a member struct derived from FAnimNode_SkeletalControlBase
+ */
+UCLASS(Abstract)
+class ANIMGRAPH_API UAnimGraphNode_SkeletalControlBase : public UAnimGraphNode_Base
 {
 	GENERATED_UCLASS_BODY()
 
 public:
 	// UEdGraphNode interface
-	ANIMGRAPH_API virtual FLinearColor GetNodeTitleColor() const override;
-	ANIMGRAPH_API virtual FText GetTooltipText() const override;
+	virtual FLinearColor GetNodeTitleColor() const override;
+	virtual FText GetTooltipText() const override;
 	// End of UEdGraphNode interface
 
 	// UAnimGraphNode_Base interface
-	ANIMGRAPH_API virtual FString GetNodeCategory() const override;
-	ANIMGRAPH_API virtual void CreateOutputPins() override;
+	virtual FString GetNodeCategory() const override;
+	virtual void CreateOutputPins() override;
 	// End of UAnimGraphNode_Base interface
 
 	// Draw function for supporting visualization
-	ANIMGRAPH_API virtual void Draw(FPrimitiveDrawInterface* PDI, USkeletalMeshComponent * PreviewSkelMeshComp) const {};
+	virtual void Draw(FPrimitiveDrawInterface* PDI, USkeletalMeshComponent * PreviewSkelMeshComp) const {}
 
 	/**
 	 * methods related to widget control
 	 */
-	ANIMGRAPH_API virtual FVector GetWidgetLocation(const USkeletalMeshComponent* SkelComp, struct FAnimNode_SkeletalControlBase* AnimNode)
+	virtual FVector GetWidgetLocation(const USkeletalMeshComponent* SkelComp, struct FAnimNode_SkeletalControlBase* AnimNode)
 	{
 		return FVector::ZeroVector;
 	}
 	// to keep data consistency between anim nodes 
-	ANIMGRAPH_API virtual void	CopyNodeDataTo(FAnimNode_Base* OutAnimNode){}
-	ANIMGRAPH_API virtual void	CopyNodeDataFrom(const FAnimNode_Base* NewAnimNode){}
+	virtual void CopyNodeDataTo(FAnimNode_Base* OutAnimNode){}
+	virtual void CopyNodeDataFrom(const FAnimNode_Base* NewAnimNode){}
 
 	/** Are we currently showing this pin */
-	ANIMGRAPH_API bool IsPinShown(const FString& PinName) const;
+	bool IsPinShown(const FString& PinName) const;
 
 	// return current widget mode this anim graph node supports
-	ANIMGRAPH_API virtual int32 GetWidgetMode(const USkeletalMeshComponent* SkelComp);
+	virtual int32 GetWidgetMode(const USkeletalMeshComponent* SkelComp);
 	// called when the user changed widget mode by pressing "Space" key
-	ANIMGRAPH_API virtual int32 ChangeToNextWidgetMode(const USkeletalMeshComponent* SkelComp, int32 CurWidgetMode);
+	virtual int32 ChangeToNextWidgetMode(const USkeletalMeshComponent* SkelComp, int32 CurWidgetMode);
 	// called when the user set widget mode directly, returns true if InWidgetMode is available
-	ANIMGRAPH_API virtual bool SetWidgetMode(const USkeletalMeshComponent* SkelComp, int32 InWidgetMode){ return false; }
+	virtual bool SetWidgetMode(const USkeletalMeshComponent* SkelComp, int32 InWidgetMode) { return false; }
 
 	// 
-	ANIMGRAPH_API virtual FName FindSelectedBone();
+	virtual FName FindSelectedBone();
 
 	// if anim graph node needs other actors to select other bones, move actor's positions when this is called
-	ANIMGRAPH_API virtual void MoveSelectActorLocation(const USkeletalMeshComponent* SkelComp, FAnimNode_SkeletalControlBase* InAnimNode){}
+	virtual void MoveSelectActorLocation(const USkeletalMeshComponent* SkelComp, FAnimNode_SkeletalControlBase* InAnimNode) {}
 
-	ANIMGRAPH_API virtual bool IsActorClicked(HActor* ActorHitProxy){ return false; }
-	ANIMGRAPH_API virtual void ProcessActorClick(HActor* ActorHitProxy){}
+	virtual bool IsActorClicked(HActor* ActorHitProxy) { return false; }
+	virtual void ProcessActorClick(HActor* ActorHitProxy) {}
 	// if it has select-actors, should hide all actors when de-select is called  
-	ANIMGRAPH_API virtual void	DeselectActor(USkeletalMeshComponent* SkelComp){}
+	virtual void DeselectActor(USkeletalMeshComponent* SkelComp){}
 
 	// called when the widget is dragged in translation mode
-	ANIMGRAPH_API virtual void DoTranslation(const USkeletalMeshComponent* SkelComp, FVector& Drag, FAnimNode_Base* InOutAnimNode){}
+	virtual void DoTranslation(const USkeletalMeshComponent* SkelComp, FVector& Drag, FAnimNode_Base* InOutAnimNode) {}
 	// called when the widget is dragged in rotation mode
-	ANIMGRAPH_API virtual void DoRotation(const USkeletalMeshComponent* SkelComp, FRotator& Rotation, FAnimNode_Base* InOutAnimNode){}
+	virtual void DoRotation(const USkeletalMeshComponent* SkelComp, FRotator& Rotation, FAnimNode_Base* InOutAnimNode) {}
 	// called when the widget is dragged in scale mode
-	ANIMGRAPH_API virtual void DoScale(const USkeletalMeshComponent* SkelComp, FVector& Scale, FAnimNode_Base* InOutAnimNode){}
+	virtual void DoScale(const USkeletalMeshComponent* SkelComp, FVector& Scale, FAnimNode_Base* InOutAnimNode) {}
 
 protected:
 	// Returns the short descriptive name of the controller
-	ANIMGRAPH_API virtual FText GetControllerDescription() const;
+	virtual FText GetControllerDescription() const;
 
 	/**
 	* helper functions for bone control preview
