Anim transition getter refactor

diff --git a/Engine/Source/Editor/AnimGraph/Classes/AnimGraphNode_AssetPlayerBase.cpp b/Engine/Source/Editor/AnimGraph/Classes/AnimGraphNode_AssetPlayerBase.cpp
new file mode 100644
--- /dev/null
+++ b/Engine/Source/Editor/AnimGraph/Classes/AnimGraphNode_AssetPlayerBase.cpp
@@ -0,0 +1,4 @@
+
+#include "AnimGraphPrivatePCH.h"
+#include "AnimGraphNode_AssetPlayerBase.h"
+
diff --git a/Engine/Source/Editor/AnimGraph/Classes/AnimGraphNode_AssetPlayerBase.h b/Engine/Source/Editor/AnimGraph/Classes/AnimGraphNode_AssetPlayerBase.h
new file mode 100644
--- /dev/null
+++ b/Engine/Source/Editor/AnimGraph/Classes/AnimGraphNode_AssetPlayerBase.h
@@ -0,0 +1,15 @@
+// Copyright 1998-2015 Epic Games, Inc. All Rights Reserved.
+
+#pragma once
+
+#include "AnimGraphNode_Base.h"
+#include "AnimGraphNode_AssetPlayerBase.generated.h"
+
+class UAnimSequenceBase;
+
+/** Helper / intermediate for asset player graphical nodes */
+UCLASS(Abstract, MinimalAPI)
+class UAnimGraphNode_AssetPlayerBase : public UAnimGraphNode_Base
+{
+	GENERATED_BODY()
+};
\ No newline at end of file
diff --git a/Engine/Source/Editor/AnimGraph/Classes/AnimGraphNode_BlendSpaceBase.h b/Engine/Source/Editor/AnimGraph/Classes/AnimGraphNode_BlendSpaceBase.h
--- a/Engine/Source/Editor/AnimGraph/Classes/AnimGraphNode_BlendSpaceBase.h
+++ b/Engine/Source/Editor/AnimGraph/Classes/AnimGraphNode_BlendSpaceBase.h
@@ -3,10 +3,11 @@
 #pragma once
 #include "AnimGraphNode_Base.h"
 #include "Animation/BlendSpaceBase.h"
+#include "AnimGraphNode_AssetPlayerBase.h"
 #include "AnimGraphNode_BlendSpaceBase.generated.h"
 
 UCLASS(Abstract, MinimalAPI)
-class UAnimGraphNode_BlendSpaceBase : public UAnimGraphNode_Base
+class UAnimGraphNode_BlendSpaceBase : public UAnimGraphNode_AssetPlayerBase
 {
 	GENERATED_UCLASS_BODY()
 
diff --git a/Engine/Source/Editor/AnimGraph/Classes/AnimGraphNode_SequenceEvaluator.h b/Engine/Source/Editor/AnimGraph/Classes/AnimGraphNode_SequenceEvaluator.h
--- a/Engine/Source/Editor/AnimGraph/Classes/AnimGraphNode_SequenceEvaluator.h
+++ b/Engine/Source/Editor/AnimGraph/Classes/AnimGraphNode_SequenceEvaluator.h
@@ -2,12 +2,13 @@
 
 #pragma once
 #include "AnimGraphNode_Base.h"
+#include "AnimGraphNode_AssetPlayerBase.h"
 #include "Animation/AnimNode_SequenceEvaluator.h"
 #include "EdGraph/EdGraphNodeUtils.h" // for FNodeTextCache
 #include "AnimGraphNode_SequenceEvaluator.generated.h"
 
 UCLASS(MinimalAPI)
-class UAnimGraphNode_SequenceEvaluator : public UAnimGraphNode_Base
+class UAnimGraphNode_SequenceEvaluator : public UAnimGraphNode_AssetPlayerBase
 {
 	GENERATED_UCLASS_BODY()
 
diff --git a/Engine/Source/Editor/AnimGraph/Classes/AnimGraphNode_SequencePlayer.h b/Engine/Source/Editor/AnimGraph/Classes/AnimGraphNode_SequencePlayer.h
--- a/Engine/Source/Editor/AnimGraph/Classes/AnimGraphNode_SequencePlayer.h
+++ b/Engine/Source/Editor/AnimGraph/Classes/AnimGraphNode_SequencePlayer.h
@@ -4,10 +4,11 @@
 #include "AnimGraphNode_Base.h"
 #include "Animation/AnimNode_SequencePlayer.h"
 #include "EdGraph/EdGraphNodeUtils.h" // for FNodeTitleTextTable
+#include "AnimGraphNode_AssetPlayerBase.h"
 #include "AnimGraphNode_SequencePlayer.generated.h"
 
 UCLASS(MinimalAPI)
-class UAnimGraphNode_SequencePlayer : public UAnimGraphNode_Base
+class UAnimGraphNode_SequencePlayer : public UAnimGraphNode_AssetPlayerBase
 {
 	GENERATED_UCLASS_BODY()
 
diff --git a/Engine/Source/Editor/AnimGraph/Classes/K2Node_AnimGetter.h b/Engine/Source/Editor/AnimGraph/Classes/K2Node_AnimGetter.h
new file mode 100644
--- /dev/null
+++ b/Engine/Source/Editor/AnimGraph/Classes/K2Node_AnimGetter.h
@@ -0,0 +1,101 @@
+// Copyright 1998-2015 Epic Games, Inc. All Rights Reserved.
+
+#pragma once
+
+#include "K2Node_CallFunction.h"
+#include "K2Node_AnimGetter.generated.h"
+
+class UAnimGraphNode_Base;
+class UAnimStateNodeBase;
+
+USTRUCT()
+struct FNodeSpawnData
+{
+	GENERATED_BODY()
+
+	FNodeSpawnData();
+
+	// Title to use for the spawned node
+	UPROPERTY()
+	FText CachedTitle;
+
+	// The node the spawned getter accesses
+	UPROPERTY()
+	UAnimGraphNode_Base* SourceNode;
+
+	// The state node the spawned getter accesses
+	UPROPERTY()
+	UAnimStateNodeBase* SourceStateNode;
+
+	// The instance class the spawned getter is defined on
+	UPROPERTY()
+	UClass* AnimInstanceClass;
+
+	// The blueprint the getter is valid within
+	UPROPERTY()
+	const UAnimBlueprint* SourceBlueprint;
+
+	// The UFunction (as a UField) 
+	UPROPERTY()
+	UField* Getter;
+
+	// String of combined valid contexts for the spawned getter
+	UPROPERTY()
+	FString GetterContextString;
+};
+
+UCLASS(MinimalAPI)
+class UK2Node_AnimGetter : public UK2Node_CallFunction
+{
+	GENERATED_BODY()
+public:
+
+	// UEdGraphNode interface
+	virtual void AllocateDefaultPins() override;
+	virtual FText GetNodeTitle(ENodeTitleType::Type TitleType) const override;
+	virtual bool CanCreateUnderSpecifiedSchema(const UEdGraphSchema* Schema) const override;
+	virtual bool IsActionFilteredOut(FBlueprintActionFilter const& Filter) override;
+	// End of UEdGraphNode interface
+
+	// UK2Node interface
+	virtual void GetMenuActions(FBlueprintActionDatabaseRegistrar& ActionRegistrar) const override;
+	// end of UK2Node interface
+
+	// The node that is required for the getter
+	UPROPERTY()
+	UAnimGraphNode_Base* SourceNode;
+
+	// UAnimStateNode doesn't use the same hierarchy so we need to have a seperate property here to handle
+	// those.
+	UPROPERTY()
+	UAnimStateNodeBase* SourceStateNode;
+
+	// The UAnimInstance derived class that implements the getter we are running
+	UPROPERTY()
+	UClass* GetterClass;
+
+	// The anim blueprint that generated this getter
+	UPROPERTY()
+	const UAnimBlueprint* SourceAnimBlueprint;
+
+	// Cached node title
+	UPROPERTY()
+	FText CachedTitle;
+
+	// List of valid contexts for the node
+	UPROPERTY()
+	TArray<FString> Contexts;
+
+protected:
+	
+	//UFunction* GetSourceBlueprintFunction() const;
+
+	/** Returns whether or not the provided UFunction requires the named parameter */
+	bool GetterRequiresParameter(const UFunction* Getter, FString ParamName) const;
+
+	/** Checks the cached context strings to make sure this getter is valid within the provided schema */
+	bool IsContextValidForSchema(const UEdGraphSchema* Schema) const;
+
+	/** Passed to blueprint spawners to configure spawned nodes */
+	void PostSpawnNodeSetup(UEdGraphNode* NewNode, bool bIsTemplateNode, FNodeSpawnData SpawnData);
+};
\ No newline at end of file
diff --git a/Engine/Source/Editor/AnimGraph/Classes/K2Node_TransitionRuleGetter.h b/Engine/Source/Editor/AnimGraph/Classes/K2Node_TransitionRuleGetter.h
--- a/Engine/Source/Editor/AnimGraph/Classes/K2Node_TransitionRuleGetter.h
+++ b/Engine/Source/Editor/AnimGraph/Classes/K2Node_TransitionRuleGetter.h
@@ -4,6 +4,9 @@
 #pragma once
 #include "K2Node_TransitionRuleGetter.generated.h"
 
+class UAnimGraphNode_Base;
+class FBlueprintActionFilter;
+
 UENUM()
 namespace ETransitionGetter
 {
@@ -31,7 +34,7 @@ class UK2Node_TransitionRuleGetter : public UK2Node
 	TEnumAsByte<ETransitionGetter::Type> GetterType;
 
 	UPROPERTY()
-	class UAnimGraphNode_Base* AssociatedAnimAssetPlayerNode;
+	UAnimGraphNode_Base* AssociatedAnimAssetPlayerNode;
 
 	UPROPERTY()
 	UAnimStateNode* AssociatedStateNode;
@@ -42,7 +45,7 @@ class UK2Node_TransitionRuleGetter : public UK2Node
 	virtual FText GetTooltipText() const override;
 	virtual bool ShowPaletteIconOnNode() const override { return false; }
 	virtual bool CanCreateUnderSpecifiedSchema(const UEdGraphSchema* Schema) const override;
-	virtual bool IsActionFilteredOut(class FBlueprintActionFilter const& Filter) override;
+	virtual bool IsActionFilteredOut(FBlueprintActionFilter const& Filter) override;
 	// End of UEdGraphNode interface
 
 	// UK2Node interface
diff --git a/Engine/Source/Editor/AnimGraph/Private/K2Node_AnimGetter.cpp b/Engine/Source/Editor/AnimGraph/Private/K2Node_AnimGetter.cpp
new file mode 100644
--- /dev/null
+++ b/Engine/Source/Editor/AnimGraph/Private/K2Node_AnimGetter.cpp
@@ -0,0 +1,398 @@
+
+#include "AnimGraphPrivatePCH.h"
+#include "AnimationGraphSchema.h"
+#include "AnimationTransitionSchema.h"
+#include "BlueprintActionDatabaseRegistrar.h"
+#include "BlueprintActionFilter.h"
+#include "BlueprintNodeSpawner.h"
+#include "K2Node_AnimGetter.h"
+#include "AnimGraphNode_AssetPlayerBase.h"
+#include "AnimGraphNode_StateMachine.h"
+#include "AnimStateNode.h"
+#include "AnimationStateMachineGraph.h"
+#include "AnimationCustomTransitionSchema.h"
+#include "AnimStateTransitionNode.h"
+
+#define LOCTEXT_NAMESPACE "AnimGetter"
+
+void UK2Node_AnimGetter::AllocateDefaultPins()
+{
+	Super::AllocateDefaultPins();
+
+	TArray<UEdGraphPin*> PinsToHide;
+	TArray<FString> PinNames;
+
+	// TODO: Find a nicer way to maybe pull these down from the instance class and allow 
+	// projects to add new parameters from derived instances
+	PinNames.Add(TEXT("CurrentTime"));
+	PinNames.Add(TEXT("AssetPlayerIndex"));
+	PinNames.Add(TEXT("MachineIndex"));
+	PinNames.Add(TEXT("StateIndex"));
+	PinNames.Add(TEXT("TransitionIndex"));
+
+	for(FString& PinName : PinNames)
+	{
+		if(UEdGraphPin* FoundPin = FindPin(PinName))
+		{
+			PinsToHide.Add(FoundPin);
+		}
+	}
+
+	for(UEdGraphPin* Pin : PinsToHide)
+	{
+		Pin->bHidden = true;
+	}
+}
+
+FText UK2Node_AnimGetter::GetNodeTitle(ENodeTitleType::Type TitleType) const
+{
+	return CachedTitle;
+}
+
+bool UK2Node_AnimGetter::CanCreateUnderSpecifiedSchema(const UEdGraphSchema* Schema) const
+{
+	return Cast<UAnimationGraphSchema>(Schema) != NULL || Cast<UAnimationTransitionSchema>(Schema) != NULL;
+}
+
+void UK2Node_AnimGetter::GetMenuActions(FBlueprintActionDatabaseRegistrar& ActionRegistrar) const
+{
+	// First cache the available functions for getters
+	UClass* ActionKey = GetClass();
+	const UAnimBlueprint* AnimBlueprint = Cast<UAnimBlueprint>(ActionRegistrar.GetActionKeyFilter());
+	if(AnimBlueprint && ActionRegistrar.IsOpenForRegistration(AnimBlueprint))
+	{
+		UClass* BPClass = (AnimBlueprint) ? *AnimBlueprint->ParentClass : UAnimInstance::StaticClass();
+		while(BPClass && !BPClass->HasAnyClassFlags(CLASS_Native))
+		{
+			BPClass = BPClass->GetSuperClass();
+		}
+
+		if(BPClass)
+		{
+			TArray<UFunction*> AnimGetters;
+			for(TFieldIterator<UFunction> FuncIter(BPClass) ; FuncIter ; ++FuncIter)
+			{
+				UFunction* Func = *FuncIter;
+
+				if(Func->HasMetaData(TEXT("AnimGetter")) && Func->HasAnyFunctionFlags(FUNC_Native))
+				{
+					AnimGetters.Add(Func);
+				}
+			}
+
+			auto UiSpecOverride = [](const FBlueprintActionContext& /*Context*/, const IBlueprintNodeBinder::FBindingSet& Bindings, FBlueprintActionUiSpec* UiSpecOut, FText Title)
+			{
+				UiSpecOut->MenuName = Title;
+			};
+
+			TArray<UAnimGraphNode_AssetPlayerBase*> AssetPlayerNodes;
+			TArray<UAnimGraphNode_StateMachine*> MachineNodes;
+			TArray<UAnimStateNode*> StateNodes;
+			TArray<UAnimStateTransitionNode*> TransitionNodes;
+			
+			FBlueprintEditorUtils::GetAllNodesOfClass(AnimBlueprint, AssetPlayerNodes);
+			FBlueprintEditorUtils::GetAllNodesOfClass(AnimBlueprint, MachineNodes);
+			FBlueprintEditorUtils::GetAllNodesOfClass(AnimBlueprint, StateNodes);
+			FBlueprintEditorUtils::GetAllNodesOfClass(AnimBlueprint, TransitionNodes);
+
+			for(UFunction* Getter : AnimGetters)
+			{
+				FNodeSpawnData Params;
+				Params.AnimInstanceClass = BPClass;
+				Params.Getter = Getter;
+				Params.SourceBlueprint = AnimBlueprint;
+				Params.GetterContextString = Getter->GetMetaData(TEXT("GetterContext"));
+
+				if(GetterRequiresParameter(Getter, TEXT("AssetPlayerIndex")))
+				{
+					for(UAnimGraphNode_Base* AssetNode : AssetPlayerNodes)
+					{
+						// Should always succeed
+						if(UAnimationAsset* NodeAsset = AssetNode->GetAnimationAsset())
+						{
+							FText Title = FText::Format(LOCTEXT("NodeTitle", "{0} ({1})"), Getter->GetDisplayNameText(), FText::FromString(*NodeAsset->GetName()));
+							Params.SourceNode = AssetNode;
+							Params.CachedTitle = Title;
+
+							UBlueprintNodeSpawner* Spawner = UBlueprintNodeSpawner::Create(UK2Node_AnimGetter::StaticClass(), /*AssetNode->GetGraph()*/nullptr, UBlueprintNodeSpawner::FCustomizeNodeDelegate::CreateUObject(this, &UK2Node_AnimGetter::PostSpawnNodeSetup, Params));
+							Spawner->DynamicUiSignatureGetter = UBlueprintNodeSpawner::FUiSpecOverrideDelegate::CreateStatic(UiSpecOverride, Title);
+							ActionRegistrar.AddBlueprintAction(AnimBlueprint, Spawner);
+						}
+					}
+				}
+				else if(GetterRequiresParameter(Getter, TEXT("MachineIndex")))
+				{
+					if(GetterRequiresParameter(Getter, TEXT("StateIndex")))
+					{
+						for(UAnimStateNode* StateNode : StateNodes)
+						{
+							// Get the state machine node from the outer chain
+							UAnimationStateMachineGraph* Graph = Cast<UAnimationStateMachineGraph>(StateNode->GetOuter());
+							if(Graph)
+							{
+								if(UAnimGraphNode_StateMachine* MachineNode = Cast<UAnimGraphNode_StateMachine>(Graph->GetOuter()))
+								{
+									Params.SourceNode = MachineNode;
+								}
+							}
+
+							FText Title = FText::Format(LOCTEXT("NodeTitle", "{0} ({1})"), Getter->GetDisplayNameText(), StateNode->GetNodeTitle(ENodeTitleType::ListView));
+							Params.SourceStateNode = StateNode;
+							Params.CachedTitle = Title;
+
+							UBlueprintNodeSpawner* Spawner = UBlueprintNodeSpawner::Create(UK2Node_AnimGetter::StaticClass(), /*StateNode->GetGraph()*/nullptr, UBlueprintNodeSpawner::FCustomizeNodeDelegate::CreateUObject(this, &UK2Node_AnimGetter::PostSpawnNodeSetup, Params));
+							Spawner->DynamicUiSignatureGetter = UBlueprintNodeSpawner::FUiSpecOverrideDelegate::CreateStatic(UiSpecOverride, Title);
+							ActionRegistrar.AddBlueprintAction(AnimBlueprint, Spawner);
+						}
+					}
+					else if(GetterRequiresParameter(Getter, TEXT("TransitionIndex")))
+					{
+						for(UAnimStateTransitionNode* TransitionNode : TransitionNodes)
+						{
+							UAnimationStateMachineGraph* Graph = Cast<UAnimationStateMachineGraph>(TransitionNode->GetOuter());
+							if(Graph)
+							{
+								if(UAnimGraphNode_StateMachine* MachineNode = Cast<UAnimGraphNode_StateMachine>(Graph->GetOuter()))
+								{
+									Params.SourceNode = MachineNode;
+								}
+							}
+
+							FText Title = FText::Format(LOCTEXT("NodeTitle", "{0} ({1})"), Getter->GetDisplayNameText(), TransitionNode->GetNodeTitle(ENodeTitleType::ListView));
+							Params.SourceStateNode = TransitionNode;
+							Params.CachedTitle = Title;
+
+							UBlueprintNodeSpawner* Spawner = UBlueprintNodeSpawner::Create(UK2Node_AnimGetter::StaticClass(), /*TransitionNode->GetGraph()*/nullptr, UBlueprintNodeSpawner::FCustomizeNodeDelegate::CreateUObject(this, &UK2Node_AnimGetter::PostSpawnNodeSetup, Params));
+							Spawner->DynamicUiSignatureGetter = UBlueprintNodeSpawner::FUiSpecOverrideDelegate::CreateStatic(UiSpecOverride, Title);
+							ActionRegistrar.AddBlueprintAction(AnimBlueprint, Spawner);
+						}
+					}
+					else
+					{
+						// Only requires the state machine
+						for(UAnimGraphNode_StateMachine* MachineNode : MachineNodes)
+						{
+							FText Title = FText::Format(LOCTEXT("NodeTitle", "{0} ({1})"), Getter->GetDisplayNameText(), MachineNode->GetNodeTitle(ENodeTitleType::ListView));
+							Params.SourceNode = MachineNode;
+							Params.CachedTitle = Title;
+
+							UBlueprintNodeSpawner* Spawner = UBlueprintNodeSpawner::Create(UK2Node_AnimGetter::StaticClass(), /*MachineNode*/nullptr, UBlueprintNodeSpawner::FCustomizeNodeDelegate::CreateUObject(this, &UK2Node_AnimGetter::PostSpawnNodeSetup, Params));
+							Spawner->DynamicUiSignatureGetter = UBlueprintNodeSpawner::FUiSpecOverrideDelegate::CreateStatic(UiSpecOverride, Title);
+							ActionRegistrar.AddBlueprintAction(AnimBlueprint, Spawner);
+						}
+					}
+				}
+				else
+				{
+					// Doesn't operate on a node, only need one entry
+					FText Title = FText::Format(LOCTEXT("NodeTitle", "{0}"), Getter->GetDisplayNameText());
+					Params.CachedTitle = Title;
+
+					UBlueprintNodeSpawner* Spawner = UBlueprintNodeSpawner::Create(UK2Node_AnimGetter::StaticClass(), nullptr, UBlueprintNodeSpawner::FCustomizeNodeDelegate::CreateUObject(this, &UK2Node_AnimGetter::PostSpawnNodeSetup, Params));
+					Spawner->DynamicUiSignatureGetter = UBlueprintNodeSpawner::FUiSpecOverrideDelegate::CreateStatic(UiSpecOverride, Title);
+					ActionRegistrar.AddBlueprintAction(AnimBlueprint, Spawner);
+				}
+			}
+		}
+	}
+}
+
+bool UK2Node_AnimGetter::IsActionFilteredOut(FBlueprintActionFilter const& Filter)
+{
+	if(Filter.Context.Graphs.Num() > 0)
+	{
+		if(Filter.Context.Blueprints.Num() > 0)
+		{
+			UAnimBlueprint* AnimBlueprint = Cast<UAnimBlueprint>(Filter.Context.Blueprints[0]);
+			check(AnimBlueprint);
+
+			if(SourceAnimBlueprint == AnimBlueprint)
+			{
+				// Get the native anim instance derived class
+				UClass* NativeInstanceClass = AnimBlueprint->ParentClass;
+				while(NativeInstanceClass && !NativeInstanceClass->HasAnyClassFlags(CLASS_Native))
+				{
+					NativeInstanceClass = NativeInstanceClass->GetSuperClass();
+				}
+
+				if(GetterClass != NativeInstanceClass)
+				{
+					// If the anim instance containing the getter is not the class we're currently using then bail
+					return true;
+				}
+
+				const UEdGraphSchema* Schema = Filter.Context.Graphs[0]->GetSchema();
+				
+				// Bail if we aren't meant for this graph at all
+				if(!IsContextValidForSchema(Schema))
+				{
+					return true;
+				}
+
+				if(Cast<UAnimationTransitionSchema>(Schema) || Cast<UAnimationCustomTransitionSchema>(Schema))
+				{
+					if(!SourceNode && !SourceStateNode)
+					{
+						// No dependancies, always allow
+						return false;
+					}
+
+					// Inside a transition graph
+					if(SourceNode)
+					{
+						if(UAnimBlueprintGeneratedClass* BPSkelClass = AnimBlueprint->GetAnimBlueprintSkeletonClass())
+						{
+							auto GetTransitionNodeFromGraphLambda = [](const FAnimBlueprintDebugData& DebugData, const UEdGraph* Graph) -> UAnimStateTransitionNode*
+							{
+								if(const TWeakObjectPtr<UAnimStateTransitionNode>* TransNodePtr = DebugData.TransitionGraphToNodeMap.Find(Graph))
+								{
+									return TransNodePtr->Get();
+								}
+
+								if(const TWeakObjectPtr<UAnimStateTransitionNode>* TransNodePtr = DebugData.TransitionBlendGraphToNodeMap.Find(Graph))
+								{
+									return TransNodePtr->Get();
+								}
+
+								return NULL;
+							};
+
+							if(UAnimStateTransitionNode* TransitionNode = GetTransitionNodeFromGraphLambda(BPSkelClass->GetAnimBlueprintDebugData(), Filter.Context.Graphs[0]))
+							{
+								if(SourceStateNode)
+								{
+									if(UAnimStateTransitionNode* SourceTransitionNode = Cast<UAnimStateTransitionNode>(SourceStateNode))
+									{
+										// if we have a transition node, make sure it's the same as the one we're in
+										if(SourceTransitionNode == TransitionNode)
+										{
+											return false;
+										}
+									}
+									else if(UAnimStateNode* PreviousStateNode = Cast<UAnimStateNode>(TransitionNode->GetPreviousState()))
+									{
+										// Only allow actions using states that are referencing the previous state
+										if(SourceStateNode == PreviousStateNode)
+										{
+											return false;
+										}
+									}
+								}
+								else if(UAnimGraphNode_StateMachine* MachineNode = Cast<UAnimGraphNode_StateMachine>(SourceNode))
+								{
+									// Available everywhere
+									return false;
+								}
+								else if(UAnimStateNode* PrevStateNode = Cast<UAnimStateNode>(TransitionNode->GetPreviousState()))
+								{
+									// Make sure the attached asset node is in the source graph
+									if(SourceNode && SourceNode->GetGraph() == PrevStateNode->BoundGraph)
+									{
+										return false;
+									}
+								}
+							}
+						}
+					}
+				}
+				else if(Cast<UAnimationGraphSchema>(Schema))
+				{
+					// Inside normal anim graph
+					if(SourceStateNode)
+					{
+						for(UBlueprint* Blueprint : Filter.Context.Blueprints)
+						{
+							TArray<UAnimStateNode*> StateNodes;
+							FBlueprintEditorUtils::GetAllNodesOfClass(Blueprint, StateNodes);
+
+							if(StateNodes.Contains(SourceStateNode))
+							{
+								return false;
+							}
+						}
+					}
+				}
+			}
+		}
+	}
+
+	return true;
+}
+
+bool UK2Node_AnimGetter::GetterRequiresParameter(const UFunction* Getter, FString ParamName) const
+{
+	bool bRequiresParameter = false;
+
+	for(TFieldIterator<UProperty> PropIt(Getter); PropIt && (PropIt->PropertyFlags & CPF_Parm); ++PropIt)
+	{
+		UProperty* Prop = *PropIt;
+		
+		if(Prop->GetName() == ParamName)
+		{
+			bRequiresParameter = true;
+			break;
+		}
+	}
+
+	return bRequiresParameter;
+}
+
+
+void UK2Node_AnimGetter::PostSpawnNodeSetup(UEdGraphNode* NewNode, bool bIsTemplateNode, FNodeSpawnData SpawnData)
+{
+	UK2Node_AnimGetter* TypedNode = CastChecked<UK2Node_AnimGetter>(NewNode);
+
+	// Apply parameters
+	TypedNode->SourceNode = SpawnData.SourceNode;
+	TypedNode->SourceStateNode = SpawnData.SourceStateNode;
+	TypedNode->GetterClass = SpawnData.AnimInstanceClass;
+	TypedNode->SourceAnimBlueprint = SpawnData.SourceBlueprint;
+	TypedNode->SetFromFunction((UFunction*)SpawnData.Getter);
+	TypedNode->CachedTitle = SpawnData.CachedTitle;
+	
+	SpawnData.GetterContextString.ParseIntoArray(TypedNode->Contexts, TEXT("|"), 1);
+}
+
+bool UK2Node_AnimGetter::IsContextValidForSchema(const UEdGraphSchema* Schema) const
+{
+	if(Contexts.Num() == 0)
+	{
+		// Valid in all graphs
+		return true;
+	}
+	
+	for(const FString& Context : Contexts)
+	{
+		UClass* ClassToCheck = nullptr;
+		if(Context == TEXT("CustomBlend"))
+		{
+			ClassToCheck = UAnimationCustomTransitionSchema::StaticClass();
+		}
+
+		if(Context == TEXT("Transition"))
+		{
+			ClassToCheck = UAnimationTransitionSchema::StaticClass();
+		}
+
+		if(Context == TEXT("AnimGraph"))
+		{
+			ClassToCheck = UAnimationGraphSchema::StaticClass();
+		}
+
+		return Schema->GetClass() == ClassToCheck;
+	}
+
+	return false;
+}
+
+FNodeSpawnData::FNodeSpawnData()
+	: SourceNode(nullptr)
+	, SourceStateNode(nullptr)
+	, AnimInstanceClass(nullptr)
+	, SourceBlueprint(nullptr)
+	, Getter(nullptr)
+{
+
+}
+
+#undef LOCTEXT_NAMESPACE
diff --git a/Engine/Source/Editor/AnimGraph/Private/K2Node_TransitionRuleGetter.cpp b/Engine/Source/Editor/AnimGraph/Private/K2Node_TransitionRuleGetter.cpp
--- a/Engine/Source/Editor/AnimGraph/Private/K2Node_TransitionRuleGetter.cpp
+++ b/Engine/Source/Editor/AnimGraph/Private/K2Node_TransitionRuleGetter.cpp
@@ -262,14 +262,7 @@ void UK2Node_TransitionRuleGetter::GetNonStateSpecificMenuActions(FBlueprintActi
 
 void UK2Node_TransitionRuleGetter::GetMenuActions(FBlueprintActionDatabaseRegistrar& ActionRegistrar) const
 {
-	if( const UAnimBlueprint* AnimBlueprint = Cast<UAnimBlueprint>( ActionRegistrar.GetActionKeyFilter() ) )
-	{
-		GetStateSpecificMenuActions(ActionRegistrar, AnimBlueprint);
-	}
-	else
-	{
-		GetNonStateSpecificMenuActions(ActionRegistrar);
-	}
+
 }
 
 FText UK2Node_TransitionRuleGetter::GetTooltipText() const
diff --git a/Engine/Source/Editor/GraphEditor/Private/AnimationNodes/SGraphNodeSequencePlayer.cpp b/Engine/Source/Editor/GraphEditor/Private/AnimationNodes/SGraphNodeSequencePlayer.cpp
--- a/Engine/Source/Editor/GraphEditor/Private/AnimationNodes/SGraphNodeSequencePlayer.cpp
+++ b/Engine/Source/Editor/GraphEditor/Private/AnimationNodes/SGraphNodeSequencePlayer.cpp
@@ -119,7 +119,7 @@ bool SGraphNodeSequencePlayer::GetSequencePositionInfo(float& Out_Position, floa
 	{
 		if (UAnimSequenceBase* BoundSequence = SequencePlayer->Sequence)
 		{
-			Out_Position = SequencePlayer->InternalTimeAccumulator;
+			Out_Position = SequencePlayer->GetAccumulatedTime();
 			Out_Length = BoundSequence->SequenceLength;
 			Out_FrameCount = BoundSequence->GetNumberOfFrames();
 
@@ -152,7 +152,7 @@ void SGraphNodeSequencePlayer::SetSequencePositionRatio(float NewRatio)
 		if (SequencePlayer->Sequence != NULL)
 		{
 			const float NewTime = NewRatio * SequencePlayer->Sequence->SequenceLength;
-			SequencePlayer->InternalTimeAccumulator = NewTime;
+			SequencePlayer->SetAccumulatedTime(NewTime);
 		}
 	}
 }
diff --git a/Engine/Source/Editor/KismetCompiler/Private/AnimBlueprintCompiler.cpp b/Engine/Source/Editor/KismetCompiler/Private/AnimBlueprintCompiler.cpp
--- a/Engine/Source/Editor/KismetCompiler/Private/AnimBlueprintCompiler.cpp
+++ b/Engine/Source/Editor/KismetCompiler/Private/AnimBlueprintCompiler.cpp
@@ -29,6 +29,7 @@
 #include "AnimGraphNode_TransitionPoseEvaluator.h"
 #include "AnimGraphNode_TransitionResult.h"
 #include "K2Node_TransitionRuleGetter.h"
+#include "K2Node_AnimGetter.h"
 
 #define LOCTEXT_NAMESPACE "AnimBlueprintCompiler"
 
@@ -504,6 +505,9 @@ void FAnimBlueprintCompiler::ProcessAllAnimationNodes()
 	TArray<UK2Node_TransitionRuleGetter*> Getters;
 	ConsolidatedEventGraph->GetNodesOfClass<UK2Node_TransitionRuleGetter>(/*out*/ Getters);
 
+	TArray<UK2Node_AnimGetter*> AnimGetters;
+	ConsolidatedEventGraph->GetNodesOfClass<UK2Node_AnimGetter>(AnimGetters);
+
 	// Find the root node
 	UAnimGraphNode_Root* PrePhysicsRoot = NULL;
 	TArray<UAnimGraphNode_Base*> RootSet;
@@ -549,6 +553,11 @@ void FAnimBlueprintCompiler::ProcessAllAnimationNodes()
 			ProcessTransitionGetter(*GetterIt, NULL); // transition nodes should not appear at top-level
 		}
 
+		for(auto AnimGetterIt = AnimGetters.CreateIterator(); AnimGetterIt; ++AnimGetterIt)
+		{
+			AutoWireAnimGetter(*AnimGetterIt, NULL);
+		}
+
 		NewAnimBlueprintClass->RootAnimNodeIndex = GetAllocationIndexOfNode(PrePhysicsRoot);
 	}
 	else
@@ -566,6 +575,8 @@ int32 FAnimBlueprintCompiler::ExpandGraphAndProcessNodes(UEdGraph* SourceGraph,
 	UAnimGraphNode_Base* TargetRootNode = NULL;
 	TArray<UAnimGraphNode_Base*> AnimNodeList;
 	TArray<UK2Node_TransitionRuleGetter*> Getters;
+	TArray<UK2Node_AnimGetter*> AnimGetterNodes;
+
 	for (auto NodeIt = ClonedGraph->Nodes.CreateIterator(); NodeIt; ++NodeIt)
 	{
 		UEdGraphNode* Node = *NodeIt;
@@ -574,6 +585,10 @@ int32 FAnimBlueprintCompiler::ExpandGraphAndProcessNodes(UEdGraph* SourceGraph,
 		{
 			Getters.Add(GetterNode);
 		}
+		else if(UK2Node_AnimGetter* NewGetterNode = Cast<UK2Node_AnimGetter>(Node))
+		{
+			AnimGetterNodes.Add(NewGetterNode);
+		}
 		else if (UAnimGraphNode_Base* TestNode = Cast<UAnimGraphNode_Base>(Node))
 		{
 			AnimNodeList.Add(TestNode);
@@ -609,6 +624,12 @@ int32 FAnimBlueprintCompiler::ExpandGraphAndProcessNodes(UEdGraph* SourceGraph,
 		ProcessTransitionGetter(*GetterIt, TransitionNode);
 	}
 
+	// Wire anim getter nodes
+	for(auto AnimGetterIt = AnimGetterNodes.CreateIterator(); AnimGetterIt; ++AnimGetterIt)
+	{
+		AutoWireAnimGetter(*AnimGetterIt, TransitionNode);
+	}
+
 	// Returns the index of the processed cloned version of SourceRootNode
 	return GetAllocationIndexOfNode(TargetRootNode);	
 }
@@ -882,6 +903,25 @@ void FAnimBlueprintCompiler::ProcessStateMachine(UAnimGraphNode_StateMachineBase
 
 		FBakedAnimationState& BakedState = Oven.GetMachine().States[StateIndex];
 
+		// Add indices to all player nodes
+		TArray<UEdGraph*> GraphsToCheck;
+		TArray<UAnimGraphNode_AssetPlayerBase*> AssetPlayerNodes;
+		GraphsToCheck.Add(StateNode->GetBoundGraph());
+		StateNode->GetBoundGraph()->GetAllChildrenGraphs(GraphsToCheck);
+
+		for(UEdGraph* ChildGraph : GraphsToCheck)
+		{
+			ChildGraph->GetNodesOfClass(AssetPlayerNodes);
+		}
+
+		for(UAnimGraphNode_AssetPlayerBase* Node : AssetPlayerNodes)
+		{
+			if(int32* IndexPtr = NewAnimBlueprintClass->AnimBlueprintDebugData.NodeGuidToIndexMap.Find(Node->NodeGuid))
+			{
+				BakedState.PlayerNodeIndices.Add(*IndexPtr);
+			}
+		}
+
 		// Handle all the transitions out of this node
 		TArray<class UAnimStateTransitionNode*> TransitionList;
 		StateNode->GetTransitionList(/*out*/ TransitionList, /*bWantSortedList=*/ true);
@@ -1220,6 +1260,7 @@ void FAnimBlueprintCompiler::ProcessTransitionGetter(UK2Node_TransitionRuleGette
 
 	UEdGraphPin* SourceTimePin = NULL;
 	UAnimationAsset* AnimAsset= NULL;
+	int32 PlayerNodeIndex = INDEX_NONE;
 
 	if (UAnimGraphNode_Base* SourcePlayerNode = Getter->AssociatedAnimAssetPlayerNode)
 	{
@@ -1241,7 +1282,7 @@ void FAnimBlueprintCompiler::ProcessTransitionGetter(UK2Node_TransitionRuleGette
 		}
 
 		// Make sure the referenced AnimAsset player has been allocated
-		const int32 PlayerNodeIndex = GetAllocationIndexOfNode(UndertypedPlayerNode);
+		PlayerNodeIndex = GetAllocationIndexOfNode(UndertypedPlayerNode);
 		if (PlayerNodeIndex == INDEX_NONE)
 		{
 			MessageLog.Error(*LOCTEXT("BadAnimAssetNodeUsedInGetter", "@@ doesn't have a valid associated AnimAsset node.  Delete and recreate it").ToString(), Getter);
@@ -1627,6 +1668,89 @@ void FAnimBlueprintCompiler::DumpAnimDebugData()
 	}
 }
 
+void FAnimBlueprintCompiler::AutoWireAnimGetter(class UK2Node_AnimGetter* Getter, UAnimStateTransitionNode* InTransitionNode)
+{
+	UEdGraphPin* ReferencedNodeTimePin = nullptr;
+	int32 ReferencedNodeIndex = INDEX_NONE;
+	int32 SubNodeIndex = INDEX_NONE;
+	
+	UAnimGraphNode_Base* ProcessedNodeCheck = NULL;
+
+	if(UAnimGraphNode_Base* SourceNode = Getter->SourceNode)
+	{
+		UAnimGraphNode_Base* ActualSourceNode = MessageLog.FindSourceObjectTypeChecked<UAnimGraphNode_Base>(SourceNode);
+		
+		if(UAnimGraphNode_Base* ProcessedSourceNode = SourceNodeToProcessedNodeMap.FindRef(ActualSourceNode))
+		{
+			ProcessedNodeCheck = ProcessedSourceNode;
+
+			ReferencedNodeIndex = GetAllocationIndexOfNode(ProcessedSourceNode);
+
+			if(ProcessedSourceNode->DoesSupportTimeForTransitionGetter())
+			{
+				UScriptStruct* TimePropertyInStructType = ProcessedSourceNode->GetTimePropertyStruct();
+				const TCHAR* TimePropertyName = ProcessedSourceNode->GetTimePropertyName();
+
+				if(ReferencedNodeIndex != INDEX_NONE && TimePropertyName && TimePropertyInStructType)
+				{
+					UProperty* NodeProperty = AllocatedPropertiesByIndex.FindChecked(ReferencedNodeIndex);
+
+					UK2Node_StructMemberGet* ReaderNode = SpawnIntermediateNode<UK2Node_StructMemberGet>(Getter, ConsolidatedEventGraph);
+					ReaderNode->VariableReference.SetSelfMember(NodeProperty->GetFName());
+					ReaderNode->StructType = TimePropertyInStructType;
+					ReaderNode->AllocatePinsForSingleMemberGet(TimePropertyName);
+
+					ReferencedNodeTimePin = ReaderNode->FindPinChecked(TimePropertyName);
+				}
+			}
+		}
+	}
+	
+	if(Getter->SourceStateNode)
+	{
+		UObject* SourceObject = MessageLog.FindSourceObject(Getter->SourceStateNode);
+		if(UAnimStateNode* SourceStateNode = Cast<UAnimStateNode>(SourceObject))
+		{
+			if(FStateMachineDebugData* DebugData = NewAnimBlueprintClass->GetAnimBlueprintDebugData().StateMachineDebugData.Find(SourceStateNode->GetGraph()))
+			{
+				if(int32* StateIndexPtr = DebugData->NodeToStateIndex.Find(SourceStateNode))
+				{
+					SubNodeIndex = *StateIndexPtr;
+				}
+			}
+		}
+		else if(UAnimStateTransitionNode* TransitionNode = Cast<UAnimStateTransitionNode>(SourceObject))
+		{
+			if(FStateMachineDebugData* DebugData = NewAnimBlueprintClass->GetAnimBlueprintDebugData().StateMachineDebugData.Find(TransitionNode->GetGraph()))
+			{
+				if(int32* TransitionIndexPtr = DebugData->NodeToTransitionIndex.Find(TransitionNode))
+				{
+					SubNodeIndex = *TransitionIndexPtr;
+				}
+			}
+		}
+	}
+
+	check(Getter->IsNodePure());
+
+	for(UEdGraphPin* Pin : Getter->Pins)
+	{
+		// Hook up autowired parameters / pins
+		if(Pin->PinName == TEXT("CurrentTime"))
+		{
+			Pin->MakeLinkTo(ReferencedNodeTimePin);
+		}
+		else if(Pin->PinName == TEXT("AssetPlayerIndex") || Pin->PinName == TEXT("MachineIndex"))
+		{
+			Pin->DefaultValue = FString::FromInt(ReferencedNodeIndex);
+		}
+		else if(Pin->PinName == TEXT("StateIndex") || Pin->PinName == TEXT("TransitionIndex"))
+		{
+			Pin->DefaultValue = FString::FromInt(SubNodeIndex);
+		}
+	}
+}
+
 //////////////////////////////////////////////////////////////////////////
 
 #undef LOCTEXT_NAMESPACE
diff --git a/Engine/Source/Editor/KismetCompiler/Private/AnimBlueprintCompiler.h b/Engine/Source/Editor/KismetCompiler/Private/AnimBlueprintCompiler.h
--- a/Engine/Source/Editor/KismetCompiler/Private/AnimBlueprintCompiler.h
+++ b/Engine/Source/Editor/KismetCompiler/Private/AnimBlueprintCompiler.h
@@ -221,6 +221,9 @@ private:
 	//
 	void ProcessAnimationNodesGivenRoot(TArray<UAnimGraphNode_Base*>& AnimNodeList, const TArray<UAnimGraphNode_Base*>& RootSet);
 
+	// Automatically fill in parameters for the specified Getter node
+	void AutoWireAnimGetter(class UK2Node_AnimGetter* Getter, UAnimStateTransitionNode* InTransitionNode);
+
 	// This function does the following steps:
 	//   Clones the nodes in the specified source graph
 	//   Merges them into the ConsolidatedEventGraph
diff --git a/Engine/Source/Runtime/Engine/Classes/Animation/AnimInstance.h b/Engine/Source/Runtime/Engine/Classes/Animation/AnimInstance.h
--- a/Engine/Source/Runtime/Engine/Classes/Animation/AnimInstance.h
+++ b/Engine/Source/Runtime/Engine/Classes/Animation/AnimInstance.h
@@ -22,6 +22,7 @@ class UCanvas;
 class UWorld;
 struct FTransform;
 class FDebugDisplayInfo;
+struct FAnimNode_AssetPlayerBase;
 
 DECLARE_DELEGATE_OneParam(FOnMontageStarted, UAnimMontage*)
 DECLARE_DELEGATE_TwoParams(FOnMontageEnded, UAnimMontage*, bool /*bInterrupted*/)
@@ -604,32 +605,140 @@ private:
 	/** Trigger a Montage Ended event */
 	void TriggerMontageEndedEvent(const FQueuedMontageEndedEvent& MontageEndedEvent);
 
+public:
+
+	/** 
+	 * NOTE: Derived anim getters
+	 *
+	 * Anim getter functions can be defined for any instance deriving UAnimInstance.
+	 * To do this the function must be marked BlueprintPure, and have the AnimGetter metadata entry set to
+	 * "true". Following the instructions below, getters should appear correctly in the blueprint node context
+	 * menu for the derived classes
+	 *
+	 * A context string can be provided in the GetterContext metadata and can contain any (or none) of the
+	 * following entries separated by a pipe (|)
+	 * Transition  - Only available in a transition rule
+	 * AnimGraph   - Only available in an animgraph (also covers state anim graphs)
+	 * CustomBlend - Only available in a custom blend graph
+	 *
+	 * Anim getters support a number of automatic parameters that will be baked at compile time to be passed
+	 * to the functions. They will not appear as pins on the graph node. They are as follows:
+	 * AssetPlayerIndex - Index of an asset player node to operate on, one getter will be added to the blueprint action list per asset node available
+	 * MachineIndex     - Index of a state machine in the animation blueprint, one getter will be added to the blueprint action list per state machine
+	 * StateIndex       - Index of a state inside a state machine, also requires MachineIndex. One getter will be added to the blueprint action list per state
+	 * TransitionIndex  - Index of a transition inside a state machine, also requires MachineIndex. One getter will be added to the blueprint action list per transition
+	 */
+
+	/** Gets the length in seconds of the asset referenced in an asset player node */
+	UFUNCTION(BlueprintPure, Category="Asset Player", meta=(DisplayName="Length", BlueprintInternalUseOnly="true", AnimGetter="true"))
+	float GetInstanceAssetPlayerLength(int32 AssetPlayerIndex);
+
+	/** Get the current accumulated time in seconds for an asset player node */
+	UFUNCTION(BlueprintPure, Category="Asset Player", meta = (DisplayName = "Current Time", BlueprintInternalUseOnly = "true", AnimGetter = "true"))
+	float GetInstanceAssetPlayerTime(int32 AssetPlayerIndex);
+
+	/** Get the current accumulated time as a fraction for an asset player node */
+	UFUNCTION(BlueprintPure, Category="Asset Player", meta=(DisplayName="Current Time (ratio)", BlueprintInternalUseOnly="true", AnimGetter="true"))
+	float GetInstanceAssetPlayerTimeFraction(int32 AssetPlayerIndex);
+
+	/** Get the time in seconds from the end of an animation in an asset player node */
+	UFUNCTION(BlueprintPure, Category="Asset Player", meta=(DisplayName="Time Remaining", BlueprintInternalUseOnly="true", AnimGetter="true"))
+	float GetInstanceAssetPlayerTimeFromEnd(int32 AssetPlayerIndex);
+
+	/** Get the time as a fraction of the asset length of an animation in an asset player node */
+	UFUNCTION(BlueprintPure, Category="Asset Player", meta=(DisplayName="Time Remaining (ratio)", BlueprintInternalUseOnly="true", AnimGetter="true"))
+	float GetInstanceAssetPlayerTimeFromEndFraction(int32 AssetPlayerIndex);
+
+	/** Get the blend weight of a specified state */
+	UFUNCTION(BlueprintPure, Category="States", meta = (DisplayName="State Weight", BlueprintInternalUseOnly = "true", AnimGetter="true"))
+	float GetInstanceStateWeight(int32 MachineIndex, int32 StateIndex);
+
+	/** Get the current elapsed time of a state within the specified state machine */
+	UFUNCTION(BlueprintPure, Category="States", meta = (DisplayName="Current State Time", BlueprintInternalUseOnly = "true", AnimGetter="true", GetterContext="Transition"))
+	float GetInstanceCurrentStateElapsedTime(int32 MachineIndex);
+
+	/** Get the crossfade duration of a specified transition */
+	UFUNCTION(BlueprintPure, Category="Transitions", meta = (DisplayName="Get Transition Crossfade Duration", BlueprintInternalUseOnly = "true", AnimGetter="true"))
+	float GetInstanceTransitionCrossfadeDuration(int32 MachineIndex, int32 TransitionIndex);
+
+	/** Get the elapsed time in seconds of a specified transition */
+	UFUNCTION(BlueprintPure, Category="Transitions", meta = (DisplayName="Get Transition Time Elapsed", BlueprintInternalUseOnly = "true", AnimGetter="true", GetterContext="CustomBlend"))
+	float GetInstanceTransitionTimeElapsed(int32 MachineIndex, int32 TransitionIndex);
+
+	/** Get the elapsed time as a fraction of the crossfade duration of a specified transition */
+	UFUNCTION(BlueprintPure, Category="Transitions", meta = (DisplayName="Get Transition Time Elapsed (ratio)", BlueprintInternalUseOnly = "true", AnimGetter="true", GetterContext="CustomBlend"))
+	float GetInstanceTransitionTimeElapsedFraction(int32 MachineIndex, int32 TransitionIndex);
+
+	/** Get the time remaining in seconds for the most relevant animation in the source state */
+	UFUNCTION(BlueprintPure, Category="Asset Player", meta = (BlueprintInternalUseOnly = "true", AnimGetter="true", GetterContext="Transition"))
+	float GetRelevantAnimTimeRemaining(int32 MachineIndex, int32 StateIndex);
+
+	/** Get the time remaining as a fraction of the duration for the most relevant animation in the source state */
+	UFUNCTION(BlueprintPure, Category = "Asset Player", meta = (BlueprintInternalUseOnly = "true", AnimGetter = "true", GetterContext = "Transition"))
+	float GetRelevantAnimTimeRemainingFraction(int32 MachineIndex, int32 StateIndex);
+
+	/** Get the length in seconds of the most relevant animation in the source state */
+	UFUNCTION(BlueprintPure, Category = "Asset Player", meta = (BlueprintInternalUseOnly = "true", AnimGetter = "true", GetterContext = "Transition"))
+	float GetRelevantAnimLength(int32 MachineIndex, int32 StateIndex);
+
+	/** Get the current accumulated time in seconds for the most relevant animation in the source state */
+	UFUNCTION(BlueprintPure, Category = "Asset Player", meta = (BlueprintInternalUseOnly = "true", AnimGetter = "true", GetterContext = "Transition"))
+	float GetRelevantAnimTime(int32 MachineIndex, int32 StateIndex);
+
+	/** Get the current accumulated time as a fraction of the length of the most relevant animation in the source state */
+	UFUNCTION(BlueprintPure, Category = "Asset Player", meta = (BlueprintInternalUseOnly = "true", AnimGetter = "true", GetterContext = "Transition"))
+	float GetRelevantAnimTimeFraction(int32 MachineIndex, int32 StateIndex);
+
+protected:
+
+	/** Gets the runtime instance of the specified state machine */
+	FAnimNode_StateMachine* GetStateMachineInstance(int32 MachineIndex);
+
+	/** Gets a checked node given an index into the node property array */
+	template<class NodeType>
+	NodeType* GetCheckedNodeFromIndex(int32 NodeIdx);
+
+	/** Gets an unchecked (can return nullptr) node given an index into the node property array */
+	template<class NodeType>
+	NodeType* GetNodeFromIndex(int32 NodeIdx);
+
+	/** Gets the most relevant asset player in a specified state */
+	FAnimNode_AssetPlayerBase* GetRelevantAssetPlayerFromState(int32 MachineIndex, int32 StateIndex);
+
+	//////////////////////////////////////////////////////////////////////////
+
 public:
 	/** Returns the value of a named curve. */
 	UFUNCTION(BlueprintPure, Category="Animation")
 	float GetCurveValue(FName CurveName);
 
 	/** Returns the length (in seconds) of an animation AnimAsset. */
+	DEPRECATED(4.9, "GetAnimAssetPlayerLength is deprecated, use GetInstanceAssetPlayerLength instead")
 	UFUNCTION(BlueprintPure, Category="Animation", meta=(BlueprintInternalUseOnly = "true"))
 	static float GetAnimAssetPlayerLength(UAnimationAsset* AnimAsset);
 
 	//** Returns how far through the animation AnimAsset we are (as a proportion between 0.0 and 1.0). */
+	DEPRECATED(4.9, "GetAnimAssetPlayerTimeFraction is deprecated, use GetInstanceAssetPlayerTimeFraction instead")
 	UFUNCTION(BlueprintPure, Category="Animation", meta=(BlueprintInternalUseOnly = "true"))
 	static float GetAnimAssetPlayerTimeFraction(UAnimationAsset* AnimAsset, float CurrentTime);
 
 	/** Returns how long until the end of the animation AnimAsset (in seconds). */
+	DEPRECATED(4.9, "GetAnimAssetPlayerTimeFromEnd is deprecated, use GetInstanceAssetPlayerTimeFromEnd instead")
 	UFUNCTION(BlueprintPure, Category="Animation", meta=(BlueprintInternalUseOnly = "true"))
-	static float GetAnimAssetPlayerTimeFromEnd(UAnimationAsset* AnimAsset, float CurrentTime);
+	float GetAnimAssetPlayerTimeFromEnd(UAnimationAsset* AnimAsset, float CurrentTime);
 
 	/** Returns how long until the end of the animation AnimAsset we are (as a proportion between 0.0 and 1.0). */
+	DEPRECATED(4.9, "GetAnimAssetPlayerTimeFromEndFraction is deprecated, use GetInstanceAssetPlayerTimeFromEndFraction instead")
 	UFUNCTION(BlueprintPure, Category="Animation", meta=(BlueprintInternalUseOnly = "true"))
 	static float GetAnimAssetPlayerTimeFromEndFraction(UAnimationAsset* AnimAsset, float CurrentTime);
 
 	/** Returns the weight of a state in a state machine. */
+	DEPRECATED(4.9, "GetStateWeight is deprecated, use GetInstanceStateWeight instead")
 	UFUNCTION(BlueprintPure, Category="Animation", meta=(BlueprintInternalUseOnly = "true"))
 	float GetStateWeight(int32 MachineIndex, int32 StateIndex);
 
 	/** Returns (in seconds) the time a state machine has been active. */
+	DEPRECATED(4.9, "GetCurrentStateElapsedTime is deprecated, use GetInstanceCurrentStateElapsedTime instead")
 	UFUNCTION(BlueprintPure, Category="Animation", meta=(BlueprintInternalUseOnly = "true"))
 	float GetCurrentStateElapsedTime(int32 MachineIndex);
 
diff --git a/Engine/Source/Runtime/Engine/Classes/Animation/AnimNode_AssetPlayerBase.cpp b/Engine/Source/Runtime/Engine/Classes/Animation/AnimNode_AssetPlayerBase.cpp
new file mode 100644
--- /dev/null
+++ b/Engine/Source/Runtime/Engine/Classes/Animation/AnimNode_AssetPlayerBase.cpp
@@ -0,0 +1,40 @@
+// Copyright 1998-2015 Epic Games, Inc. All Rights Reserved.
+
+#include "EnginePrivate.h"
+#include "AnimNode_AssetPlayerBase.h"
+
+FAnimNode_AssetPlayerBase::FAnimNode_AssetPlayerBase()
+	: bIgnoreForRelevancyTest(false)
+	, BlendWeight(0.0f)
+	, InternalTimeAccumulator(0.0f)
+{
+
+}
+
+void FAnimNode_AssetPlayerBase::Update(const FAnimationUpdateContext& Context)
+{
+	// Cache the current weight and update the node
+	BlendWeight = Context.GetFinalBlendWeight();
+	UpdateAssetPlayer(Context);
+}
+
+float FAnimNode_AssetPlayerBase::GetCachedBlendWeight()
+{
+	return BlendWeight;
+}
+
+float FAnimNode_AssetPlayerBase::GetAccumulatedTime()
+{
+	return InternalTimeAccumulator;
+}
+
+void FAnimNode_AssetPlayerBase::SetAccumulatedTime(const float& NewTime)
+{
+	InternalTimeAccumulator = NewTime;
+}
+
+UAnimationAsset* FAnimNode_AssetPlayerBase::GetAnimAsset()
+{
+	return nullptr;
+}
+
diff --git a/Engine/Source/Runtime/Engine/Classes/Animation/AnimNode_AssetPlayerBase.h b/Engine/Source/Runtime/Engine/Classes/Animation/AnimNode_AssetPlayerBase.h
new file mode 100644
--- /dev/null
+++ b/Engine/Source/Runtime/Engine/Classes/Animation/AnimNode_AssetPlayerBase.h
@@ -0,0 +1,51 @@
+// Copyright 1998-2015 Epic Games, Inc. All Rights Reserved.
+
+#pragma once
+
+#include "AnimNodeBase.h"
+#include "AnimNode_AssetPlayerBase.generated.h"
+
+/* Base class for any asset playing anim node */
+USTRUCT()
+struct ENGINE_API FAnimNode_AssetPlayerBase : public FAnimNode_Base
+{
+	GENERATED_BODY();
+
+	/** If true, "Relevant anim" nodes that look for the highest weighted animation in a state will ignore
+	 *  this node
+	 */
+	UPROPERTY(EditAnywhere, BlueprintReadWrite, Category=Relevancy, meta=(PinHiddenByDefault))
+	bool bIgnoreForRelevancyTest;
+
+	FAnimNode_AssetPlayerBase();
+
+	/** Get the last encountered blend weight for this node */
+	virtual float GetCachedBlendWeight();
+	
+	/** Get the currently referenced time within the asset player node */
+	virtual float GetAccumulatedTime();
+
+	/** Override the currently accumulated time */
+	virtual void SetAccumulatedTime(const float& NewTime);
+
+	/** Get the animation asset associated with the node, derived classes should implement this */
+	virtual UAnimationAsset* GetAnimAsset();
+
+	/** Update the node, marked final so we can always handle blendweight caching.
+	 *  Derived classes should implement UpdateAssetPlayer
+	 */
+	virtual void Update(const FAnimationUpdateContext& Context) final override;
+
+	/** Update method for the asset player, to be implemented by derived classes */
+	virtual void UpdateAssetPlayer(const FAnimationUpdateContext& Context) {};
+
+protected:
+
+	/** Last encountered blendweight for this node */
+	UPROPERTY(BlueprintReadWrite, Transient, Category=DoNotEdit)
+	float BlendWeight;
+
+	/** Accumulated time used to reference the asset in this node */
+	UPROPERTY(BlueprintReadWrite, Transient, Category=DoNotEdit)
+	float InternalTimeAccumulator;
+};
\ No newline at end of file
diff --git a/Engine/Source/Runtime/Engine/Classes/Animation/AnimNode_BlendSpaceEvaluator.h b/Engine/Source/Runtime/Engine/Classes/Animation/AnimNode_BlendSpaceEvaluator.h
--- a/Engine/Source/Runtime/Engine/Classes/Animation/AnimNode_BlendSpaceEvaluator.h
+++ b/Engine/Source/Runtime/Engine/Classes/Animation/AnimNode_BlendSpaceEvaluator.h
@@ -22,7 +22,7 @@ public:
 	FAnimNode_BlendSpaceEvaluator();
 
 	// FAnimNode_Base interface
-	virtual void Update(const FAnimationUpdateContext& Context) override;
+	virtual void UpdateAssetPlayer(const FAnimationUpdateContext& Context) override;
 	virtual void GatherDebugData(FNodeDebugData& DebugData) override;
 	// End of FAnimNode_Base interface
 };
diff --git a/Engine/Source/Runtime/Engine/Classes/Animation/AnimNode_BlendSpacePlayer.h b/Engine/Source/Runtime/Engine/Classes/Animation/AnimNode_BlendSpacePlayer.h
--- a/Engine/Source/Runtime/Engine/Classes/Animation/AnimNode_BlendSpacePlayer.h
+++ b/Engine/Source/Runtime/Engine/Classes/Animation/AnimNode_BlendSpacePlayer.h
@@ -3,11 +3,12 @@
 #pragma once
 
 #include "AnimNodeBase.h"
+#include "AnimNode_AssetPlayerBase.h"
 #include "AnimNode_BlendSpacePlayer.generated.h"
 
 //@TODO: Comment
 USTRUCT()
-struct ENGINE_API FAnimNode_BlendSpacePlayer : public FAnimNode_Base
+struct ENGINE_API FAnimNode_BlendSpacePlayer : public FAnimNode_AssetPlayerBase
 {
 	GENERATED_USTRUCT_BODY()
 
@@ -44,8 +45,6 @@ public:
 	UPROPERTY()
 	TEnumAsByte<EAnimGroupRole::Type> GroupRole;
 protected:
-	UPROPERTY(BlueprintReadWrite, Transient, Category=DoNotEdit)
-	float InternalTimeAccumulator;
 
 	UPROPERTY()
 	FBlendFilter BlendFilter;
@@ -59,12 +58,16 @@ public:
 	// FAnimNode_Base interface
 	virtual void Initialize(const FAnimationInitializeContext& Context) override;
 	virtual void CacheBones(const FAnimationCacheBonesContext& Context) override;
-	virtual void Update(const FAnimationUpdateContext& Context) override;
+	virtual void UpdateAssetPlayer(const FAnimationUpdateContext& Context) override;
 	virtual void Evaluate(FPoseContext& Output) override;
 	virtual void OverrideAsset(UAnimationAsset* NewAsset) override;
 	virtual void GatherDebugData(FNodeDebugData& DebugData) override;
 	// End of FAnimNode_Base interface
 
+	float GetTimeFromEnd(float CurrentTime);
+
+	UAnimationAsset* GetAnimAsset();
+
 protected:
 	void UpdateInternal(const FAnimationUpdateContext& Context);
 };
diff --git a/Engine/Source/Runtime/Engine/Classes/Animation/AnimNode_RotationOffsetBlendSpace.h b/Engine/Source/Runtime/Engine/Classes/Animation/AnimNode_RotationOffsetBlendSpace.h
--- a/Engine/Source/Runtime/Engine/Classes/Animation/AnimNode_RotationOffsetBlendSpace.h
+++ b/Engine/Source/Runtime/Engine/Classes/Animation/AnimNode_RotationOffsetBlendSpace.h
@@ -20,7 +20,7 @@ public:
 	// FAnimNode_Base interface
 	virtual void Initialize(const FAnimationInitializeContext& Context) override;
 	virtual void CacheBones(const FAnimationCacheBonesContext& Context) override;
-	virtual void Update(const FAnimationUpdateContext& Context) override;
+	virtual void UpdateAssetPlayer(const FAnimationUpdateContext& Context) override;
 	virtual void Evaluate(FPoseContext& Output) override;
 	virtual void GatherDebugData(FNodeDebugData& DebugData) override;
 	// End of FAnimNode_Base interface
diff --git a/Engine/Source/Runtime/Engine/Classes/Animation/AnimNode_SequenceEvaluator.h b/Engine/Source/Runtime/Engine/Classes/Animation/AnimNode_SequenceEvaluator.h
--- a/Engine/Source/Runtime/Engine/Classes/Animation/AnimNode_SequenceEvaluator.h
+++ b/Engine/Source/Runtime/Engine/Classes/Animation/AnimNode_SequenceEvaluator.h
@@ -2,13 +2,14 @@
 
 #pragma once
 #include "Animation/AnimNodeBase.h"
+#include "AnimNode_AssetPlayerBase.h"
 #include "AnimNode_SequenceEvaluator.generated.h"
 
 // Evaluates a point in an anim sequence, using a specific time input rather than advancing time internally.
 // Typically the playback position of the animation for this node will represent something other than time, like jump height.
 // This node will not trigger any notifies present in the associated sequence.
 USTRUCT()
-struct ENGINE_API FAnimNode_SequenceEvaluator : public FAnimNode_Base
+struct ENGINE_API FAnimNode_SequenceEvaluator : public FAnimNode_AssetPlayerBase
 {
 	GENERATED_USTRUCT_BODY()
 public:
@@ -29,9 +30,15 @@ public:
 	// FAnimNode_Base interface
 	virtual void Initialize(const FAnimationInitializeContext& Context) override;
 	virtual void CacheBones(const FAnimationCacheBonesContext& Context) override;
-	virtual void Update(const FAnimationUpdateContext& Context) override;
+	virtual void UpdateAssetPlayer(const FAnimationUpdateContext& Context) override;
 	virtual void Evaluate(FPoseContext& Output) override;
 	virtual void OverrideAsset(UAnimationAsset* NewAsset) override;
 	virtual void GatherDebugData(FNodeDebugData& DebugData) override;
 	// End of FAnimNode_Base interface
+
+	// FAnimNode_AssetPlayerBase Interface
+	virtual float GetAccumulatedTime() {return ExplicitTime;}
+	virtual void SetAccumulatedTime(const float& NewTime) {ExplicitTime = NewTime;}
+	virtual UAnimationAsset* GetAnimAsset() {return Sequence;}
+	// End of FAnimNode_AssetPlayerBase Interface
 };
diff --git a/Engine/Source/Runtime/Engine/Classes/Animation/AnimNode_SequencePlayer.h b/Engine/Source/Runtime/Engine/Classes/Animation/AnimNode_SequencePlayer.h
--- a/Engine/Source/Runtime/Engine/Classes/Animation/AnimNode_SequencePlayer.h
+++ b/Engine/Source/Runtime/Engine/Classes/Animation/AnimNode_SequencePlayer.h
@@ -2,13 +2,14 @@
 
 #pragma once
 #include "Animation/AnimNodeBase.h"
+#include "AnimNode_AssetPlayerBase.h"
 #include "AnimNode_SequencePlayer.generated.h"
 
 #pragma once
 
 // Sequence player node
 USTRUCT()
-struct ENGINE_API FAnimNode_SequencePlayer : public FAnimNode_Base
+struct ENGINE_API FAnimNode_SequencePlayer : public FAnimNode_AssetPlayerBase
 {
 	GENERATED_USTRUCT_BODY()
 public:
@@ -16,10 +17,6 @@ public:
 	UPROPERTY(EditAnywhere, BlueprintReadWrite, Category=Settings, meta=(PinHiddenByDefault))
 	UAnimSequenceBase* Sequence;
 
-	// The internal time accumulator (position within the animation asset being played)
-	UPROPERTY(BlueprintReadWrite, Transient, Category=DoNotEdit)
-	float InternalTimeAccumulator;
-
 	// Should the animation continue looping when it reaches the end?
 	UPROPERTY(EditAnywhere, BlueprintReadWrite, Category=Settings, meta=(PinHiddenByDefault))
 	mutable bool bLoopAnimation;
@@ -38,7 +35,6 @@ public:
 public:	
 	FAnimNode_SequencePlayer()
 		: Sequence(NULL)
-		, InternalTimeAccumulator(0.0f)
 		, bLoopAnimation(true)
 		, PlayRate(1.0f)
 		, GroupIndex(INDEX_NONE)
@@ -48,9 +44,13 @@ public:
 	// FAnimNode_Base interface
 	virtual void Initialize(const FAnimationInitializeContext& Context) override;
 	virtual void CacheBones(const FAnimationCacheBonesContext& Context) override;
-	virtual void Update(const FAnimationUpdateContext& Context) override;
+	virtual void UpdateAssetPlayer(const FAnimationUpdateContext& Context) override;
 	virtual void Evaluate(FPoseContext& Output) override;
 	virtual void OverrideAsset(UAnimationAsset* NewAsset) override;
 	virtual void GatherDebugData(FNodeDebugData& DebugData) override;
 	// End of FAnimNode_Base interface
+
+	UAnimationAsset* GetAnimAsset() {return Sequence;}
+
+	float GetTimeFromEnd(float CurrentNodeTime);
 };
diff --git a/Engine/Source/Runtime/Engine/Classes/Animation/AnimNode_StateMachine.h b/Engine/Source/Runtime/Engine/Classes/Animation/AnimNode_StateMachine.h
--- a/Engine/Source/Runtime/Engine/Classes/Animation/AnimNode_StateMachine.h
+++ b/Engine/Source/Runtime/Engine/Classes/Animation/AnimNode_StateMachine.h
@@ -182,6 +182,8 @@ public:
 	const FBakedAnimationState& GetStateInfo(int32 StateIndex) const;
 	const FAnimationTransitionBetweenStates& GetTransitionInfo(int32 TransIndex) const;
 	
+	bool IsValidTransitionIndex(int32 TransitionIndex) const;
+
 protected:
 	// Tries to get the instance information for the state machine
 	FBakedAnimationStateMachine* GetMachineDescription();
diff --git a/Engine/Source/Runtime/Engine/Classes/Animation/AnimStateMachineTypes.h b/Engine/Source/Runtime/Engine/Classes/Animation/AnimStateMachineTypes.h
--- a/Engine/Source/Runtime/Engine/Classes/Animation/AnimStateMachineTypes.h
+++ b/Engine/Source/Runtime/Engine/Classes/Animation/AnimStateMachineTypes.h
@@ -226,6 +226,9 @@ struct FBakedAnimationState
 	UPROPERTY()
 	int32 EntryRuleNodeIndex;
 
+	// Indices into the property array for player nodes in the state
+	UPROPERTY()
+	TArray<int32> PlayerNodeIndices;
 
 public:
 	FBakedAnimationState()
diff --git a/Engine/Source/Runtime/Engine/Private/Animation/AnimInstance.cpp b/Engine/Source/Runtime/Engine/Private/Animation/AnimInstance.cpp
--- a/Engine/Source/Runtime/Engine/Private/Animation/AnimInstance.cpp
+++ b/Engine/Source/Runtime/Engine/Private/Animation/AnimInstance.cpp
@@ -20,6 +20,9 @@
 #include "Animation/AnimMontage.h"
 #include "Animation/AnimInstance.h"
 #include "Animation/AnimNode_TransitionResult.h"
+#include "Animation/AnimNode_SequencePlayer.h"
+#include "Animation/AnimNode_BlendSpacePlayer.h"
+#include "Animation/AnimNode_AssetPlayerBase.h"
 
 /** Anim stats */
 
@@ -2660,4 +2663,287 @@ void UAnimInstance::UpdateMontageEvaluationData()
 	}
 }
 
+float UAnimInstance::GetInstanceAssetPlayerLength(int32 AssetPlayerIndex)
+{
+	if(FAnimNode_AssetPlayerBase* PlayerNode = GetNodeFromIndex<FAnimNode_AssetPlayerBase>(AssetPlayerIndex))
+	{
+		return PlayerNode->GetAnimAsset()->GetMaxCurrentTime();
+	}
+
+	return 0.0f;
+}
+
+float UAnimInstance::GetInstanceAssetPlayerTime(int32 AssetPlayerIndex)
+{
+	if(FAnimNode_AssetPlayerBase* PlayerNode = GetNodeFromIndex<FAnimNode_AssetPlayerBase>(AssetPlayerIndex))
+	{
+		return PlayerNode->GetAccumulatedTime();
+	}
+
+	return 0.0f;
+}
+
+float UAnimInstance::GetInstanceAssetPlayerTimeFraction(int32 AssetPlayerIndex)
+{
+	if(FAnimNode_AssetPlayerBase* PlayerNode = GetNodeFromIndex<FAnimNode_AssetPlayerBase>(AssetPlayerIndex))
+	{
+		UAnimationAsset* Asset = PlayerNode->GetAnimAsset();
+
+		float Length = Asset ? Asset->GetMaxCurrentTime() : 0.0f;
+
+		if(Length > 0.0f)
+		{
+			return PlayerNode->GetAccumulatedTime() / Length;
+		}
+	}
+
+	
+
+	return 0.0f;
+}
+
+float UAnimInstance::GetInstanceAssetPlayerTimeFromEndFraction(int32 AssetPlayerIndex)
+{
+	if(FAnimNode_AssetPlayerBase* PlayerNode = GetNodeFromIndex<FAnimNode_AssetPlayerBase>(AssetPlayerIndex))
+	{
+		UAnimationAsset* Asset = PlayerNode->GetAnimAsset();
+		
+		float Length = Asset ? Asset->GetMaxCurrentTime() : 0.f;
+
+		if(Length > 0.f)
+		{
+			return (Length - PlayerNode->GetAccumulatedTime()) / Length;
+		}
+	}
+
+	return 0.0f;
+}
+
+float UAnimInstance::GetInstanceAssetPlayerTimeFromEnd(int32 AssetPlayerIndex)
+{
+	if(FAnimNode_AssetPlayerBase* PlayerNode = GetNodeFromIndex<FAnimNode_AssetPlayerBase>(AssetPlayerIndex))
+	{
+		UAnimationAsset* Asset = PlayerNode->GetAnimAsset();
+		if(Asset)
+		{
+			return PlayerNode->GetAnimAsset()->GetMaxCurrentTime() - PlayerNode->GetAccumulatedTime();
+		}
+	}
+
+	return 0.0f;
+}
+
+float UAnimInstance::GetInstanceStateWeight(int32 MachineIndex, int32 StateIndex)
+{
+	if(FAnimNode_StateMachine* MachineInstance = GetStateMachineInstance(MachineIndex))
+	{
+		return MachineInstance->GetStateWeight(StateIndex);
+	}
+	
+	return 0.0f;
+}
+
+float UAnimInstance::GetInstanceCurrentStateElapsedTime(int32 MachineIndex)
+{
+	if(FAnimNode_StateMachine* MachineInstance = GetStateMachineInstance(MachineIndex))
+	{
+		return MachineInstance->GetCurrentStateElapsedTime();
+	}
+
+	return 0.0f;
+}
+
+float UAnimInstance::GetInstanceTransitionCrossfadeDuration(int32 MachineIndex, int32 TransitionIndex)
+{
+	if(FAnimNode_StateMachine* MachineInstance = GetStateMachineInstance(MachineIndex))
+	{
+		if(MachineInstance->IsValidTransitionIndex(TransitionIndex))
+		{
+			return MachineInstance->GetTransitionInfo(TransitionIndex).CrossfadeDuration;
+		}
+	}
+
+	return 0.0f;
+}
+
+float UAnimInstance::GetInstanceTransitionTimeElapsed(int32 MachineIndex, int32 TransitionIndex)
+{
+	// Just an alias for readability in the anim graph
+	return GetInstanceCurrentStateElapsedTime(MachineIndex);
+}
+
+float UAnimInstance::GetInstanceTransitionTimeElapsedFraction(int32 MachineIndex, int32 TransitionIndex)
+{
+	if(FAnimNode_StateMachine* MachineInstance = GetStateMachineInstance(MachineIndex))
+	{
+		if(MachineInstance->IsValidTransitionIndex(TransitionIndex))
+		{
+			return MachineInstance->GetCurrentStateElapsedTime() / MachineInstance->GetTransitionInfo(TransitionIndex).CrossfadeDuration;
+		}
+	}
+
+	return 0.0f;
+}
+
+float UAnimInstance::GetRelevantAnimTimeRemaining(int32 MachineIndex, int32 StateIndex)
+{
+	if(FAnimNode_AssetPlayerBase* AssetPlayer = GetRelevantAssetPlayerFromState(MachineIndex, StateIndex))
+	{
+		if(AssetPlayer->GetAnimAsset())
+		{
+			return AssetPlayer->GetAnimAsset()->GetMaxCurrentTime() - AssetPlayer->GetAccumulatedTime();
+		}
+	}
+
+	return 0.0f;
+}
+
+float UAnimInstance::GetRelevantAnimTimeRemainingFraction(int32 MachineIndex, int32 StateIndex)
+{
+	if(FAnimNode_AssetPlayerBase* AssetPlayer = GetRelevantAssetPlayerFromState(MachineIndex, StateIndex))
+	{
+		if(AssetPlayer->GetAnimAsset())
+		{
+			float Length = AssetPlayer->GetAnimAsset()->GetMaxCurrentTime();
+			if(Length > 0.0f)
+			{
+				return (Length - AssetPlayer->GetAccumulatedTime()) / Length;
+			}
+		}
+	}
+
+	return 0.0f;
+}
+
+float UAnimInstance::GetRelevantAnimLength(int32 MachineIndex, int32 StateIndex)
+{
+	if(FAnimNode_AssetPlayerBase* AssetPlayer = GetRelevantAssetPlayerFromState(MachineIndex, StateIndex))
+	{
+		if(AssetPlayer->GetAnimAsset())
+		{
+			return AssetPlayer->GetAnimAsset()->GetMaxCurrentTime();
+		}
+	}
+
+	return 0.0f;
+}
+
+float UAnimInstance::GetRelevantAnimTime(int32 MachineIndex, int32 StateIndex)
+{
+	if(FAnimNode_AssetPlayerBase* AssetPlayer = GetRelevantAssetPlayerFromState(MachineIndex, StateIndex))
+	{
+		return AssetPlayer->GetAccumulatedTime();
+	}
+
+	return 0.0f;
+}
+
+float UAnimInstance::GetRelevantAnimTimeFraction(int32 MachineIndex, int32 StateIndex)
+{
+	if(FAnimNode_AssetPlayerBase* AssetPlayer = GetRelevantAssetPlayerFromState(MachineIndex, StateIndex))
+	{
+		float Length = AssetPlayer->GetAnimAsset()->GetMaxCurrentTime();
+		if(Length > 0.0f)
+		{
+			return AssetPlayer->GetAccumulatedTime() / Length;
+		}
+	}
+
+	return 0.0f;
+}
+
+FAnimNode_StateMachine* UAnimInstance::GetStateMachineInstance(int32 MachineIndex)
+{
+	if(UAnimBlueprintGeneratedClass* AnimBlueprintClass = Cast<UAnimBlueprintGeneratedClass>((UObject*)GetClass()))
+	{
+		if((MachineIndex >= 0) && (MachineIndex < AnimBlueprintClass->AnimNodeProperties.Num()))
+		{
+			const int32 InstancePropertyIndex = AnimBlueprintClass->AnimNodeProperties.Num() - 1 - MachineIndex;
+
+			UStructProperty* MachineInstanceProperty = AnimBlueprintClass->AnimNodeProperties[InstancePropertyIndex];
+			checkSlow(MachineInstanceProperty->Struct->IsChildOf(FAnimNode_StateMachine::StaticStruct()));
+
+			return MachineInstanceProperty->ContainerPtrToValuePtr<FAnimNode_StateMachine>(this);
+		}
+	}
+
+	return nullptr;
+}
+
+template<class NodeType>
+NodeType* UAnimInstance::GetCheckedNodeFromIndex(int32 NodeIdx)
+{
+	NodeType* NodePtr = nullptr;
+	if(UAnimBlueprintGeneratedClass* AnimBlueprintClass = Cast<UAnimBlueprintGeneratedClass>(GetClass()))
+	{
+		const int32 InstanceIdx = AnimBlueprintClass->AnimNodeProperties.Num() - 1 - NodeIdx;
+
+		if(AnimBlueprintClass->AnimNodeProperties.IsValidIndex(InstanceIdx))
+		{
+			UStructProperty* NodeProperty = AnimBlueprintClass->AnimNodeProperties[InstanceIdx];
+
+			if(NodeProperty->Struct->IsChildOf(NodeType::StaticStruct()))
+			{
+				NodePtr = NodeProperty->ContainerPtrToValuePtr<NodeType>(this);
+			}
+			else
+			{
+				checkfSlow(false, TEXT("Requested a node of type %s but found node of type %s"), *NodeType::StaticStruct()->GetName(), *NodeProperty->Struct->GetName());
+			}
+		}
+		else
+		{
+			checkfSlow(false, TEXT("Requested node of type %s at index %d/%d, index out of bounds."), *NodeType::StaticStruct()->GetName(), NodeIdx, InstanceIdx);
+		}
+	}
+
+	checkfSlow(NodePtr, TEXT("Requested node at index %d not found!"), NodeIdx);
+
+	return NodePtr;
+}
+
+template<class NodeType>
+NodeType* UAnimInstance::GetNodeFromIndex(int32 NodeIdx)
+{
+	NodeType* NodePtr = nullptr;
+	if(UAnimBlueprintGeneratedClass* AnimBlueprintClass = Cast<UAnimBlueprintGeneratedClass>(GetClass()))
+	{
+		const int32 InstanceIdx = AnimBlueprintClass->AnimNodeProperties.Num() - 1 - NodeIdx;
+
+		if(AnimBlueprintClass->AnimNodeProperties.IsValidIndex(InstanceIdx))
+		{
+			UStructProperty* NodeProperty = AnimBlueprintClass->AnimNodeProperties[InstanceIdx];
+
+			if(NodeProperty->Struct->IsChildOf(NodeType::StaticStruct()))
+			{
+				NodePtr = NodeProperty->ContainerPtrToValuePtr<NodeType>(this);
+			}
+		}
+	}
+
+	return NodePtr;
+}
+
+FAnimNode_AssetPlayerBase* UAnimInstance::GetRelevantAssetPlayerFromState(int32 MachineIndex, int32 StateIndex)
+{
+	FAnimNode_AssetPlayerBase* ResultPlayer = nullptr;
+	if(FAnimNode_StateMachine* MachineInstance = GetStateMachineInstance(MachineIndex))
+	{
+		float MaxWeight = 0.0f;
+		const FBakedAnimationState& State = MachineInstance->GetStateInfo(StateIndex);
+		for(const int32& PlayerIdx : State.PlayerNodeIndices)
+		{
+			if(FAnimNode_AssetPlayerBase* Player = GetNodeFromIndex<FAnimNode_AssetPlayerBase>(PlayerIdx))
+			{
+				if(!Player->bIgnoreForRelevancyTest && Player->GetCachedBlendWeight() > MaxWeight)
+				{
+					MaxWeight = Player->GetCachedBlendWeight();
+					ResultPlayer = Player;
+				}
+			}
+		}
+	}
+	return ResultPlayer;
+}
+
 #undef LOCTEXT_NAMESPACE 
diff --git a/Engine/Source/Runtime/Engine/Private/Animation/AnimNode_BlendSpaceEvaluator.cpp b/Engine/Source/Runtime/Engine/Private/Animation/AnimNode_BlendSpaceEvaluator.cpp
--- a/Engine/Source/Runtime/Engine/Private/Animation/AnimNode_BlendSpaceEvaluator.cpp
+++ b/Engine/Source/Runtime/Engine/Private/Animation/AnimNode_BlendSpaceEvaluator.cpp
@@ -13,7 +13,7 @@ FAnimNode_BlendSpaceEvaluator::FAnimNode_BlendSpaceEvaluator()
 {
 }
 
-void FAnimNode_BlendSpaceEvaluator::Update(const FAnimationUpdateContext& Context)
+void FAnimNode_BlendSpaceEvaluator::UpdateAssetPlayer(const FAnimationUpdateContext& Context)
 {
 	EvaluateGraphExposedInputs.Execute(Context);
 	InternalTimeAccumulator = NormalizedTime;
diff --git a/Engine/Source/Runtime/Engine/Private/Animation/AnimNode_BlendSpacePlayer.cpp b/Engine/Source/Runtime/Engine/Private/Animation/AnimNode_BlendSpacePlayer.cpp
--- a/Engine/Source/Runtime/Engine/Private/Animation/AnimNode_BlendSpacePlayer.cpp
+++ b/Engine/Source/Runtime/Engine/Private/Animation/AnimNode_BlendSpacePlayer.cpp
@@ -15,7 +15,6 @@ FAnimNode_BlendSpacePlayer::FAnimNode_BlendSpacePlayer()
 	, bLoop(true)
 	, BlendSpace(NULL)
 	, GroupIndex(INDEX_NONE)
-	, InternalTimeAccumulator(0.0f)
 {
 }
 
@@ -44,7 +43,7 @@ void FAnimNode_BlendSpacePlayer::CacheBones(const FAnimationCacheBonesContext& C
 {
 }
 
-void FAnimNode_BlendSpacePlayer::Update(const FAnimationUpdateContext& Context)
+void FAnimNode_BlendSpacePlayer::UpdateAssetPlayer(const FAnimationUpdateContext& Context)
 {
 	EvaluateGraphExposedInputs.Execute(Context);
 
@@ -102,4 +101,14 @@ void FAnimNode_BlendSpacePlayer::GatherDebugData(FNodeDebugData& DebugData)
 
 		//BlendSpace->GetBlendSamples();
 	}
-}
\ No newline at end of file
+}
+
+float FAnimNode_BlendSpacePlayer::GetTimeFromEnd(float CurrentTime)
+{
+	return BlendSpace->GetMaxCurrentTime() - CurrentTime;
+}
+
+UAnimationAsset* FAnimNode_BlendSpacePlayer::GetAnimAsset()
+{
+	return BlendSpace;
+}
diff --git a/Engine/Source/Runtime/Engine/Private/Animation/AnimNode_RotationOffsetBlendSpace.cpp b/Engine/Source/Runtime/Engine/Private/Animation/AnimNode_RotationOffsetBlendSpace.cpp
--- a/Engine/Source/Runtime/Engine/Private/Animation/AnimNode_RotationOffsetBlendSpace.cpp
+++ b/Engine/Source/Runtime/Engine/Private/Animation/AnimNode_RotationOffsetBlendSpace.cpp
@@ -18,9 +18,9 @@ void FAnimNode_RotationOffsetBlendSpace::CacheBones(const FAnimationCacheBonesCo
 	BasePose.CacheBones(Context);
 }
 
-void FAnimNode_RotationOffsetBlendSpace::Update(const FAnimationUpdateContext& Context)
+void FAnimNode_RotationOffsetBlendSpace::UpdateAssetPlayer(const FAnimationUpdateContext& Context)
 {
-	FAnimNode_BlendSpacePlayer::Update(Context);
+	FAnimNode_BlendSpacePlayer::UpdateAssetPlayer(Context);
 	BasePose.Update(Context);
 }
 
diff --git a/Engine/Source/Runtime/Engine/Private/Animation/AnimNode_SequenceEvaluator.cpp b/Engine/Source/Runtime/Engine/Private/Animation/AnimNode_SequenceEvaluator.cpp
--- a/Engine/Source/Runtime/Engine/Private/Animation/AnimNode_SequenceEvaluator.cpp
+++ b/Engine/Source/Runtime/Engine/Private/Animation/AnimNode_SequenceEvaluator.cpp
@@ -14,7 +14,7 @@ void FAnimNode_SequenceEvaluator::CacheBones(const FAnimationCacheBonesContext&
 {
 }
 
-void FAnimNode_SequenceEvaluator::Update(const FAnimationUpdateContext& Context)
+void FAnimNode_SequenceEvaluator::UpdateAssetPlayer(const FAnimationUpdateContext& Context)
 {
 	EvaluateGraphExposedInputs.Execute(Context);
 }
diff --git a/Engine/Source/Runtime/Engine/Private/Animation/AnimNode_SequencePlayer.cpp b/Engine/Source/Runtime/Engine/Private/Animation/AnimNode_SequencePlayer.cpp
--- a/Engine/Source/Runtime/Engine/Private/Animation/AnimNode_SequencePlayer.cpp
+++ b/Engine/Source/Runtime/Engine/Private/Animation/AnimNode_SequencePlayer.cpp
@@ -24,7 +24,7 @@ void FAnimNode_SequencePlayer::CacheBones(const FAnimationCacheBonesContext& Con
 {
 }
 
-void FAnimNode_SequencePlayer::Update(const FAnimationUpdateContext& Context)
+void FAnimNode_SequencePlayer::UpdateAssetPlayer(const FAnimationUpdateContext& Context)
 {
 	EvaluateGraphExposedInputs.Execute(Context);
 
@@ -72,4 +72,9 @@ void FAnimNode_SequencePlayer::GatherDebugData(FNodeDebugData& DebugData)
 	
 	DebugLine += FString::Printf(TEXT("('%s' Play Time: %.3f)"), *Sequence->GetName(), InternalTimeAccumulator);
 	DebugData.AddDebugItem(DebugLine, true);
-}
\ No newline at end of file
+}
+
+float FAnimNode_SequencePlayer::GetTimeFromEnd(float CurrentNodeTime)
+{
+	return Sequence->GetMaxCurrentTime() - CurrentNodeTime;
+}
diff --git a/Engine/Source/Runtime/Engine/Private/Animation/AnimNode_StateMachine.cpp b/Engine/Source/Runtime/Engine/Private/Animation/AnimNode_StateMachine.cpp
--- a/Engine/Source/Runtime/Engine/Private/Animation/AnimNode_StateMachine.cpp
+++ b/Engine/Source/Runtime/Engine/Private/Animation/AnimNode_StateMachine.cpp
@@ -466,7 +466,7 @@ bool FAnimNode_StateMachine::FindValidTransition(const FAnimationUpdateContext&
 				if ((SequencePlayer != nullptr) && (SequencePlayer->Sequence != nullptr))
 				{
 					const float SequenceLength = SequencePlayer->Sequence->GetMaxCurrentTime();
-					const float PlayerTime = SequencePlayer->InternalTimeAccumulator;
+					const float PlayerTime = SequencePlayer->GetAccumulatedTime();
 					const float PlayerTimeLeft = SequenceLength - PlayerTime;
 
 					const FAnimationTransitionBetweenStates& TransitionInfo = GetTransitionInfo(TransitionRule.TransitionIndex);
@@ -883,3 +883,8 @@ bool FAnimNode_StateMachine::IsAConduitState(int32 StateIndex) const
 	return ((PRIVATE_MachineDescription != NULL) && (StateIndex < PRIVATE_MachineDescription->States.Num())) ? GetStateInfo(StateIndex).bIsAConduit : false;
 }
 
+bool FAnimNode_StateMachine::IsValidTransitionIndex(int32 TransitionIndex) const
+{
+	return PRIVATE_MachineDescription->Transitions.IsValidIndex(TransitionIndex);
+}
+
