Animation: Cleaned up UAnimGraphNode_Base a bit (added a class comment, full DLL export, NULL->nullptr, forward declarations)

diff --git a/Engine/Source/Editor/AnimGraph/Classes/AnimGraphNode_Base.h b/Engine/Source/Editor/AnimGraph/Classes/AnimGraphNode_Base.h
--- a/Engine/Source/Editor/AnimGraph/Classes/AnimGraphNode_Base.h
+++ b/Engine/Source/Editor/AnimGraph/Classes/AnimGraphNode_Base.h
@@ -1,19 +1,28 @@
 // Copyright 1998-2015 Epic Games, Inc. All Rights Reserved.
 
 #pragma once
+
 #include "K2Node.h"
 #include "Animation/AnimNodeBase.h"
 #include "AnimGraphNode_Base.generated.h"
 
-//
-// Forward declarations.
-//
-struct FEdGraphSchemaAction_K2NewNode;
+// Forward declarations
 
+struct FEdGraphSchemaAction_K2NewNode;
+struct FPropertyChangedEvent;
+class FCompilerResultsLog;
+class UAnimGraphNode_Base;
+class USkeleton;
+class UAnimBlueprintGeneratedClass;
+class IDetailLayoutBuilder;
+class FAnimBlueprintCompiler;
+class FAnimGraphNodeDetails;
+
+// 
 struct FPoseLinkMappingRecord
 {
 public:
-	static FPoseLinkMappingRecord MakeFromArrayEntry(class UAnimGraphNode_Base* LinkingNode, class UAnimGraphNode_Base* LinkedNode, UArrayProperty* ArrayProperty, int32 ArrayIndex)
+	static FPoseLinkMappingRecord MakeFromArrayEntry(UAnimGraphNode_Base* LinkingNode, UAnimGraphNode_Base* LinkedNode, UArrayProperty* ArrayProperty, int32 ArrayIndex)
 	{
 		checkSlow(CastChecked<UStructProperty>(ArrayProperty->Inner)->Struct->IsChildOf(FPoseLinkBase::StaticStruct()));
 
@@ -26,7 +35,7 @@ public:
 		return Result;
 	}
 
-	static FPoseLinkMappingRecord MakeFromMember(class UAnimGraphNode_Base* LinkingNode, class UAnimGraphNode_Base* LinkedNode, UStructProperty* MemberProperty)
+	static FPoseLinkMappingRecord MakeFromMember(UAnimGraphNode_Base* LinkingNode, UAnimGraphNode_Base* LinkedNode, UStructProperty* MemberProperty)
 	{
 		checkSlow(MemberProperty->Struct->IsChildOf(FPoseLinkBase::StaticStruct()));
 
@@ -47,15 +56,15 @@ public:
 
 	bool IsValid() const
 	{
-		return LinkedNode != NULL;
+		return LinkedNode != nullptr;
 	}
 
-	class UAnimGraphNode_Base* GetLinkedNode() const
+	UAnimGraphNode_Base* GetLinkedNode() const
 	{
 		return LinkedNode;
 	}
 
-	class UAnimGraphNode_Base* GetLinkingNode() const
+	UAnimGraphNode_Base* GetLinkingNode() const
 	{
 		return LinkingNode;
 	}
@@ -63,19 +72,19 @@ public:
 	ANIMGRAPH_API void PatchLinkIndex(uint8* DestinationPtr, int32 LinkID, int32 SourceLinkID) const;
 protected:
 	FPoseLinkMappingRecord()
-		: LinkedNode(NULL)
-		, LinkingNode(NULL)
-		, ChildProperty(NULL)
+		: LinkedNode(nullptr)
+		, LinkingNode(nullptr)
+		, ChildProperty(nullptr)
 		, ChildPropertyIndex(INDEX_NONE)
 	{
 	}
 
 protected:
-	// Linked node for this pose link, can be NULL
-	class UAnimGraphNode_Base* LinkedNode;
+	// Linked node for this pose link, can be nullptr
+	UAnimGraphNode_Base* LinkedNode;
 
-	// Linking node for this pose link, can be NULL
-	class UAnimGraphNode_Base* LinkingNode;
+	// Linking node for this pose link, can be nullptr
+	UAnimGraphNode_Base* LinkingNode;
 
 	// Will either be an array property containing FPoseLinkBase derived structs, indexed by ChildPropertyIndex, or a FPoseLinkBase derived struct property 
 	UProperty* ChildProperty;
@@ -84,9 +93,14 @@ protected:
 	int32 ChildPropertyIndex;
 };
 
-
-UCLASS(Abstract, MinimalAPI)
-class UAnimGraphNode_Base : public UK2Node
+/**
+  * This is the base class for any animation graph nodes that generate or consume an animation pose in
+  * the animation blend graph.
+  *
+  * Any concrete implementations will be paired with a runtime graph node derived from FAnimNode_Base
+  */
+UCLASS(Abstract)
+class ANIMGRAPH_API UAnimGraphNode_Base : public UK2Node
 {
 	GENERATED_UCLASS_BODY()
 
@@ -94,26 +108,26 @@ class UAnimGraphNode_Base : public UK2Node
 	TArray<FOptionalPinFromProperty> ShowPinForProperties;
 
 	// UObject interface
-	ANIMGRAPH_API virtual void PostEditChangeProperty(struct FPropertyChangedEvent& PropertyChangedEvent) override;
+	virtual void PostEditChangeProperty(FPropertyChangedEvent& PropertyChangedEvent) override;
 	// End of UObject interface
 
 	// UEdGraphNode interface
-	ANIMGRAPH_API virtual void AllocateDefaultPins() override;
-	ANIMGRAPH_API virtual FLinearColor GetNodeTitleColor() const override;
-	ANIMGRAPH_API virtual FString GetDocumentationLink() const override;
-	ANIMGRAPH_API virtual void GetPinHoverText(const UEdGraphPin& Pin, FString& HoverTextOut) const override;
-	ANIMGRAPH_API virtual bool ShowPaletteIconOnNode() const override{ return false; }
+	virtual void AllocateDefaultPins() override;
+	virtual FLinearColor GetNodeTitleColor() const override;
+	virtual FString GetDocumentationLink() const override;
+	virtual void GetPinHoverText(const UEdGraphPin& Pin, FString& HoverTextOut) const override;
+	virtual bool ShowPaletteIconOnNode() const override{ return false; }
 	// End of UEdGraphNode interface
 
 	// UK2Node interface
 	virtual bool NodeCausesStructuralBlueprintChange() const override { return true; }
 	virtual bool ShouldShowNodeProperties() const override { return true; }
 	virtual bool CanPlaceBreakpoints() const override { return false; }
-	ANIMGRAPH_API virtual void ReallocatePinsDuringReconstruction(TArray<UEdGraphPin*>& OldPins) override;
-	ANIMGRAPH_API virtual bool CanCreateUnderSpecifiedSchema(const UEdGraphSchema* DesiredSchema) const override;
-	ANIMGRAPH_API virtual void GetNodeAttributes(TArray<TKeyValuePair<FString, FString>>& OutNodeAttributes) const override;
-	ANIMGRAPH_API virtual void GetMenuActions(FBlueprintActionDatabaseRegistrar& ActionRegistrar) const override;
-	ANIMGRAPH_API virtual FText GetMenuCategory() const override;
+	virtual void ReallocatePinsDuringReconstruction(TArray<UEdGraphPin*>& OldPins) override;
+	virtual bool CanCreateUnderSpecifiedSchema(const UEdGraphSchema* DesiredSchema) const override;
+	virtual void GetNodeAttributes(TArray<TKeyValuePair<FString, FString>>& OutNodeAttributes) const override;
+	virtual void GetMenuActions(FBlueprintActionDatabaseRegistrar& ActionRegistrar) const override;
+	virtual FText GetMenuCategory() const override;
 
 	// By default return any animation assets we have
 	virtual UObject* GetJumpTargetForDoubleClick() const override { return GetAnimationAsset(); }
@@ -122,31 +136,31 @@ class UAnimGraphNode_Base : public UK2Node
 	// UAnimGraphNode_Base interface
 
 	// Gets the menu category this node belongs in
-	ANIMGRAPH_API virtual FString GetNodeCategory() const;
+	virtual FString GetNodeCategory() const;
 
 	// Is this node a sink that has no pose outputs?
 	virtual bool IsSinkNode() const { return false; }
 
 	// Create any output pins necessary for this node
-	ANIMGRAPH_API virtual void CreateOutputPins();
+	virtual void CreateOutputPins();
 
 	// customize pin data based on the input
 	virtual void CustomizePinData(UEdGraphPin* Pin, FName SourcePropertyName, int32 ArrayIndex) const {}
 
 	// Gives each visual node a chance to do final validation before it's node is harvested for use at runtime
-	virtual void ValidateAnimNodeDuringCompilation(class USkeleton* ForSkeleton, class FCompilerResultsLog& MessageLog) {}
+	virtual void ValidateAnimNodeDuringCompilation(USkeleton* ForSkeleton, FCompilerResultsLog& MessageLog) {}
 
 	// Gives each visual node a chance to validate that they are still valid in the context of the compiled class, giving a last shot at error or warning generation after primary compilation is finished
-	virtual void ValidateAnimNodePostCompile(class FCompilerResultsLog& MessageLog, class UAnimBlueprintGeneratedClass* CompiledClass, int32 CompiledNodeIndex) {}
+	virtual void ValidateAnimNodePostCompile(FCompilerResultsLog& MessageLog, UAnimBlueprintGeneratedClass* CompiledClass, int32 CompiledNodeIndex) {}
 
 	// Gives each visual node a chance to update the node template before it is inserted in the compiled class
-	virtual void BakeDataDuringCompilation(class FCompilerResultsLog& MessageLog) {}
+	virtual void BakeDataDuringCompilation(FCompilerResultsLog& MessageLog) {}
 
 	// preload asset required for this node in this function
 	virtual void PreloadRequiredAssets() {}
 
 	// Give the node a chance to change the display name of a pin
-	ANIMGRAPH_API virtual void PostProcessPinName(const UEdGraphPin* Pin, FString& DisplayName) const;
+	virtual void PostProcessPinName(const UEdGraphPin* Pin, FString& DisplayName) const;
 
 	/** Get the animation blueprint to which this node belongs */
 	UAnimBlueprint* GetAnimBlueprint() const { return CastChecked<UAnimBlueprint>(GetBlueprint()); }
@@ -168,57 +182,55 @@ class UAnimGraphNode_Base : public UK2Node
 	// if you return true for DoesSupportExposeTimeForTransitionGetter
 	// you should implement all below functions
 	virtual bool DoesSupportTimeForTransitionGetter() const { return false; }
-	virtual UAnimationAsset* GetAnimationAsset() const { return NULL; }
-	virtual const TCHAR* GetTimePropertyName() const { return NULL; }
-	virtual UScriptStruct* GetTimePropertyStruct() const { return NULL; }
+	virtual UAnimationAsset* GetAnimationAsset() const { return nullptr; }
+	virtual const TCHAR* GetTimePropertyName() const { return nullptr; }
+	virtual UScriptStruct* GetTimePropertyStruct() const { return nullptr; }
 	// END Interface to support transition getter
 
 	// can customize details tab 
-	virtual void CustomizeDetails(class IDetailLayoutBuilder& DetailBuilder){ }
+	virtual void CustomizeDetails(IDetailLayoutBuilder& DetailBuilder){ }
 
 protected:
-	friend class FAnimBlueprintCompiler;
-	friend class FAnimGraphNodeDetails;
+	friend FAnimBlueprintCompiler;
+	friend FAnimGraphNodeDetails;
 
 	// Gets the animation FNode type represented by this ed graph node
-	ANIMGRAPH_API UScriptStruct* GetFNodeType() const;
+	UScriptStruct* GetFNodeType() const;
 
 	// Gets the animation FNode property represented by this ed graph node
-	ANIMGRAPH_API UStructProperty* GetFNodeProperty() const;
+	UStructProperty* GetFNodeProperty() const;
 
 	// This will be called when a pose link is found, and can be called with PoseProperty being either of:
 	//  - an array property (ArrayIndex >= 0)
 	//  - a single pose property (ArrayIndex == INDEX_NONE)
-	virtual ANIMGRAPH_API void CreatePinsForPoseLink(UProperty* PoseProperty, int32 ArrayIndex);
+	virtual void CreatePinsForPoseLink(UProperty* PoseProperty, int32 ArrayIndex);
 
 	//
-	virtual ANIMGRAPH_API FPoseLinkMappingRecord GetLinkIDLocation(const UScriptStruct* NodeType, UEdGraphPin* InputLinkPin);
+	virtual FPoseLinkMappingRecord GetLinkIDLocation(const UScriptStruct* NodeType, UEdGraphPin* InputLinkPin);
 
 	//
-	virtual ANIMGRAPH_API void GetPinAssociatedProperty(const UScriptStruct* NodeType, UEdGraphPin* InputPin, UProperty*& OutProperty, int32& OutIndex);
+	virtual void GetPinAssociatedProperty(const UScriptStruct* NodeType, UEdGraphPin* InputPin, UProperty*& OutProperty, int32& OutIndex);
 
 	// Allocates or reallocates pins
-	ANIMGRAPH_API void InternalPinCreation(TArray<UEdGraphPin*>* OldPins);
+	void InternalPinCreation(TArray<UEdGraphPin*>* OldPins);
 };
 
 template<class AssetType>
 void UAnimGraphNode_Base::HandleAnimReferenceReplacement(AssetType*& OriginalAsset, const TMap<UAnimationAsset*, UAnimationAsset*>& ComplexAnimsMap, const TMap<UAnimSequence*, UAnimSequence*>& AnimSequenceMap)
 {
 	AssetType* CacheOriginalAsset = OriginalAsset;
-	OriginalAsset = NULL;
+	OriginalAsset = nullptr;
 
-	if(UAnimSequence* AnimSequence = Cast<UAnimSequence>(CacheOriginalAsset))
+	if (UAnimSequence* AnimSequence = Cast<UAnimSequence>(CacheOriginalAsset))
 	{
-		UAnimSequence* const* ReplacementAsset = AnimSequenceMap.Find(AnimSequence);
-		if(ReplacementAsset)
+		if (UAnimSequence* const* ReplacementAsset = AnimSequenceMap.Find(AnimSequence))
 		{
 			OriginalAsset = Cast<AssetType>(*ReplacementAsset);
 		}
 	}
-	else if(UAnimationAsset* AnimAsset = Cast<UAnimationAsset>(CacheOriginalAsset))
+	else if (UAnimationAsset* AnimAsset = Cast<UAnimationAsset>(CacheOriginalAsset))
 	{
-		UAnimationAsset* const* ReplacementAsset = ComplexAnimsMap.Find(AnimAsset);
-		if(ReplacementAsset)
+		if (UAnimationAsset* const* ReplacementAsset = ComplexAnimsMap.Find(AnimAsset))
 		{
 			OriginalAsset = Cast<AssetType>(*ReplacementAsset);
 		}
