Animation: Added compilation warnings to UAnimGraphNode_ModifyBone when it will do nothing

diff --git a/Engine/Source/Editor/AnimGraph/Classes/AnimGraphNode_ModifyBone.h b/Engine/Source/Editor/AnimGraph/Classes/AnimGraphNode_ModifyBone.h
--- a/Engine/Source/Editor/AnimGraph/Classes/AnimGraphNode_ModifyBone.h
+++ b/Engine/Source/Editor/AnimGraph/Classes/AnimGraphNode_ModifyBone.h
@@ -22,6 +22,10 @@ public:
 	// End of UEdGraphNode interface
 
 protected:
+	// UAnimGraphNode_Base interface
+	virtual void ValidateAnimNodeDuringCompilation(USkeleton* ForSkeleton, FCompilerResultsLog& MessageLog) override;
+	// End of UAnimGraphNode_Base interface
+
 	// UAnimGraphNode_SkeletalControlBase interface
 	virtual FText GetControllerDescription() const override;
 	virtual FVector GetWidgetLocation(const USkeletalMeshComponent* SkelComp, struct FAnimNode_SkeletalControlBase* AnimNode) override;
diff --git a/Engine/Source/Editor/AnimGraph/Private/AnimGraphNode_ModifyBone.cpp b/Engine/Source/Editor/AnimGraph/Private/AnimGraphNode_ModifyBone.cpp
--- a/Engine/Source/Editor/AnimGraph/Private/AnimGraphNode_ModifyBone.cpp
+++ b/Engine/Source/Editor/AnimGraph/Private/AnimGraphNode_ModifyBone.cpp
@@ -2,6 +2,7 @@
 
 #include "AnimGraphPrivatePCH.h"
 #include "AnimGraphNode_ModifyBone.h"
+#include "CompilerResultsLog.h"
 
 /////////////////////////////////////////////////////
 // UAnimGraphNode_ModifyBone
@@ -14,6 +15,21 @@ UAnimGraphNode_ModifyBone::UAnimGraphNode_ModifyBone(const FObjectInitializer& O
 	CurWidgetMode = (int32)FWidget::WM_Rotate;
 }
 
+void UAnimGraphNode_ModifyBone::ValidateAnimNodeDuringCompilation(USkeleton* ForSkeleton, FCompilerResultsLog& MessageLog)
+{
+	if (ForSkeleton->GetReferenceSkeleton().FindBoneIndex(Node.BoneToModify.BoneName) == INDEX_NONE)
+	{
+		MessageLog.Warning(*LOCTEXT("NoBoneToModify", "@@ - You must pick a bone to modify").ToString(), this);
+	}
+
+	if ((Node.TranslationMode == BMM_Ignore) && (Node.RotationMode == BMM_Ignore) && (Node.ScaleMode == BMM_Ignore))
+	{
+		MessageLog.Warning(*LOCTEXT("NothingToModify", "@@ - No components to modify selected.  Either Rotation, Translation, or Scale should be set to something other than Ignore").ToString(), this);
+	}
+
+	Super::ValidateAnimNodeDuringCompilation(ForSkeleton, MessageLog);
+}
+
 FText UAnimGraphNode_ModifyBone::GetControllerDescription() const
 {
 	return LOCTEXT("TransformModifyBone", "Transform (Modify) Bone");
@@ -53,7 +69,7 @@ FText UAnimGraphNode_ModifyBone::GetNodeTitle(ENodeTitleType::Type TitleType) co
 
 FVector UAnimGraphNode_ModifyBone::GetWidgetLocation(const USkeletalMeshComponent* SkelComp, FAnimNode_SkeletalControlBase* AnimNode)
 {
-	USkeleton * Skeleton = SkelComp->SkeletalMesh->Skeleton;
+	USkeleton* Skeleton = SkelComp->SkeletalMesh->Skeleton;
 	FVector WidgetLoc = FVector::ZeroVector;
 
 	// if the current widget mode is translate, then shows the widget according to translation space
