#ANIM: Fixed additive layer not working properly

- It is broken whenever you apply any bone transform due to controller is not copying local atoms back to it.

diff --git a/Engine/Source/Editor/AnimGraph/Private/AnimPreviewInstance.cpp b/Engine/Source/Editor/AnimGraph/Private/AnimPreviewInstance.cpp
--- a/Engine/Source/Editor/AnimGraph/Private/AnimPreviewInstance.cpp
+++ b/Engine/Source/Editor/AnimGraph/Private/AnimPreviewInstance.cpp
@@ -149,16 +149,16 @@ bool UAnimPreviewInstance::NativeEvaluateAnimation(FPoseContext& Output)
 			// create bone controllers from 
 			if(BoneControllers.Num() > 0 || CurveBoneControllers.Num() > 0)
 			{
-				FCompactPose PreController, PostController;
+				FPoseContext PreController(Output), PostController(Output);
 				// if set key is true, we should save pre controller local space transform 
 				// so that we can calculate the delta correctly
 				if(bSetKey)
 				{
-					PreController = Output.Pose;
+					PreController = Output;
 				}
 
 				FCSPose<FCompactPose> OutMeshPose;
-				OutMeshPose.InitPose(&RequiredBones);
+				OutMeshPose.InitPose(Output.Pose);
 
 				// apply curve data first
 				ApplyBoneControllers(Component, CurveBoneControllers, OutMeshPose);
@@ -173,8 +173,8 @@ bool UAnimPreviewInstance::NativeEvaluateAnimation(FPoseContext& Output)
 				if(bSetKey)
 				{
 					// now we have post controller, and calculate delta now
-					PostController = Output.Pose;
-					SetKeyImplementation(PreController, PostController);
+					PostController = Output;
+					SetKeyImplementation(PreController.Pose, PostController.Pose);
 				}
 			}
 			// if any other bone is selected, still go for set key even if nothing changed
