Animation: Created an Observe Bone skeletal control that can be used to debug animations (prints out the current transform of a bone in the input pose as a bubble above the node)

diff --git a/Engine/Source/Editor/AnimGraph/AnimGraph.Build.cs b/Engine/Source/Editor/AnimGraph/AnimGraph.Build.cs
--- a/Engine/Source/Editor/AnimGraph/AnimGraph.Build.cs
+++ b/Engine/Source/Editor/AnimGraph/AnimGraph.Build.cs
@@ -22,10 +22,12 @@ public class AnimGraph : ModuleRules
 
 		PrivateDependencyModuleNames.AddRange(
 			new string[] {
+				"InputCore",
 				"SlateCore",
 				"UnrealEd",
                 "GraphEditor",
-				"PropertyEditor"
+				"PropertyEditor",
+				"EditorStyle"
 			}
 		);
 
diff --git a/Engine/Source/Editor/AnimGraph/Private/AnimGraphNode_ObserveBone.cpp b/Engine/Source/Editor/AnimGraph/Private/AnimGraphNode_ObserveBone.cpp
new file mode 100644
--- /dev/null
+++ b/Engine/Source/Editor/AnimGraph/Private/AnimGraphNode_ObserveBone.cpp
@@ -0,0 +1,195 @@
+// Copyright 1998-2015 Epic Games, Inc. All Rights Reserved.
+
+#include "AnimGraphPrivatePCH.h"
+#include "AnimGraphNode_ObserveBone.h"
+#include "CompilerResultsLog.h"
+#include "SGraphNode.h"
+#include "KismetNodeInfoContext.h"
+#include "KismetDebugUtilities.h"
+
+#define LOCTEXT_NAMESPACE "ObserveBone"
+
+/////////////////////////////////////////////////////
+// SGraphNodeObserveBone
+
+class SGraphNodeObserveBone : public SGraphNode
+{
+public:
+	SLATE_BEGIN_ARGS(SGraphNodeObserveBone){}
+	SLATE_END_ARGS()
+
+	void Construct(const FArguments& InArgs, UAnimGraphNode_ObserveBone* InNode)
+	{
+		GraphNode = InNode;
+
+		SetCursor(EMouseCursor::CardinalCross);
+
+		UpdateGraphNode();
+	}
+
+	// SNodePanel::SNode interface
+	void GetNodeInfoPopups(FNodeInfoContext* Context, TArray<FGraphInformationPopupInfo>& Popups) const override;
+	// End of SNodePanel::SNode interface
+
+	// SGraphNode interface
+	virtual void UpdateGraphNode() override
+	{
+		SGraphNode::UpdateGraphNode();
+
+		// Prevent the comment bubble from being displayed
+		GetOrAddSlot(ENodeZone::TopCenter)
+		[
+			SNullWidget::NullWidget
+		];
+	}
+	// End of SGraphNode interface
+
+	static FString PrettyVectorToString(const FVector& Vector, const FString& PerComponentPrefix);
+};
+
+void SGraphNodeObserveBone::GetNodeInfoPopups(FNodeInfoContext* Context, TArray<FGraphInformationPopupInfo>& Popups) const
+{
+	const FLinearColor TimelineBubbleColor(0.7f, 0.5f, 0.5f);
+	FKismetNodeInfoContext* K2Context = (FKismetNodeInfoContext*)Context;
+
+	// Display the observed bone transform status bubble
+ 	if (UObject* ActiveObject = K2Context->ActiveObjectBeingDebugged)
+ 	{
+		UProperty* NodeProperty = FKismetDebugUtilities::FindClassPropertyForNode(K2Context->SourceBlueprint, GraphNode);
+		if (UStructProperty* StructProperty = Cast<UStructProperty>(NodeProperty))
+ 		{
+ 			UClass* ContainingClass = StructProperty->GetTypedOuter<UClass>();
+ 			if (ActiveObject->IsA(ContainingClass) && (StructProperty->Struct == FAnimNode_ObserveBone::StaticStruct()))
+ 			{
+				FAnimNode_ObserveBone* ObserveBone = StructProperty->ContainerPtrToValuePtr<FAnimNode_ObserveBone>(ActiveObject);
+
+				check(ObserveBone);
+				const FString Message = FString::Printf(TEXT("%s\n%s\n%s"),
+					*PrettyVectorToString(ObserveBone->Translation, TEXT("T")),
+					*PrettyVectorToString(ObserveBone->Rotation.Euler(), TEXT("R")),
+					*PrettyVectorToString(ObserveBone->Scale, TEXT("S")));
+
+				new (Popups) FGraphInformationPopupInfo(NULL, TimelineBubbleColor, Message);
+			}
+ 			else
+ 			{
+				const FString ErrorText = FString::Printf(*LOCTEXT("StaleDebugData", "Stale debug data\nProperty is on %s\nDebugging a %s").ToString(), *ContainingClass->GetName(), *ActiveObject->GetClass()->GetName());
+				new (Popups) FGraphInformationPopupInfo(NULL, TimelineBubbleColor, ErrorText);
+ 			}
+ 		}
+	}
+
+	SGraphNode::GetNodeInfoPopups(Context, Popups);
+}
+
+FString SGraphNodeObserveBone::PrettyVectorToString(const FVector& Vector, const FString& PerComponentPrefix)
+{
+	const FString ReturnString = FString::Printf(TEXT("%sX=%.2f, %sY=%.2f, %sZ=%.2f"), *PerComponentPrefix, Vector.X, *PerComponentPrefix, Vector.Y, *PerComponentPrefix, Vector.Z);
+	return ReturnString;
+}
+
+/////////////////////////////////////////////////////
+// UAnimGraphNode_ObserveBone
+
+TSharedPtr<SGraphNode> UAnimGraphNode_ObserveBone::CreateVisualWidget()
+{
+	return SNew(SGraphNodeObserveBone, this);
+}
+
+UAnimGraphNode_ObserveBone::UAnimGraphNode_ObserveBone(const FObjectInitializer& ObjectInitializer)
+	: Super(ObjectInitializer)
+{
+}
+
+void UAnimGraphNode_ObserveBone::ValidateAnimNodeDuringCompilation(USkeleton* ForSkeleton, FCompilerResultsLog& MessageLog)
+{
+	if (ForSkeleton->GetReferenceSkeleton().FindBoneIndex(Node.BoneToObserve.BoneName) == INDEX_NONE)
+	{
+		MessageLog.Warning(*LOCTEXT("NoBoneToObserve", "@@ - You must pick a bone to observe").ToString(), this);
+	}
+
+	Super::ValidateAnimNodeDuringCompilation(ForSkeleton, MessageLog);
+}
+
+FText UAnimGraphNode_ObserveBone::GetControllerDescription() const
+{
+	return LOCTEXT("AnimGraphNode_ObserveBone", "Observe Bone");
+}
+
+FText UAnimGraphNode_ObserveBone::GetTooltipText() const
+{
+	return LOCTEXT("AnimGraphNode_ObserveBone_Tooltip", "Observes a bone for debugging purposes");
+}
+
+FText UAnimGraphNode_ObserveBone::GetNodeTitle(ENodeTitleType::Type TitleType) const
+{
+	if (((TitleType == ENodeTitleType::ListView) || (TitleType == ENodeTitleType::MenuTitle)) && (Node.BoneToObserve.BoneName == NAME_None))
+	{
+		return GetControllerDescription();
+	}
+	else
+	{
+		FFormatNamedArguments Args;
+		Args.Add(TEXT("ControllerDescription"), GetControllerDescription());
+		Args.Add(TEXT("BoneName"), FText::FromName(Node.BoneToObserve.BoneName));
+
+		return FText::Format(LOCTEXT("AnimGraphNode_ObserveBone_Title", "{ControllerDescription}: {BoneName}"), Args);
+	}
+}
+
+int32 UAnimGraphNode_ObserveBone::GetWidgetCoordinateSystem(const USkeletalMeshComponent* SkelComp)
+{
+	switch (Node.DisplaySpace)
+	{
+	default:
+	case BCS_ParentBoneSpace:
+		//@TODO: No good way of handling this one
+		return COORD_World;
+	case BCS_BoneSpace:
+		return COORD_Local;
+	case BCS_ComponentSpace:
+	case BCS_WorldSpace:
+		return COORD_World;
+	}
+}
+
+FVector UAnimGraphNode_ObserveBone::GetWidgetLocation(const USkeletalMeshComponent* SkelComp, FAnimNode_SkeletalControlBase* AnimNode)
+{
+	USkeleton* Skeleton = SkelComp->SkeletalMesh->Skeleton;
+	FVector WidgetLoc = FVector::ZeroVector;
+
+	int32 MeshBoneIndex = SkelComp->GetBoneIndex(Node.BoneToObserve.BoneName);
+
+	if (MeshBoneIndex != INDEX_NONE)
+	{
+		const FTransform BoneTM = SkelComp->GetBoneTransform(MeshBoneIndex);
+		WidgetLoc = BoneTM.GetLocation();
+	}
+	
+	return WidgetLoc;
+}
+
+int32 UAnimGraphNode_ObserveBone::GetWidgetMode(const USkeletalMeshComponent* SkelComp)
+{
+	return (int32)FWidget::WM_Translate;
+}
+
+FName UAnimGraphNode_ObserveBone::FindSelectedBone()
+{
+	return Node.BoneToObserve.BoneName;
+}
+
+FLinearColor UAnimGraphNode_ObserveBone::GetNodeTitleColor() const
+{
+	return FLinearColor(0.7f, 0.7f, 0.7f);
+}
+
+void UAnimGraphNode_ObserveBone::CustomizePinData(UEdGraphPin* Pin, FName SourcePropertyName, int32 ArrayIndex) const
+{
+	if (Pin->PinName == GET_MEMBER_NAME_STRING_CHECKED(FAnimNode_SkeletalControlBase, Alpha))
+	{
+		Pin->bHidden = true;
+	}
+}
+
+#undef LOCTEXT_NAMESPACE
diff --git a/Engine/Source/Editor/AnimGraph/Public/AnimGraphNode_ObserveBone.h b/Engine/Source/Editor/AnimGraph/Public/AnimGraphNode_ObserveBone.h
new file mode 100644
--- /dev/null
+++ b/Engine/Source/Editor/AnimGraph/Public/AnimGraphNode_ObserveBone.h
@@ -0,0 +1,40 @@
+// Copyright 1998-2015 Epic Games, Inc. All Rights Reserved.
+
+#pragma once
+
+#include "AnimGraphNode_SkeletalControlBase.h"
+#include "BoneControllers/AnimNode_ObserveBone.h"
+#include "EdGraph/EdGraphNodeUtils.h" // for FNodeTitleTextTable
+#include "AnimGraphNode_ObserveBone.generated.h"
+
+UCLASS()
+class ANIMGRAPH_API UAnimGraphNode_ObserveBone : public UAnimGraphNode_SkeletalControlBase
+{
+	GENERATED_UCLASS_BODY()
+
+	UPROPERTY(EditAnywhere, Category=Settings)
+	FAnimNode_ObserveBone Node;
+
+public:
+	// UEdGraphNode interface
+	virtual FText GetNodeTitle(ENodeTitleType::Type TitleType) const override;
+	virtual FText GetTooltipText() const override;
+	virtual TSharedPtr<SGraphNode> CreateVisualWidget() override;
+	virtual FLinearColor GetNodeTitleColor() const override;
+	// End of UEdGraphNode interface
+
+protected:
+	// UAnimGraphNode_Base interface
+	virtual void ValidateAnimNodeDuringCompilation(USkeleton* ForSkeleton, FCompilerResultsLog& MessageLog) override;
+	virtual void CustomizePinData(UEdGraphPin* Pin, FName SourcePropertyName, int32 ArrayIndex) const override;
+	// End of UAnimGraphNode_Base interface
+
+	// UAnimGraphNode_SkeletalControlBase interface
+	virtual FText GetControllerDescription() const override;
+	virtual int32 GetWidgetCoordinateSystem(const USkeletalMeshComponent* SkelComp) override;
+	virtual FVector GetWidgetLocation(const USkeletalMeshComponent* SkelComp, struct FAnimNode_SkeletalControlBase* AnimNode) override;
+	virtual int32 GetWidgetMode(const USkeletalMeshComponent* SkelComp) override;
+	virtual FName FindSelectedBone() override;
+	// End of UAnimGraphNode_SkeletalControlBase interface
+};
+
diff --git a/Engine/Source/Runtime/AnimGraphRuntime/Private/BoneControllers/AnimNode_ObserveBone.cpp b/Engine/Source/Runtime/AnimGraphRuntime/Private/BoneControllers/AnimNode_ObserveBone.cpp
new file mode 100644
--- /dev/null
+++ b/Engine/Source/Runtime/AnimGraphRuntime/Private/BoneControllers/AnimNode_ObserveBone.cpp
@@ -0,0 +1,50 @@
+// Copyright 1998-2015 Epic Games, Inc. All Rights Reserved.
+
+#include "AnimGraphRuntimePrivatePCH.h"
+#include "BoneControllers/AnimNode_ObserveBone.h"
+#include "AnimationRuntime.h"
+
+/////////////////////////////////////////////////////
+// FAnimNode_ObserveBone
+
+FAnimNode_ObserveBone::FAnimNode_ObserveBone()
+	: DisplaySpace(BCS_ComponentSpace)
+	, Translation(FVector::ZeroVector)
+	, Rotation(FRotator::ZeroRotator)
+	, Scale(FVector(1.0f))
+{
+}
+
+void FAnimNode_ObserveBone::GatherDebugData(FNodeDebugData& DebugData)
+{
+	const FString DebugLine = FString::Printf(TEXT("(Bone: %s has T(%s), R(%s), S(%s))"), *BoneToObserve.BoneName.ToString(), *Translation.ToString(), *Rotation.Euler().ToString(), *Scale.ToString());
+
+	DebugData.AddDebugItem(DebugLine);
+
+	ComponentPose.GatherDebugData(DebugData);
+}
+
+void FAnimNode_ObserveBone::EvaluateBoneTransforms(USkeletalMeshComponent* SkelComp, FCSPose<FCompactPose>& MeshBases, TArray<FBoneTransform>& OutBoneTransforms)
+{
+	const FBoneContainer BoneContainer = MeshBases.GetPose().GetBoneContainer();
+
+	const FCompactPoseBoneIndex BoneIndex = BoneToObserve.GetCompactPoseIndex(BoneContainer);
+	FTransform BoneTM = MeshBases.GetComponentSpaceTransform(BoneIndex);
+	
+	// Convert to the specific display space if necessary
+	FAnimationRuntime::ConvertCSTransformToBoneSpace(SkelComp, MeshBases, BoneTM, BoneIndex, DisplaySpace);
+
+	Translation = BoneTM.GetTranslation();
+	Rotation = BoneTM.GetRotation().Rotator();
+	Scale = BoneTM.GetScale3D();
+}
+
+bool FAnimNode_ObserveBone::IsValidToEvaluate(const USkeleton* Skeleton, const FBoneContainer& RequiredBones)
+{
+	return (BoneToObserve.IsValid(RequiredBones));
+}
+
+void FAnimNode_ObserveBone::InitializeBoneReferences(const FBoneContainer& RequiredBones)
+{
+	BoneToObserve.Initialize(RequiredBones);
+}
\ No newline at end of file
diff --git a/Engine/Source/Runtime/AnimGraphRuntime/Public/BoneControllers/AnimNode_ObserveBone.h b/Engine/Source/Runtime/AnimGraphRuntime/Public/BoneControllers/AnimNode_ObserveBone.h
new file mode 100644
--- /dev/null
+++ b/Engine/Source/Runtime/AnimGraphRuntime/Public/BoneControllers/AnimNode_ObserveBone.h
@@ -0,0 +1,52 @@
+// Copyright 1998-2015 Epic Games, Inc. All Rights Reserved.
+
+#pragma once
+
+#include "AnimNode_SkeletalControlBase.h"
+#include "AnimNode_ObserveBone.generated.h"
+
+/**
+ *	Debugging node that displays the current value of a bone in a specific space.
+ */
+USTRUCT()
+struct ANIMGRAPHRUNTIME_API FAnimNode_ObserveBone : public FAnimNode_SkeletalControlBase
+{
+	GENERATED_USTRUCT_BODY()
+
+	/** Name of bone to observe. */
+	UPROPERTY(EditAnywhere, BlueprintReadWrite, Category=SkeletalControl) 
+	FBoneReference BoneToObserve;
+
+	/** Reference frame to display the bone transform in. */
+	UPROPERTY(EditAnywhere, BlueprintReadWrite, Category=SkeletalControl)
+	TEnumAsByte<EBoneControlSpace> DisplaySpace;
+
+	/** Translation of the bone being observed. */
+	UPROPERTY()
+	FVector Translation;
+
+	/** Rotation of the bone being observed. */
+	UPROPERTY()
+	FRotator Rotation;
+
+	/** Scale of the bone being observed. */
+	UPROPERTY()
+	FVector Scale;
+
+public:
+	FAnimNode_ObserveBone();
+
+	// FAnimNode_Base interface
+	virtual void GatherDebugData(FNodeDebugData& DebugData) override;
+	// End of FAnimNode_Base interface
+
+	// FAnimNode_SkeletalControlBase interface
+	virtual void EvaluateBoneTransforms(USkeletalMeshComponent* SkelComp, FCSPose<FCompactPose>& MeshBases, TArray<FBoneTransform>& OutBoneTransforms) override;
+	virtual bool IsValidToEvaluate(const USkeleton* Skeleton, const FBoneContainer& RequiredBones) override;
+	// End of FAnimNode_SkeletalControlBase interface
+
+private:
+	// FAnimNode_SkeletalControlBase interface
+	virtual void InitializeBoneReferences(const FBoneContainer& RequiredBones) override;
+	// End of FAnimNode_SkeletalControlBase interface
+};
