Animation: Improved bone driven controllers
- One driver value can now drive multiple components
- How a component is driven is now configurable (add to input, replace, add to ref pose and replace)
- The range limit is now defined on the input (before the multiplier is applied) rather than the output (existing ranges are fixed up on load)
- New details customization to make working with the node easier
- Individual scale components can now be a driver or one of the driven components, rather than just the max component of scale
- The node title now displays the driving equation (e.g., driven.rotation.z = -0.5 * driver.rotation.x)

Animation: Added a custom version for AnimGraph/AnimGraphRuntime usage

Upgrade Note: The exact behavior of the original node was not preserved, but there are now more configuration options, so please revisit any uses of bone driven controllers in your animation blueprints

diff --git a/Engine/Source/Editor/AnimGraph/Classes/AnimGraphNode_BoneDrivenController.h b/Engine/Source/Editor/AnimGraph/Classes/AnimGraphNode_BoneDrivenController.h
--- a/Engine/Source/Editor/AnimGraph/Classes/AnimGraphNode_BoneDrivenController.h
+++ b/Engine/Source/Editor/AnimGraph/Classes/AnimGraphNode_BoneDrivenController.h
@@ -2,6 +2,9 @@
 
 #pragma once
 
+class IDetailCategoryBuilder;
+class IPropertyHandle;
+
 #include "AnimGraphNode_SkeletalControlBase.h"
 #include "BoneControllers/AnimNode_BoneDrivenController.h"
 #include "EdGraph/EdGraphNodeUtils.h" // for FNodeTitleTextTable
@@ -19,6 +22,9 @@ class ANIMGRAPH_API UAnimGraphNode_BoneDrivenController : public UAnimGraphNode_
 	FAnimNode_BoneDrivenController Node;
 
 public:
+	// UObject interface
+	virtual void Serialize(FArchive& Ar) override;
+	// End of UObject interface
 
 	// UEdGraphNode interface
 	virtual FText GetNodeTitle(ENodeTitleType::Type TitleType) const override;
@@ -43,7 +49,7 @@ protected:
 	// Should non-curve mapping values be shown (multipler, range)?
 	EVisibility AreNonCurveMappingValuesVisible() const;
 
-private:
-	/** Constructing FText strings can be costly, so we cache the node's title */
-	FNodeTitleTextTable CachedNodeTitles;
+	static void AddTripletPropertyRow(const FText& Name, const FText& Tooltip, IDetailCategoryBuilder& Category, TSharedRef<IPropertyHandle> PropertyHandle, const FName XPropertyName, const FName YPropertyName, const FName ZPropertyName);
+
+	static FText ComponentTypeToText(EComponentType::Type Component);
 };
diff --git a/Engine/Source/Editor/AnimGraph/Private/AnimGraphNode_BoneDrivenController.cpp b/Engine/Source/Editor/AnimGraph/Private/AnimGraphNode_BoneDrivenController.cpp
--- a/Engine/Source/Editor/AnimGraph/Private/AnimGraphNode_BoneDrivenController.cpp
+++ b/Engine/Source/Editor/AnimGraph/Private/AnimGraphNode_BoneDrivenController.cpp
@@ -4,6 +4,7 @@
 #include "AnimGraphNode_BoneDrivenController.h"
 #include "CompilerResultsLog.h"
 #include "PropertyEditing.h"
+#include "AnimationCustomVersion.h"
 
 #define LOCTEXT_NAMESPACE "A3Nodes"
 
@@ -13,6 +14,28 @@ UAnimGraphNode_BoneDrivenController::UAnimGraphNode_BoneDrivenController(const F
 
 }
 
+void UAnimGraphNode_BoneDrivenController::Serialize(FArchive& Ar)
+{
+	Super::Serialize(Ar);
+
+	Ar.UsingCustomVersion(FAnimationCustomVersion::GUID);
+
+	if (Ar.CustomVer(FAnimationCustomVersion::GUID) < FAnimationCustomVersion::BoneDrivenControllerMatchingMaya)
+	{
+		// The node used to be able to only drive a single component rather than a selection of components
+		Node.ConvertTargetComponentToBits();
+
+		// The old definition of range was clamping the output, rather than the input
+		if (Node.bUseRange && !FMath::IsNearlyZero(Node.Multiplier))
+		{
+			// Before: Output = clamp(Input * Multipler)
+			// After: Output = clamp(Input) * Multiplier
+			Node.RangeMin /= Node.Multiplier;
+			Node.RangeMax /= Node.Multiplier;
+		}
+	}
+}
+
 FText UAnimGraphNode_BoneDrivenController::GetTooltipText() const
 {
 	return LOCTEXT("UAnimGraphNode_BoneDrivenController_ToolTip", "Drives the transform of a bone using the transform of another");
@@ -20,20 +43,60 @@ FText UAnimGraphNode_BoneDrivenController::GetTooltipText() const
 
 FText UAnimGraphNode_BoneDrivenController::GetNodeTitle(ENodeTitleType::Type TitleType) const
 {
-	if ((Node.SourceBone.BoneName == NAME_None) && (Node.TargetBone.BoneName == NAME_None) && (TitleType == ENodeTitleType::ListView || TitleType == ENodeTitleType::MenuTitle))
+	if ((Node.SourceBone.BoneName == NAME_None) && (Node.TargetBone.BoneName == NAME_None) && ((TitleType == ENodeTitleType::ListView) || (TitleType == ENodeTitleType::MenuTitle)))
 	{
 		return GetControllerDescription();
 	}
-	// @TODO: the bone can be altered in the property editor, so we have to 
-	//        choose to mark this dirty when that happens for this to properly work
-	else //if (!CachedNodeTitles.IsTitleCached(TitleType, this))
+	else
 	{
+		// Determine the mapping
+		FText FinalSourceExpression;
+		{
+			FFormatNamedArguments SourceArgs;
+			SourceArgs.Add(TEXT("SourceBone"), FText::FromName(Node.SourceBone.BoneName));
+			SourceArgs.Add(TEXT("SourceComponent"), ComponentTypeToText(Node.SourceComponent));
+
+			if (Node.DrivingCurve != nullptr)
+			{
+				FinalSourceExpression = LOCTEXT("BoneDrivenByCurve", "curve({SourceBone}.{SourceComponent})");
+			}
+			else
+			{
+				if (Node.Multiplier == 1.0f)
+				{
+					FinalSourceExpression = LOCTEXT("BoneMultiplierIs1", "{SourceBone}.{SourceComponent}");
+				}
+				else
+				{
+					SourceArgs.Add(TEXT("Multiplier"), FText::AsNumber(Node.Multiplier));
+					FinalSourceExpression = LOCTEXT("BoneMultiplierIs1", "{SourceBone}.{SourceComponent} * {Multiplier}");
+				}
+			}
+
+			FinalSourceExpression = FText::Format(FinalSourceExpression, SourceArgs);
+		}
+
 		FFormatNamedArguments Args;
 		Args.Add(TEXT("ControllerDesc"), GetControllerDescription());
-		Args.Add(TEXT("SourceBone"), FText::FromName(Node.SourceBone.BoneName));
 		Args.Add(TEXT("TargetBone"), FText::FromName(Node.TargetBone.BoneName));
 
-		if (TitleType == ENodeTitleType::ListView || TitleType == ENodeTitleType::MenuTitle)
+		// Determine the target component
+		int32 NumComponents = 0;
+		EComponentType::Type TargetComponent = EComponentType::None;
+
+		#define UE_CHECK_TARGET_COMPONENT(ComponentProperty, ComponentChoice) if (ComponentProperty) { ++NumComponents; TargetComponent = ComponentChoice; }
+		UE_CHECK_TARGET_COMPONENT(Node.bAffectTargetTranslationX, EComponentType::TranslationX);
+		UE_CHECK_TARGET_COMPONENT(Node.bAffectTargetTranslationY, EComponentType::TranslationY);
+		UE_CHECK_TARGET_COMPONENT(Node.bAffectTargetTranslationZ, EComponentType::TranslationZ);
+		UE_CHECK_TARGET_COMPONENT(Node.bAffectTargetRotationX, EComponentType::RotationX);
+		UE_CHECK_TARGET_COMPONENT(Node.bAffectTargetRotationY, EComponentType::RotationY);
+		UE_CHECK_TARGET_COMPONENT(Node.bAffectTargetRotationZ, EComponentType::RotationZ);
+		UE_CHECK_TARGET_COMPONENT(Node.bAffectTargetScaleX, EComponentType::ScaleX);
+		UE_CHECK_TARGET_COMPONENT(Node.bAffectTargetScaleY, EComponentType::ScaleY);
+		UE_CHECK_TARGET_COMPONENT(Node.bAffectTargetScaleZ, EComponentType::ScaleZ);
+		Args.Add(TEXT("TargetComponents"), (NumComponents <= 1) ? ComponentTypeToText(TargetComponent) : LOCTEXT("MultipleTargetComponents", "multiple"));
+
+		if ((TitleType == ENodeTitleType::ListView) || (TitleType == ENodeTitleType::MenuTitle))
 		{
 			Args.Add(TEXT("Delim"), FText::FromString(TEXT(" - ")));
 		}
@@ -42,10 +105,9 @@ FText UAnimGraphNode_BoneDrivenController::GetNodeTitle(ENodeTitleType::Type Tit
 			Args.Add(TEXT("Delim"), FText::FromString(TEXT("\n")));
 		}
 
-		// FText::Format() is slow, so we cache this to save on performance
-		CachedNodeTitles.SetCachedTitle(TitleType, FText::Format(LOCTEXT("AnimGraphNode_BoneDrivenController_Title", "{ControllerDesc}{Delim}Driver Bone: {SourceBone}{Delim}Driven Bone: {TargetBone}"), Args), this);
+		Args.Add(TEXT("SourceExpression"), FinalSourceExpression);
+		return FText::Format(LOCTEXT("AnimGraphNode_BoneDrivenController_Title", "{TargetBone}.{TargetComponents} = {SourceExpression}{Delim}{ControllerDesc}"), Args);
 	}	
-	return CachedNodeTitles[TitleType];
 }
 
 FText UAnimGraphNode_BoneDrivenController::GetControllerDescription() const
@@ -101,29 +163,173 @@ void UAnimGraphNode_BoneDrivenController::ValidateAnimNodeDuringCompilation(USke
 		MessageLog.Warning(*LOCTEXT("NoTargetBone", "@@ - You must pick a target bone as the Driven joint").ToString(), this);
 	}
 
-	if (Node.TargetComponent == EComponentType::None)
+	const bool bAffectsTranslation = Node.bAffectTargetTranslationX || Node.bAffectTargetTranslationY || Node.bAffectTargetTranslationZ;
+	const bool bAffectsRotation = Node.bAffectTargetRotationX || Node.bAffectTargetRotationY || Node.bAffectTargetRotationZ;
+	const bool bAffectsScale = Node.bAffectTargetScaleX || Node.bAffectTargetScaleY || Node.bAffectTargetScaleZ;
+
+	if (!bAffectsTranslation && !bAffectsRotation && !bAffectsScale)
 	{
-		MessageLog.Warning(*LOCTEXT("NoTargetComponent", "@@ - You must pick a target component on the Driven joint").ToString(), this);
+		MessageLog.Warning(*LOCTEXT("NoTargetComponent", "@@ - You must pick one or more target components on the Driven joint").ToString(), this);
 	}
 
 	Super::ValidateAnimNodeDuringCompilation(ForSkeleton, MessageLog);
 }
 
-void UAnimGraphNode_BoneDrivenController::CustomizeDetails(IDetailLayoutBuilder& DetailBuilder)
+void UAnimGraphNode_BoneDrivenController::AddTripletPropertyRow(const FText& Name, const FText& Tooltip, IDetailCategoryBuilder& Category, TSharedRef<IPropertyHandle> PropertyHandle, const FName XPropertyName, const FName YPropertyName, const FName ZPropertyName)
 {
-	DetailBuilder.EditCategory(TEXT("Source (Driver)"));
-	IDetailCategoryBuilder& MappingCategory = DetailBuilder.EditCategory(TEXT("Mapping"));
-	DetailBuilder.EditCategory(TEXT("Destination (Driven)"));
+	const float XYZPadding = 5.0f;
+
+	TSharedPtr<IPropertyHandle> XProperty = PropertyHandle->GetChildHandle(XPropertyName);
+	Category.GetParentLayout().HideProperty(XProperty);
+
+	TSharedPtr<IPropertyHandle> YProperty = PropertyHandle->GetChildHandle(YPropertyName);
+	Category.GetParentLayout().HideProperty(YProperty);
 
+	TSharedPtr<IPropertyHandle> ZProperty = PropertyHandle->GetChildHandle(ZPropertyName);
+	Category.GetParentLayout().HideProperty(ZProperty);
+
+	Category.AddCustomRow(Name)
+	.NameContent()
+	[
+		SNew(STextBlock)
+		.Text(Name)
+		.ToolTipText(Tooltip)
+		.Font(IDetailLayoutBuilder::GetDetailFont())
+	]
+	.ValueContent()
+	[
+		SNew(SHorizontalBox)
+		+ SHorizontalBox::Slot()
+		.Padding(0.f, 0.f, XYZPadding, 0.f)
+		.AutoWidth()
+		[
+			SNew(SHorizontalBox)
+			+ SHorizontalBox::Slot()
+			.AutoWidth()
+			[
+				XProperty->CreatePropertyNameWidget()
+			]
+			+ SHorizontalBox::Slot()
+			.AutoWidth()
+			[
+				XProperty->CreatePropertyValueWidget()
+			]
+		]
+
+		+ SHorizontalBox::Slot()
+		.Padding(0.f, 0.f, XYZPadding, 0.f)
+		.AutoWidth()
+		[
+			SNew(SHorizontalBox)
+			+ SHorizontalBox::Slot()
+			.AutoWidth()
+			[
+				YProperty->CreatePropertyNameWidget()
+			]
+			+ SHorizontalBox::Slot()
+			.AutoWidth()
+			[
+				YProperty->CreatePropertyValueWidget()
+			]
+		]
+
+		+ SHorizontalBox::Slot()
+		.Padding(0.f, 0.f, XYZPadding, 0.f)
+		.AutoWidth()
+		[
+			SNew(SHorizontalBox)
+			+ SHorizontalBox::Slot()
+			.AutoWidth()
+			[
+				ZProperty->CreatePropertyNameWidget()
+			]
+			+ SHorizontalBox::Slot()
+			.AutoWidth()
+			[
+				ZProperty->CreatePropertyValueWidget()
+			]
+		]
+	];
+}
+
+void UAnimGraphNode_BoneDrivenController::CustomizeDetails(IDetailLayoutBuilder& DetailBuilder)
+{
 	TSharedRef<IPropertyHandle> NodeHandle = DetailBuilder.GetProperty(GET_MEMBER_NAME_CHECKED(UAnimGraphNode_BoneDrivenController, Node), GetClass());
 
+	TAttribute<EVisibility> NotUsingCurveVisibility = TAttribute<EVisibility>::Create(TAttribute<EVisibility>::FGetter::CreateUObject(this, &UAnimGraphNode_BoneDrivenController::AreNonCurveMappingValuesVisible));
+
+
+	// Source (Driver) category
+	IDetailCategoryBuilder& SourceCategory = DetailBuilder.EditCategory(TEXT("Source (Driver)"));
+
+	// Mapping category
+	IDetailCategoryBuilder& MappingCategory = DetailBuilder.EditCategory(TEXT("Mapping"));
 	MappingCategory.AddProperty(NodeHandle->GetChildHandle(GET_MEMBER_NAME_CHECKED(FAnimNode_BoneDrivenController, DrivingCurve)));
 
-	TAttribute<EVisibility> MappingVisibility = TAttribute<EVisibility>::Create(TAttribute<EVisibility>::FGetter::CreateUObject(this, &UAnimGraphNode_BoneDrivenController::AreNonCurveMappingValuesVisible));
-	MappingCategory.AddProperty(NodeHandle->GetChildHandle(GET_MEMBER_NAME_CHECKED(FAnimNode_BoneDrivenController, Multiplier))).Visibility(MappingVisibility);
-	MappingCategory.AddProperty(NodeHandle->GetChildHandle(GET_MEMBER_NAME_CHECKED(FAnimNode_BoneDrivenController, bUseRange))).Visibility(MappingVisibility);
-	MappingCategory.AddProperty(NodeHandle->GetChildHandle(GET_MEMBER_NAME_CHECKED(FAnimNode_BoneDrivenController, RangeMin))).Visibility(MappingVisibility);
-	MappingCategory.AddProperty(NodeHandle->GetChildHandle(GET_MEMBER_NAME_CHECKED(FAnimNode_BoneDrivenController, RangeMax))).Visibility(MappingVisibility);
+	MappingCategory.AddProperty(NodeHandle->GetChildHandle(GET_MEMBER_NAME_CHECKED(FAnimNode_BoneDrivenController, Multiplier))).Visibility(NotUsingCurveVisibility);
+	MappingCategory.AddProperty(NodeHandle->GetChildHandle(GET_MEMBER_NAME_CHECKED(FAnimNode_BoneDrivenController, bUseRange))).Visibility(NotUsingCurveVisibility);
+	MappingCategory.AddProperty(NodeHandle->GetChildHandle(GET_MEMBER_NAME_CHECKED(FAnimNode_BoneDrivenController, RangeMin))).Visibility(NotUsingCurveVisibility);
+	MappingCategory.AddProperty(NodeHandle->GetChildHandle(GET_MEMBER_NAME_CHECKED(FAnimNode_BoneDrivenController, RangeMax))).Visibility(NotUsingCurveVisibility);
+
+	// Destination (Driven) category
+	IDetailCategoryBuilder& TargetCategory = DetailBuilder.EditCategory(TEXT("Destination (Driven)"));
+	TargetCategory.AddProperty(NodeHandle->GetChildHandle(GET_MEMBER_NAME_CHECKED(FAnimNode_BoneDrivenController, TargetBone)));
+
+	AddTripletPropertyRow(
+		/*Name=*/ LOCTEXT("DrivenTranslationLabel", "Translation"),
+		/*Tooltip=*/ LOCTEXT("DrivenTranslationTooltip", "Should the source bone drive one or more translation components of the target bone?"),
+		/*Category=*/ TargetCategory,
+		/*PropertyHandle=*/ NodeHandle,
+		/*X=*/ GET_MEMBER_NAME_CHECKED(FAnimNode_BoneDrivenController, bAffectTargetTranslationX),
+		/*Y=*/ GET_MEMBER_NAME_CHECKED(FAnimNode_BoneDrivenController, bAffectTargetTranslationY),
+		/*Z=*/ GET_MEMBER_NAME_CHECKED(FAnimNode_BoneDrivenController, bAffectTargetTranslationZ));
+
+	AddTripletPropertyRow(
+		/*Name=*/ LOCTEXT("DrivenRotationLabel", "Rotation"),
+		/*Tooltip=*/ LOCTEXT("DrivenRotationTooltip", "Should the source bone drive one or more rotation components of the target bone?"),
+		/*Category=*/ TargetCategory,
+		/*PropertyHandle=*/ NodeHandle,
+		/*X=*/ GET_MEMBER_NAME_CHECKED(FAnimNode_BoneDrivenController, bAffectTargetRotationX),
+		/*Y=*/ GET_MEMBER_NAME_CHECKED(FAnimNode_BoneDrivenController, bAffectTargetRotationY),
+		/*Z=*/ GET_MEMBER_NAME_CHECKED(FAnimNode_BoneDrivenController, bAffectTargetRotationZ));
+
+	AddTripletPropertyRow(
+		/*Name=*/ LOCTEXT("DrivenScaleLabel", "Scale"),
+		/*Tooltip=*/ LOCTEXT("DrivenScaleTooltip", "Should the source bone drive one or more scale components of the target bone?"),
+		/*Category=*/ TargetCategory,
+		/*PropertyHandle=*/ NodeHandle,
+		/*X=*/ GET_MEMBER_NAME_CHECKED(FAnimNode_BoneDrivenController, bAffectTargetScaleX),
+		/*Y=*/ GET_MEMBER_NAME_CHECKED(FAnimNode_BoneDrivenController, bAffectTargetScaleY),
+		/*Z=*/ GET_MEMBER_NAME_CHECKED(FAnimNode_BoneDrivenController, bAffectTargetScaleZ));
+}
+
+FText UAnimGraphNode_BoneDrivenController::ComponentTypeToText(EComponentType::Type Component)
+{
+	switch (Component)
+	{
+	case EComponentType::TranslationX:
+		return LOCTEXT("ComponentType_TranslationX", "translateX");
+	case EComponentType::TranslationY:
+		return LOCTEXT("ComponentType_TranslationY", "translateY");
+	case EComponentType::TranslationZ:
+		return LOCTEXT("ComponentType_TranslationZ", "translateZ");
+	case EComponentType::RotationX:
+		return LOCTEXT("ComponentType_RotationX", "rotateX");
+	case EComponentType::RotationY:
+		return LOCTEXT("ComponentType_RotationY", "rotateY");
+	case EComponentType::RotationZ:
+		return LOCTEXT("ComponentType_RotationZ", "rotateZ");
+	case EComponentType::Scale:
+		return LOCTEXT("ComponentType_ScaleMax", "scaleMax");
+	case EComponentType::ScaleX:
+		return LOCTEXT("ComponentType_ScaleX", "scaleX");
+	case EComponentType::ScaleY:
+		return LOCTEXT("ComponentType_ScaleY", "scaleY");
+	case EComponentType::ScaleZ:
+		return LOCTEXT("ComponentType_ScaleZ", "scaleZ");
+	default:
+		return LOCTEXT("ComponentType_None", "(none)");
+	}
 }
 
 EVisibility UAnimGraphNode_BoneDrivenController::AreNonCurveMappingValuesVisible() const
@@ -131,5 +337,4 @@ EVisibility UAnimGraphNode_BoneDrivenController::AreNonCurveMappingValuesVisible
 	return (Node.DrivingCurve != nullptr) ? EVisibility::Collapsed : EVisibility::Visible;
 }
 
-
 #undef LOCTEXT_NAMESPACE
diff --git a/Engine/Source/Runtime/AnimGraphRuntime/Private/BoneControllers/AnimNode_BoneDrivenController.cpp b/Engine/Source/Runtime/AnimGraphRuntime/Private/BoneControllers/AnimNode_BoneDrivenController.cpp
--- a/Engine/Source/Runtime/AnimGraphRuntime/Private/BoneControllers/AnimNode_BoneDrivenController.cpp
+++ b/Engine/Source/Runtime/AnimGraphRuntime/Private/BoneControllers/AnimNode_BoneDrivenController.cpp
@@ -12,7 +12,17 @@ FAnimNode_BoneDrivenController::FAnimNode_BoneDrivenController()
 	, bUseRange(false)
 	, RangeMin(-1.0f)
 	, RangeMax(1.0f)
-	, TargetComponent(EComponentType::None)
+	, TargetComponent_DEPRECATED(EComponentType::None)
+	, bAffectTargetTranslationX(false)
+	, bAffectTargetTranslationY(false)
+	, bAffectTargetTranslationZ(false)
+	, bAffectTargetRotationX(false)
+	, bAffectTargetRotationY(false)
+	, bAffectTargetRotationZ(false)
+	, bAffectTargetScaleX(false)
+	, bAffectTargetScaleY(false)
+	, bAffectTargetScaleZ(false)
+	, ModificationMode(EDrivenBoneModificationMode::AddToInput)
 {
 }
 
@@ -27,12 +37,12 @@ void FAnimNode_BoneDrivenController::GatherDebugData(FNodeDebugData& DebugData)
 	ComponentPose.GatherDebugData(DebugData);
 }
 
-void FAnimNode_BoneDrivenController::EvaluateBoneTransforms(USkeletalMeshComponent* SkelComp, FCSPose<FCompactPose>& MeshBases, TArray<FBoneTransform>& OutBoneTransforms)
+void FAnimNode_BoneDrivenController::EvaluateBoneTransforms(USkeletalMeshComponent* SkelComp, FCSPose<FCompactPose>& MeshBases, TArray<FBoneTransform>& OutCSBoneTransforms)
 {
-	check(OutBoneTransforms.Num() == 0);
+	check(OutCSBoneTransforms.Num() == 0);
 	
 	// Early out if we're not driving from or to anything
-	if ((SourceComponent == EComponentType::None) || (TargetComponent == EComponentType::None))
+	if (SourceComponent == EComponentType::None)
 	{
 		return;
 	}
@@ -54,13 +64,18 @@ void FAnimNode_BoneDrivenController::EvaluateBoneTransforms(USkeletalMeshCompone
 		const FVector RotationDiff = (SourceCurr.GetRotation() * SourceOrigRef.GetRotation().Inverse()).Euler();
 		SourceValue = RotationDiff[(int32)(SourceComponent - EComponentType::RotationX)];
 	}
-	else
+	else if (SourceComponent == EComponentType::Scale)
 	{
 		const FVector CurrentScale = SourceCurr.GetScale3D();
 		const FVector RefScale = SourceOrigRef.GetScale3D();
 		const float ScaleDiff = FMath::Max3(CurrentScale[0], CurrentScale[1], CurrentScale[2]) - FMath::Max3(RefScale[0], RefScale[1], RefScale[2]);
 		SourceValue = ScaleDiff;
 	}
+	else
+	{
+		const FVector ScaleDiff = SourceCurr.GetScale3D() - SourceOrigRef.GetScale3D();
+		SourceValue = ScaleDiff[(int32)(SourceComponent - EComponentType::ScaleX)];
+	}
 	
 	// Determine the resulting value
 	float FinalDriverValue = SourceValue;
@@ -72,44 +87,96 @@ void FAnimNode_BoneDrivenController::EvaluateBoneTransforms(USkeletalMeshCompone
 	else
 	{
 		// Apply the fixed function remapping/clamping
-		FinalDriverValue *= Multiplier;
 		if (bUseRange)
 		{
 			FinalDriverValue = FMath::Clamp(FinalDriverValue, RangeMin, RangeMax);
 		}
+
+		FinalDriverValue *= Multiplier;
 	}
 
-	// Difference to apply after processing
-	FVector NewRotDiff(FVector::ZeroVector);
-	FVector NewTransDiff(FVector::ZeroVector);
-	float NewScaleDiff = 0.0f;
+	// Calculate a new local-space bone position by adding or replacing target components in the current local space position
+	const FCompactPoseBoneIndex TargetBoneIndex = TargetBone.GetCompactPoseIndex(BoneContainer);
+
+	const FTransform OriginalLocalTM = MeshBases.GetLocalSpaceTransform(TargetBoneIndex);
+	FVector NewTrans(OriginalLocalTM.GetTranslation());
+	FVector NewScale(OriginalLocalTM.GetScale3D());
+	FQuat NewRot(OriginalLocalTM.GetRotation());
 
-	// Resolve target value
-	float* DestPtr = nullptr;
-	if (TargetComponent < EComponentType::RotationX)
+	if (ModificationMode == EDrivenBoneModificationMode::AddToInput)
 	{
-		DestPtr = &NewTransDiff[(int32)(TargetComponent - EComponentType::TranslationX)];
+		// Add the mapped value to the target components
+		if (bAffectTargetTranslationX) { NewTrans.X += FinalDriverValue; }
+		if (bAffectTargetTranslationY) { NewTrans.Y += FinalDriverValue; }
+		if (bAffectTargetTranslationZ) { NewTrans.Z += FinalDriverValue; }
+
+		if (bAffectTargetRotationX || bAffectTargetRotationY || bAffectTargetRotationZ)
+		{
+			FVector NewRotDeltaEuler(ForceInitToZero);
+			if (bAffectTargetRotationX) { NewRotDeltaEuler.X = FinalDriverValue; }
+			if (bAffectTargetRotationY) { NewRotDeltaEuler.Y = FinalDriverValue; }
+			if (bAffectTargetRotationZ) { NewRotDeltaEuler.Z = FinalDriverValue; }
+			NewRot = NewRot * FQuat::MakeFromEuler(NewRotDeltaEuler);
+		}
+
+		if (bAffectTargetScaleX) { NewScale.X += FinalDriverValue; }
+		if (bAffectTargetScaleY) { NewScale.Y += FinalDriverValue; }
+		if (bAffectTargetScaleZ) { NewScale.Z += FinalDriverValue; }
 	}
-	else if (TargetComponent < EComponentType::Scale)
+	else if (ModificationMode == EDrivenBoneModificationMode::ReplaceComponent)
 	{
-		DestPtr = &NewRotDiff[(int32)(TargetComponent - EComponentType::RotationX)];
+		// Replace the target components with the mapped value
+		if (bAffectTargetTranslationX) { NewTrans.X = FinalDriverValue; }
+		if (bAffectTargetTranslationY) { NewTrans.Y = FinalDriverValue; }
+		if (bAffectTargetTranslationZ) { NewTrans.Z = FinalDriverValue; }
+
+		if (bAffectTargetRotationX || bAffectTargetRotationY || bAffectTargetRotationZ)
+		{
+			FVector NewRotEuler(NewRot.Euler());
+			if (bAffectTargetRotationX) { NewRotEuler.X = FinalDriverValue; }
+			if (bAffectTargetRotationY) { NewRotEuler.Y = FinalDriverValue; }
+			if (bAffectTargetRotationZ) { NewRotEuler.Z = FinalDriverValue; }
+			NewRot = FQuat::MakeFromEuler(NewRotEuler);
+		}
+
+		if (bAffectTargetScaleX) { NewScale.X = FinalDriverValue; }
+		if (bAffectTargetScaleY) { NewScale.Y = FinalDriverValue; }
+		if (bAffectTargetScaleZ) { NewScale.Z = FinalDriverValue; }
 	}
-	else
+	else if (ModificationMode == EDrivenBoneModificationMode::AddToRefPose)
 	{
-		DestPtr = &NewScaleDiff;
-	}
+		const FTransform RefPoseTransform = MeshBases.GetPose().GetRefPose(TargetBoneIndex);
 
-	// Set the value
-	*DestPtr = FinalDriverValue;
+		// Add the mapped value to the ref pose components
+		if (bAffectTargetTranslationX) { NewTrans.X = RefPoseTransform.GetTranslation().X + FinalDriverValue; }
+		if (bAffectTargetTranslationY) { NewTrans.Y = RefPoseTransform.GetTranslation().Y + FinalDriverValue; }
+		if (bAffectTargetTranslationZ) { NewTrans.Z = RefPoseTransform.GetTranslation().Z + FinalDriverValue; }
 
-	// Build final transform difference
-	const FTransform FinalDiff(FQuat::MakeFromEuler(NewRotDiff), NewTransDiff, FVector(NewScaleDiff));
+		if (bAffectTargetRotationX || bAffectTargetRotationY || bAffectTargetRotationZ)
+		{
+			const FVector RefPoseRotEuler(RefPoseTransform.GetRotation().Euler());
 
-	const FCompactPoseBoneIndex TargetBoneIndex = TargetBone.GetCompactPoseIndex(BoneContainer);
+			// Replace any components that are being driven with their ref pose value first and create a delta rotation as well
+			FVector SourceRotationEuler(NewRot.Euler());
+			FVector NewRotDeltaEuler(ForceInitToZero);
+			if (bAffectTargetRotationX) { SourceRotationEuler.X = RefPoseRotEuler.X; NewRotDeltaEuler.X = FinalDriverValue; }
+			if (bAffectTargetRotationY) { SourceRotationEuler.Y = RefPoseRotEuler.Y; NewRotDeltaEuler.Y = FinalDriverValue; }
+			if (bAffectTargetRotationZ) { SourceRotationEuler.Z = RefPoseRotEuler.Z; NewRotDeltaEuler.Z = FinalDriverValue; }
+
+			// Combine the (modified) source and the delta rotation
+			NewRot = FQuat::MakeFromEuler(SourceRotationEuler) * FQuat::MakeFromEuler(NewRotDeltaEuler);
+		}
+
+		if (bAffectTargetScaleX) { NewScale.X = RefPoseTransform.GetScale3D().X + FinalDriverValue; }
+		if (bAffectTargetScaleY) { NewScale.Y = RefPoseTransform.GetScale3D().Y + FinalDriverValue; }
+		if (bAffectTargetScaleZ) { NewScale.Z = RefPoseTransform.GetScale3D().Z + FinalDriverValue; }
+	}
+	else
+	{
+		ensureMsgf(false, TEXT("Unknown entry in EDrivenBoneModificationMode"));
+	}
 
-	// Starting point for the new transform
-	FTransform NewLocal = MeshBases.GetLocalSpaceTransform(TargetBoneIndex);
-	NewLocal.AccumulateWithAdditiveScale3D(FinalDiff);
+	const FTransform ModifiedLocalTM(NewRot, NewTrans, NewScale);
 
 	// If we have a parent, concatenate the transform, otherwise just take the new transform
 	const FCompactPoseBoneIndex ParentIndex = MeshBases.GetPose().GetParentBoneIndex(TargetBoneIndex);
@@ -118,11 +185,11 @@ void FAnimNode_BoneDrivenController::EvaluateBoneTransforms(USkeletalMeshCompone
 	{
 		const FTransform& ParentTM = MeshBases.GetComponentSpaceTransform(ParentIndex);
 		
-		OutBoneTransforms.Add(FBoneTransform(TargetBoneIndex, NewLocal * ParentTM));
+		OutCSBoneTransforms.Add(FBoneTransform(TargetBoneIndex, ModifiedLocalTM * ParentTM));
 	}
 	else
 	{
-		OutBoneTransforms.Add(FBoneTransform(TargetBoneIndex, NewLocal));
+		OutCSBoneTransforms.Add(FBoneTransform(TargetBoneIndex, ModifiedLocalTM));
 	}
 }
 
@@ -131,6 +198,36 @@ bool FAnimNode_BoneDrivenController::IsValidToEvaluate(const USkeleton* Skeleton
 	return SourceBone.IsValid(RequiredBones) && TargetBone.IsValid(RequiredBones);
 }
 
+void FAnimNode_BoneDrivenController::ConvertTargetComponentToBits()
+{
+	switch (TargetComponent_DEPRECATED)
+	{
+	case EComponentType::TranslationX:
+		bAffectTargetTranslationX = true;
+		break;
+	case EComponentType::TranslationY:
+		bAffectTargetTranslationY = true;
+		break;
+	case EComponentType::TranslationZ:
+		bAffectTargetTranslationZ = true;
+		break;
+	case EComponentType::RotationX:
+		bAffectTargetRotationX = true;
+		break;
+	case EComponentType::RotationY:
+		bAffectTargetRotationY = true;
+		break;
+	case EComponentType::RotationZ:
+		bAffectTargetRotationZ = true;
+		break;
+	case EComponentType::Scale:
+		bAffectTargetScaleX = true;
+		bAffectTargetScaleY = true;
+		bAffectTargetScaleZ = true;
+		break;
+	}
+}
+
 void FAnimNode_BoneDrivenController::InitializeBoneReferences(const FBoneContainer& RequiredBones)
 {
 	SourceBone.Initialize(RequiredBones);
diff --git a/Engine/Source/Runtime/AnimGraphRuntime/Public/AnimationCustomVersion.h b/Engine/Source/Runtime/AnimGraphRuntime/Public/AnimationCustomVersion.h
new file mode 100644
--- /dev/null
+++ b/Engine/Source/Runtime/AnimGraphRuntime/Public/AnimationCustomVersion.h
@@ -0,0 +1,27 @@
+// Copyright 1998-2015 Epic Games, Inc. All Rights Reserved.
+
+#pragma once
+
+// Custom serialization version for assets/classes in the AnimGraphRuntime and AnimGraph modules
+struct ANIMGRAPHRUNTIME_API FAnimationCustomVersion
+{
+	enum Type
+	{
+		// Before any version changes were made in the plugin
+		BeforeCustomVersionWasAdded = 0,
+
+		// Added support for one-to-many component mappings to FAnimNode_BoneDrivenController,
+		// changed the range to apply to the input, and added a configurable method for updating the components
+		BoneDrivenControllerMatchingMaya = 1,
+
+		// -----<new versions can be added above this line>-------------------------------------------------
+		VersionPlusOne,
+		LatestVersion = VersionPlusOne - 1
+	};
+
+	// The GUID for this custom version number
+	const static FGuid GUID;
+
+private:
+	FAnimationCustomVersion() {}
+};
diff --git a/Engine/Source/Runtime/AnimGraphRuntime/Public/BoneControllers/AnimNode_BoneDrivenController.h b/Engine/Source/Runtime/AnimGraphRuntime/Public/BoneControllers/AnimNode_BoneDrivenController.h
--- a/Engine/Source/Runtime/AnimGraphRuntime/Public/BoneControllers/AnimNode_BoneDrivenController.h
+++ b/Engine/Source/Runtime/AnimGraphRuntime/Public/BoneControllers/AnimNode_BoneDrivenController.h
@@ -7,6 +7,8 @@
 
 // Evaluation of the bone transforms relies on the size and ordering of this
 // enum, if this needs to change make sure EvaluateBoneTransforms is updated.
+
+// The transform component (attribute) to read from
 UENUM()
 namespace EComponentType
 {
@@ -19,10 +21,27 @@ namespace EComponentType
 		RotationX,
 		RotationY,
 		RotationZ,
-		Scale
+		Scale UMETA(DisplayName="Scale (largest component)"),
+		ScaleX,
+		ScaleY,
+		ScaleZ
 	};
 }
 
+// The type of modification to make to the destination component(s)
+UENUM()
+enum class EDrivenBoneModificationMode : uint8
+{
+	// Add the driven value to the input component value(s)
+	AddToInput,
+
+	// Replace the input component value(s) with the driven value
+	ReplaceComponent,
+
+	// Add the driven value to the reference pose value
+	AddToRefPose
+};
+
 /**
  * This is the runtime version of a bone driven controller, which maps part of the state from one bone to another (e.g., 2 * source.x -> target.z)
  */
@@ -43,29 +62,70 @@ struct ANIMGRAPHRUNTIME_API FAnimNode_BoneDrivenController : public FAnimNode_Sk
 	UPROPERTY(EditAnywhere, Category=Mapping)
 	UCurveFloat* DrivingCurve;
 
-	// Multiplier to apply to the input value
+	// Multiplier to apply to the input value (Note: Ignored when a curve is used)
 	UPROPERTY(EditAnywhere, Category=Mapping)
 	float Multiplier;
 
-	// Whether or not to use the range limits on the multiplied value
-	UPROPERTY(EditAnywhere, Category=Mapping)
+	// Whether or not to clamp the driver value before scaling it (Note: Ignored when a curve is used)
+	UPROPERTY(EditAnywhere, Category="Source (Driver)", meta=(DisplayName="Limit Source Range"))
 	bool bUseRange;
 
-	// Minimum limit of the output value
-	UPROPERTY(EditAnywhere, Category=Mapping, meta=(EditCondition=bUseRange))
+	// Minimum limit of the input value
+	UPROPERTY(EditAnywhere, Category="Source (Driver)", meta=(EditCondition=bUseRange, DisplayName="Source Range Min"))
 	float RangeMin;
 
-	// Maximum limit of the output value
-	UPROPERTY(EditAnywhere, Category=Mapping, meta=(EditCondition=bUseRange))
+	// Maximum limit of the input value
+	UPROPERTY(EditAnywhere, Category="Source (Driver)", meta=(EditCondition=bUseRange, DisplayName="Source Range Max"))
 	float RangeMax;
 
 	// Bone to drive using controller input
 	UPROPERTY(EditAnywhere, Category="Destination (driven)")
 	FBoneReference TargetBone;
 
-	// Transform component to write to
+private:
+	UPROPERTY()
+	TEnumAsByte<EComponentType::Type> TargetComponent_DEPRECATED;
+
+public:
+	// Affect the X component of translation on the target bone
+	UPROPERTY(EditAnywhere, Category="Destination (driven)", meta=(DisplayName="X"))
+	uint32 bAffectTargetTranslationX : 1;
+
+	// Affect the Y component of translation on the target bone
+	UPROPERTY(EditAnywhere, Category="Destination (driven)", meta=(DisplayName="Y"))
+	uint32 bAffectTargetTranslationY : 1;
+
+	// Affect the Z component of translation on the target bone
+	UPROPERTY(EditAnywhere, Category="Destination (driven)", meta=(DisplayName="Z"))
+	uint32 bAffectTargetTranslationZ : 1;
+
+	// Affect the X component of rotation on the target bone
+	UPROPERTY(EditAnywhere, Category="Destination (driven)", meta=(DisplayName="X"))
+	uint32 bAffectTargetRotationX : 1;
+
+	// Affect the Y component of rotation on the target bone
+	UPROPERTY(EditAnywhere, Category="Destination (driven)", meta=(DisplayName="Y"))
+	uint32 bAffectTargetRotationY : 1;
+
+	// Affect the Z component of rotation on the target bone
+	UPROPERTY(EditAnywhere, Category="Destination (driven)", meta=(DisplayName="Z"))
+	uint32 bAffectTargetRotationZ : 1;
+
+	// Affect the X component of scale on the target bone
+	UPROPERTY(EditAnywhere, Category = "Destination (driven)", meta=(DisplayName="X"))
+	uint32 bAffectTargetScaleX : 1;
+
+	// Affect the Y component of scale on the target bone
+	UPROPERTY(EditAnywhere, Category = "Destination (driven)", meta=(DisplayName="Y"))
+	uint32 bAffectTargetScaleY : 1;
+
+	// Affect the Z component of scale on the target bone
+	UPROPERTY(EditAnywhere, Category="Destination (driven)", meta=(DisplayName="Z"))
+	uint32 bAffectTargetScaleZ : 1;
+
+	// The type of modification to make to the destination component(s)
 	UPROPERTY(EditAnywhere, Category="Destination (driven)")
-	TEnumAsByte<EComponentType::Type> TargetComponent;
+	EDrivenBoneModificationMode ModificationMode;
 
 public:
 	FAnimNode_BoneDrivenController();
@@ -79,6 +139,8 @@ public:
 	virtual bool IsValidToEvaluate(const USkeleton* Skeleton, const FBoneContainer& RequiredBones) override;
 	// End of FAnimNode_SkeletalControlBase interface
 
+	// Upgrade a node from the output enum to the output bits (change made in FAnimationCustomVersion::BoneDrivenControllerMatchingMaya)
+	void ConvertTargetComponentToBits();
 protected:
 	
 	// FAnimNode_SkeletalControlBase protected interface
