Fixed a bug which was hiding Joint Target Location Pin in TwoBoneIK node

diff --git a/Engine/Source/Editor/AnimGraph/Classes/AnimGraphNode_TwoBoneIK.h b/Engine/Source/Editor/AnimGraph/Classes/AnimGraphNode_TwoBoneIK.h
--- a/Engine/Source/Editor/AnimGraph/Classes/AnimGraphNode_TwoBoneIK.h
+++ b/Engine/Source/Editor/AnimGraph/Classes/AnimGraphNode_TwoBoneIK.h
@@ -71,9 +71,6 @@ protected:
 	// local conversion function for drawing
 	void DrawTargetLocation(FPrimitiveDrawInterface* PDI, USkeletalMeshComponent* SkelComp, USkeleton * Skeleton, EBoneControlSpace SpaceBase, FName SpaceBoneName, const FVector & TargetLocation, const FColor & TargetColor, const FColor & BoneColor) const;
 
-	// make Pins showed / hidden by options
-	bool SetPinsVisibility(bool bShow);
-
 private:
 	/** Constructing FText strings can be costly, so we cache the node's title */
 	FNodeTitleTextTable CachedNodeTitles;
diff --git a/Engine/Source/Editor/AnimGraph/Private/AnimGraphNode_TwoBoneIK.cpp b/Engine/Source/Editor/AnimGraph/Private/AnimGraphNode_TwoBoneIK.cpp
--- a/Engine/Source/Editor/AnimGraph/Private/AnimGraphNode_TwoBoneIK.cpp
+++ b/Engine/Source/Editor/AnimGraph/Private/AnimGraphNode_TwoBoneIK.cpp
@@ -369,8 +369,6 @@ void UAnimGraphNode_TwoBoneIK::CustomizeDetails(class IDetailLayoutBuilder& Deta
 		TSharedPtr<IPropertyHandle> PropertyHandle;
 		PropertyHandle = DetailBuilder.GetProperty(*JointTargetSpaceBoneName, GetClass());
 		EffectorCategory.AddProperty(PropertyHandle);
-
-		bPinVisibilityChanged = SetPinsVisibility(true);
 	}
 	else // hide all properties in JointTarget category except for JointTargetLocationSpace
 	{
@@ -378,8 +376,6 @@ void UAnimGraphNode_TwoBoneIK::CustomizeDetails(class IDetailLayoutBuilder& Deta
 		DetailBuilder.HideProperty(PropertyHandle);
 		PropertyHandle = DetailBuilder.GetProperty(*JointTargetSpaceBoneName, GetClass());
 		DetailBuilder.HideProperty(PropertyHandle);
-
-		bPinVisibilityChanged = SetPinsVisibility(false);
 	}
 
 	const FString EffectorLocationSpace = FString::Printf(TEXT("Node.%s"), GET_MEMBER_NAME_STRING_CHECKED(FAnimNode_TwoBoneIK, EffectorLocationSpace));
@@ -399,33 +395,6 @@ void UAnimGraphNode_TwoBoneIK::CustomizeDetails(class IDetailLayoutBuilder& Deta
 		FSimpleDelegate UpdateJointSpaceDelegate = FSimpleDelegate::CreateSP(TwoBoneIKDelegate.Get(), &FTwoBoneIKDelegate::UpdateLocationSpace, &DetailBuilder);
 		JointTragetLocHandle->SetOnPropertyValueChanged(UpdateJointSpaceDelegate);
 	}
-
-	// reconstruct node for showing/hiding Pins
-	if (bPinVisibilityChanged)
-	{
-		ReconstructNode();
-	}
-}
-
-bool UAnimGraphNode_TwoBoneIK::SetPinsVisibility(bool bShow)
-{
-	bool bChanged = false;
-
-	for (FOptionalPinFromProperty& Pin : ShowPinForProperties)
-	{
-		if (Pin.PropertyName == GET_MEMBER_NAME_CHECKED(FAnimNode_TwoBoneIK, JointTargetLocation))
-		{
-			PreEditChange(NULL);
-			if (Pin.bShowPin != bShow)
-			{
-				Pin.bShowPin = bShow;
-				bChanged = true;
-			}
-			break;
-		}
-	}
-
-	return bChanged;
 }
 
 #undef LOCTEXT_NAMESPACE
