Fix for crash when undoing custom transition blend logic selection. Modify the parent graph so the transaction gets properly recorded.

diff --git a/Engine/Source/Editor/AnimGraph/Private/AnimStateTransitionNode.cpp b/Engine/Source/Editor/AnimGraph/Private/AnimStateTransitionNode.cpp
--- a/Engine/Source/Editor/AnimGraph/Private/AnimStateTransitionNode.cpp
+++ b/Engine/Source/Editor/AnimGraph/Private/AnimStateTransitionNode.cpp
@@ -432,6 +432,7 @@ void UAnimStateTransitionNode::CreateCustomTransitionGraph()
 	Schema->CreateDefaultNodesForGraph(*CustomTransitionGraph);
 
 	// Add the new graph as a child of our parent graph
+	GetGraph()->Modify();
 	GetGraph()->SubGraphs.Add(CustomTransitionGraph);
 }
 
