Animation: Replaced a property literal reference with GET_MEMBER_NAME_CHECKED

diff --git a/Engine/Source/Editor/AnimGraph/Private/AnimGraphNode_Base.cpp b/Engine/Source/Editor/AnimGraph/Private/AnimGraphNode_Base.cpp
--- a/Engine/Source/Editor/AnimGraph/Private/AnimGraphNode_Base.cpp
+++ b/Engine/Source/Editor/AnimGraph/Private/AnimGraphNode_Base.cpp
@@ -123,7 +123,7 @@ void UAnimGraphNode_Base::PostEditChangeProperty(struct FPropertyChangedEvent& P
 {
 	FName PropertyName = (PropertyChangedEvent.Property != NULL) ? PropertyChangedEvent.Property->GetFName() : NAME_None;
 
-	if ((PropertyName == TEXT("bShowPin")))
+	if ((PropertyName == GET_MEMBER_NAME_CHECKED(FOptionalPinFromProperty, bShowPin)))
 	{
 		GetSchema()->ReconstructNode(*this);
 	}
