Animation: Temporary fix for the popup bubbles on observe bone nodes being displayed offscreen

diff --git a/Engine/Source/Editor/AnimGraph/Private/AnimGraphNode_ObserveBone.cpp b/Engine/Source/Editor/AnimGraph/Private/AnimGraphNode_ObserveBone.cpp
--- a/Engine/Source/Editor/AnimGraph/Private/AnimGraphNode_ObserveBone.cpp
+++ b/Engine/Source/Editor/AnimGraph/Private/AnimGraphNode_ObserveBone.cpp
@@ -37,10 +37,10 @@ public:
 		SGraphNode::UpdateGraphNode();
 
 		// Prevent the comment bubble from being displayed
-		GetOrAddSlot(ENodeZone::TopCenter)
-		[
-			SNullWidget::NullWidget
-		];
+// 		GetOrAddSlot(ENodeZone::TopCenter)
+// 		[
+// 			SNullWidget::NullWidget
+// 		];
 	}
 	// End of SGraphNode interface
 
