Animation: Added a remap range to the bone driven controller when clamping is enabled, and added a note about the driven bone space

diff --git a/Engine/Source/Editor/AnimGraph/Classes/AnimGraphNode_BoneDrivenController.h b/Engine/Source/Editor/AnimGraph/Classes/AnimGraphNode_BoneDrivenController.h
--- a/Engine/Source/Editor/AnimGraph/Classes/AnimGraphNode_BoneDrivenController.h
+++ b/Engine/Source/Editor/AnimGraph/Classes/AnimGraphNode_BoneDrivenController.h
@@ -46,10 +46,11 @@ protected:
 	virtual FText GetControllerDescription() const override;
 	// End of UAnimGraphNode_SkeletalControlBase protected interface
 
-	// Should non-curve mapping values be shown (multipler, range)?
+	// Should non-curve mapping values be shown (multiplier, range)?
 	EVisibility AreNonCurveMappingValuesVisible() const;
+	EVisibility AreRemappingValuesVisible() const;
 
 	static void AddTripletPropertyRow(const FText& Name, const FText& Tooltip, IDetailCategoryBuilder& Category, TSharedRef<IPropertyHandle> PropertyHandle, const FName XPropertyName, const FName YPropertyName, const FName ZPropertyName);
-
+	static void AddRangePropertyRow(const FText& Name, const FText& Tooltip, IDetailCategoryBuilder& Category, TSharedRef<IPropertyHandle> PropertyHandle, const FName MinPropertyName, const FName MaxPropertyName, TAttribute<EVisibility> VisibilityAttribute);
 	static FText ComponentTypeToText(EComponentType::Type Component);
 };
diff --git a/Engine/Source/Editor/AnimGraph/Private/AnimGraphNode_BoneDrivenController.cpp b/Engine/Source/Editor/AnimGraph/Private/AnimGraphNode_BoneDrivenController.cpp
--- a/Engine/Source/Editor/AnimGraph/Private/AnimGraphNode_BoneDrivenController.cpp
+++ b/Engine/Source/Editor/AnimGraph/Private/AnimGraphNode_BoneDrivenController.cpp
@@ -20,19 +20,27 @@ void UAnimGraphNode_BoneDrivenController::Serialize(FArchive& Ar)
 
 	Ar.UsingCustomVersion(FAnimationCustomVersion::GUID);
 
-	if (Ar.CustomVer(FAnimationCustomVersion::GUID) < FAnimationCustomVersion::BoneDrivenControllerMatchingMaya)
-	{
-		// The node used to be able to only drive a single component rather than a selection of components
-		Node.ConvertTargetComponentToBits();
+	const int32 AnimVersion = Ar.CustomVer(FAnimationCustomVersion::GUID);
 
-		// The old definition of range was clamping the output, rather than the input
-		if (Node.bUseRange && !FMath::IsNearlyZero(Node.Multiplier))
+	if (AnimVersion < FAnimationCustomVersion::BoneDrivenControllerRemapping)
+	{
+		if (AnimVersion < FAnimationCustomVersion::BoneDrivenControllerMatchingMaya)
 		{
-			// Before: Output = clamp(Input * Multipler)
-			// After: Output = clamp(Input) * Multiplier
-			Node.RangeMin /= Node.Multiplier;
-			Node.RangeMax /= Node.Multiplier;
+			// The node used to be able to only drive a single component rather than a selection of components
+			Node.ConvertTargetComponentToBits();
+
+			// The old definition of range was clamping the output, rather than the input
+			if (Node.bUseRange && !FMath::IsNearlyZero(Node.Multiplier))
+			{
+				// Before: Output = clamp(Input * Multipler)
+				// After: Output = clamp(Input) * Multiplier
+				Node.RangeMin /= Node.Multiplier;
+				Node.RangeMax /= Node.Multiplier;
+			}
 		}
+
+		Node.RemappedMin = Node.RangeMin;
+		Node.RemappedMax = Node.RangeMax;
 	}
 }
 
@@ -62,14 +70,29 @@ FText UAnimGraphNode_BoneDrivenController::GetNodeTitle(ENodeTitleType::Type Tit
 			}
 			else
 			{
-				if (Node.Multiplier == 1.0f)
+				if (Node.bUseRange)
 				{
-					FinalSourceExpression = LOCTEXT("BoneMultiplierIs1", "{SourceBone}.{SourceComponent}");
+					if (Node.Multiplier == 1.0f)
+					{
+						FinalSourceExpression = LOCTEXT("WithRangeBoneMultiplierIs1", "remap({SourceBone}.{SourceComponent})");
+					}
+					else
+					{
+						SourceArgs.Add(TEXT("Multiplier"), FText::AsNumber(Node.Multiplier));
+						FinalSourceExpression = LOCTEXT("WithRangeNonUnityMultiplier", "remap({SourceBone}.{SourceComponent}) * {Multiplier}");
+					}
 				}
 				else
 				{
-					SourceArgs.Add(TEXT("Multiplier"), FText::AsNumber(Node.Multiplier));
-					FinalSourceExpression = LOCTEXT("BoneMultiplierIs1", "{SourceBone}.{SourceComponent} * {Multiplier}");
+					if (Node.Multiplier == 1.0f)
+					{
+						FinalSourceExpression = LOCTEXT("BoneMultiplierIs1", "{SourceBone}.{SourceComponent}");
+					}
+					else
+					{
+						SourceArgs.Add(TEXT("Multiplier"), FText::AsNumber(Node.Multiplier));
+						FinalSourceExpression = LOCTEXT("NonUnityMultiplier", "{SourceBone}.{SourceComponent} * {Multiplier}");
+					}
 				}
 			}
 
@@ -252,12 +275,63 @@ void UAnimGraphNode_BoneDrivenController::AddTripletPropertyRow(const FText& Nam
 	];
 }
 
+void UAnimGraphNode_BoneDrivenController::AddRangePropertyRow(const FText& Name, const FText& Tooltip, IDetailCategoryBuilder& Category, TSharedRef<IPropertyHandle> PropertyHandle, const FName MinPropertyName, const FName MaxPropertyName, TAttribute<EVisibility> VisibilityAttribute)
+{
+	const float MiddlePadding = 4.0f;
+
+	TSharedPtr<IPropertyHandle> MinProperty = PropertyHandle->GetChildHandle(MinPropertyName);
+	Category.GetParentLayout().HideProperty(MinProperty);
+
+	TSharedPtr<IPropertyHandle> MaxProperty = PropertyHandle->GetChildHandle(MaxPropertyName);
+	Category.GetParentLayout().HideProperty(MaxProperty);
+
+	Category.AddCustomRow(Name)
+	.Visibility(VisibilityAttribute)
+	.NameContent()
+	[
+		SNew(STextBlock)
+		.Text(Name)
+		.ToolTipText(Tooltip)
+		.Font(IDetailLayoutBuilder::GetDetailFont())
+	]
+	.ValueContent()
+	.MinDesiredWidth(100.0f * 2.0f)
+	.MaxDesiredWidth(100.0f * 2.0f)
+	[
+		SNew(SHorizontalBox)
+
+		+SHorizontalBox::Slot()
+		.FillWidth(1)
+		.Padding(0.0f, 0.0f, MiddlePadding, 0.0f)
+		.VAlign(VAlign_Center)
+		[
+			MinProperty->CreatePropertyValueWidget()
+		]
+
+		+SHorizontalBox::Slot()
+		.AutoWidth()
+		[
+			SNew(STextBlock)
+			.Text(LOCTEXT("MinMaxSpacer", ".."))
+			.Font(IDetailLayoutBuilder::GetDetailFont())
+		]
+
+		+SHorizontalBox::Slot()
+		.FillWidth(1)
+		.Padding(MiddlePadding, 0.0f, 0.0f, 0.0f)
+		.VAlign(VAlign_Center)
+		[
+			MaxProperty->CreatePropertyValueWidget()
+		]
+	];
+}
+
 void UAnimGraphNode_BoneDrivenController::CustomizeDetails(IDetailLayoutBuilder& DetailBuilder)
 {
 	TSharedRef<IPropertyHandle> NodeHandle = DetailBuilder.GetProperty(GET_MEMBER_NAME_CHECKED(UAnimGraphNode_BoneDrivenController, Node), GetClass());
 
 	TAttribute<EVisibility> NotUsingCurveVisibility = TAttribute<EVisibility>::Create(TAttribute<EVisibility>::FGetter::CreateUObject(this, &UAnimGraphNode_BoneDrivenController::AreNonCurveMappingValuesVisible));
-
+	TAttribute<EVisibility> MapRangeVisiblity = TAttribute<EVisibility>::Create(TAttribute<EVisibility>::FGetter::CreateUObject(this, &UAnimGraphNode_BoneDrivenController::AreRemappingValuesVisible));
 
 	// Source (Driver) category
 	IDetailCategoryBuilder& SourceCategory = DetailBuilder.EditCategory(TEXT("Source (Driver)"));
@@ -266,15 +340,47 @@ void UAnimGraphNode_BoneDrivenController::CustomizeDetails(IDetailLayoutBuilder&
 	IDetailCategoryBuilder& MappingCategory = DetailBuilder.EditCategory(TEXT("Mapping"));
 	MappingCategory.AddProperty(NodeHandle->GetChildHandle(GET_MEMBER_NAME_CHECKED(FAnimNode_BoneDrivenController, DrivingCurve)));
 
-	MappingCategory.AddProperty(NodeHandle->GetChildHandle(GET_MEMBER_NAME_CHECKED(FAnimNode_BoneDrivenController, Multiplier))).Visibility(NotUsingCurveVisibility);
 	MappingCategory.AddProperty(NodeHandle->GetChildHandle(GET_MEMBER_NAME_CHECKED(FAnimNode_BoneDrivenController, bUseRange))).Visibility(NotUsingCurveVisibility);
-	MappingCategory.AddProperty(NodeHandle->GetChildHandle(GET_MEMBER_NAME_CHECKED(FAnimNode_BoneDrivenController, RangeMin))).Visibility(NotUsingCurveVisibility);
-	MappingCategory.AddProperty(NodeHandle->GetChildHandle(GET_MEMBER_NAME_CHECKED(FAnimNode_BoneDrivenController, RangeMax))).Visibility(NotUsingCurveVisibility);
+
+	AddRangePropertyRow(
+		/*Name=*/ LOCTEXT("InputRangeLabel", "Source Range"),
+		/*Tooltip=*/ LOCTEXT("InputRangeTooltip", "The range (relative to the reference pose) over which to limit the effect of the input component on the output component"),
+		/*Category=*/ MappingCategory,
+		/*PropertyHandle=*/ NodeHandle,
+		/*X=*/ GET_MEMBER_NAME_CHECKED(FAnimNode_BoneDrivenController, RangeMin),
+		/*Y=*/ GET_MEMBER_NAME_CHECKED(FAnimNode_BoneDrivenController, RangeMax),
+		MapRangeVisiblity);
+	AddRangePropertyRow(
+		/*Name=*/ LOCTEXT("MappedRangeLabel", "Mapped Range"),
+		/*Tooltip=*/ LOCTEXT("MappedRangeTooltip", "The range of mapped values that correspond to the input range"),
+		/*Category=*/ MappingCategory,
+		/*PropertyHandle=*/ NodeHandle,
+		/*X=*/ GET_MEMBER_NAME_CHECKED(FAnimNode_BoneDrivenController, RemappedMin),
+		/*Y=*/ GET_MEMBER_NAME_CHECKED(FAnimNode_BoneDrivenController, RemappedMax),
+		MapRangeVisiblity);
+
+	MappingCategory.AddProperty(NodeHandle->GetChildHandle(GET_MEMBER_NAME_CHECKED(FAnimNode_BoneDrivenController, Multiplier))).Visibility(NotUsingCurveVisibility);
 
 	// Destination (Driven) category
 	IDetailCategoryBuilder& TargetCategory = DetailBuilder.EditCategory(TEXT("Destination (Driven)"));
 	TargetCategory.AddProperty(NodeHandle->GetChildHandle(GET_MEMBER_NAME_CHECKED(FAnimNode_BoneDrivenController, TargetBone)));
 
+	// Add a note about the space (it is not configurable, and this helps set expectations)
+	const FText TargetBoneSpaceName = LOCTEXT("TargetComponentSpace", "Target Component Space");
+	TargetCategory.AddCustomRow(TargetBoneSpaceName)
+		.NameContent()
+		[
+			SNew(STextBlock)
+			.Text(TargetBoneSpaceName)
+			.Font(IDetailLayoutBuilder::GetDetailFont())
+		]
+		.ValueContent()
+		[
+			SNew(STextBlock)
+			.Text(LOCTEXT("TargetComponentSpaceIsAlwaysParentBoneSpace", "Parent Bone Space"))
+			.Font(IDetailLayoutBuilder::GetDetailFont())
+		];
+
 	AddTripletPropertyRow(
 		/*Name=*/ LOCTEXT("DrivenTranslationLabel", "Translation"),
 		/*Tooltip=*/ LOCTEXT("DrivenTranslationTooltip", "Should the source bone drive one or more translation components of the target bone?"),
@@ -337,4 +443,9 @@ EVisibility UAnimGraphNode_BoneDrivenController::AreNonCurveMappingValuesVisible
 	return (Node.DrivingCurve != nullptr) ? EVisibility::Collapsed : EVisibility::Visible;
 }
 
+EVisibility UAnimGraphNode_BoneDrivenController::AreRemappingValuesVisible() const
+{
+	return ((Node.DrivingCurve == nullptr) && Node.bUseRange) ? EVisibility::Visible : EVisibility::Collapsed;
+}
+
 #undef LOCTEXT_NAMESPACE
diff --git a/Engine/Source/Runtime/AnimGraphRuntime/Private/BoneControllers/AnimNode_BoneDrivenController.cpp b/Engine/Source/Runtime/AnimGraphRuntime/Private/BoneControllers/AnimNode_BoneDrivenController.cpp
--- a/Engine/Source/Runtime/AnimGraphRuntime/Private/BoneControllers/AnimNode_BoneDrivenController.cpp
+++ b/Engine/Source/Runtime/AnimGraphRuntime/Private/BoneControllers/AnimNode_BoneDrivenController.cpp
@@ -12,6 +12,8 @@ FAnimNode_BoneDrivenController::FAnimNode_BoneDrivenController()
 	, bUseRange(false)
 	, RangeMin(-1.0f)
 	, RangeMax(1.0f)
+	, RemappedMin(0.0f)
+	, RemappedMax(1.0f)
 	, TargetComponent_DEPRECATED(EComponentType::None)
 	, bAffectTargetTranslationX(false)
 	, bAffectTargetTranslationY(false)
@@ -89,7 +91,8 @@ void FAnimNode_BoneDrivenController::EvaluateBoneTransforms(USkeletalMeshCompone
 		// Apply the fixed function remapping/clamping
 		if (bUseRange)
 		{
-			FinalDriverValue = FMath::Clamp(FinalDriverValue, RangeMin, RangeMax);
+			const float Alpha = FMath::Clamp(FMath::GetRangePct(RangeMin, RangeMax, FinalDriverValue), 0.0f, 1.0f);
+			FinalDriverValue = FMath::Lerp(RemappedMin, RemappedMax, Alpha);
 		}
 
 		FinalDriverValue *= Multiplier;
diff --git a/Engine/Source/Runtime/AnimGraphRuntime/Public/AnimationCustomVersion.h b/Engine/Source/Runtime/AnimGraphRuntime/Public/AnimationCustomVersion.h
--- a/Engine/Source/Runtime/AnimGraphRuntime/Public/AnimationCustomVersion.h
+++ b/Engine/Source/Runtime/AnimGraphRuntime/Public/AnimationCustomVersion.h
@@ -14,6 +14,9 @@ struct ANIMGRAPHRUNTIME_API FAnimationCustomVersion
 		// changed the range to apply to the input, and added a configurable method for updating the components
 		BoneDrivenControllerMatchingMaya = 1,
 
+		// Converted the range clamp into a remap function, rather than just clamping
+		BoneDrivenControllerRemapping = 2,
+
 		// -----<new versions can be added above this line>-------------------------------------------------
 		VersionPlusOne,
 		LatestVersion = VersionPlusOne - 1
diff --git a/Engine/Source/Runtime/AnimGraphRuntime/Public/BoneControllers/AnimNode_BoneDrivenController.h b/Engine/Source/Runtime/AnimGraphRuntime/Public/BoneControllers/AnimNode_BoneDrivenController.h
--- a/Engine/Source/Runtime/AnimGraphRuntime/Public/BoneControllers/AnimNode_BoneDrivenController.h
+++ b/Engine/Source/Runtime/AnimGraphRuntime/Public/BoneControllers/AnimNode_BoneDrivenController.h
@@ -66,18 +66,26 @@ struct ANIMGRAPHRUNTIME_API FAnimNode_BoneDrivenController : public FAnimNode_Sk
 	UPROPERTY(EditAnywhere, Category=Mapping)
 	float Multiplier;
 
-	// Whether or not to clamp the driver value before scaling it (Note: Ignored when a curve is used)
-	UPROPERTY(EditAnywhere, Category="Source (Driver)", meta=(DisplayName="Limit Source Range"))
+	// Whether or not to clamp the driver value and remap it before scaling it
+	UPROPERTY(EditAnywhere, Category=Mapping, meta=(DisplayName="Remap Source"))
 	bool bUseRange;
 
-	// Minimum limit of the input value
-	UPROPERTY(EditAnywhere, Category="Source (Driver)", meta=(EditCondition=bUseRange, DisplayName="Source Range Min"))
+	// Minimum limit of the input value (mapped to RemappedMin, only used when limiting the source range)
+	UPROPERTY(EditAnywhere, Category=Mapping, meta=(EditCondition=bUseRange, DisplayName="Source Range Min"))
 	float RangeMin;
 
-	// Maximum limit of the input value
-	UPROPERTY(EditAnywhere, Category="Source (Driver)", meta=(EditCondition=bUseRange, DisplayName="Source Range Max"))
+	// Maximum limit of the input value (mapped to RemappedMax, only used when limiting the source range)
+	UPROPERTY(EditAnywhere, Category=Mapping, meta=(EditCondition=bUseRange, DisplayName="Source Range Max"))
 	float RangeMax;
 
+	// Minimum value to apply to the destination (remapped from the input range)
+	UPROPERTY(EditAnywhere, Category=Mapping, meta=(EditCondition=bUseRange, DisplayName="Mapped Range Min"))
+	float RemappedMin;
+
+	// Maximum value to apply to the destination (remapped from the input range)
+	UPROPERTY(EditAnywhere, Category = Mapping, meta = (EditCondition = bUseRange, DisplayName="Mapped Range Max"))
+	float RemappedMax;
+
 	// Bone to drive using controller input
 	UPROPERTY(EditAnywhere, Category="Destination (driven)")
 	FBoneReference TargetBone;
