Graph Editor: Exposed a way to remove a slot from a SNodePanel::SNode and added support for node bubbles rendering correctly when there is no top slot
Animation: Updated the observe bone node to not display the ...

diff --git a/Engine/Source/Editor/AnimGraph/Private/AnimGraphNode_ObserveBone.cpp b/Engine/Source/Editor/AnimGraph/Private/AnimGraphNode_ObserveBone.cpp
--- a/Engine/Source/Editor/AnimGraph/Private/AnimGraphNode_ObserveBone.cpp
+++ b/Engine/Source/Editor/AnimGraph/Private/AnimGraphNode_ObserveBone.cpp
@@ -36,11 +36,8 @@ public:
 	{
 		SGraphNode::UpdateGraphNode();
 
-		// Prevent the comment bubble from being displayed
-// 		GetOrAddSlot(ENodeZone::TopCenter)
-// 		[
-// 			SNullWidget::NullWidget
-// 		];
+		// Remove the comment bubble slot
+		RemoveSlot(ENodeZone::TopCenter);
 	}
 	// End of SGraphNode interface
 
@@ -192,4 +189,11 @@ void UAnimGraphNode_ObserveBone::CustomizePinData(UEdGraphPin* Pin, FName Source
 	}
 }
 
+void UAnimGraphNode_ObserveBone::CustomizeDetails(IDetailLayoutBuilder& DetailBuilder)
+{
+	TSharedRef<IPropertyHandle> NodeHandle = DetailBuilder.GetProperty(GET_MEMBER_NAME_CHECKED(UAnimGraphNode_ObserveBone, Node), GetClass());
+	DetailBuilder.HideProperty(NodeHandle->GetChildHandle(GET_MEMBER_NAME_CHECKED(FAnimNode_SkeletalControlBase, Alpha)));
+	DetailBuilder.HideProperty(NodeHandle->GetChildHandle(GET_MEMBER_NAME_CHECKED(FAnimNode_SkeletalControlBase, AlphaScaleBias)));
+}
+
 #undef LOCTEXT_NAMESPACE
diff --git a/Engine/Source/Editor/GraphEditor/Private/SGraphPanel.cpp b/Engine/Source/Editor/GraphEditor/Private/SGraphPanel.cpp
--- a/Engine/Source/Editor/GraphEditor/Private/SGraphPanel.cpp
+++ b/Engine/Source/Editor/GraphEditor/Private/SGraphPanel.cpp
@@ -229,8 +229,8 @@ int32 SGraphPanel::OnPaint( const FPaintArgs& Args, const FGeometry& AllottedGeo
 
 				// Draw the comments and information popups for this node, if it has any.
 				{
-					const SNodePanel::SNode::FNodeSlot& CommentSlot = ChildNode->GetOrAddSlot( ENodeZone::TopCenter );
-					float CommentBubbleY = -CommentSlot.Offset.Get().Y;
+					const SNodePanel::SNode::FNodeSlot* CommentSlot = ChildNode->GetSlot( ENodeZone::TopCenter );
+					float CommentBubbleY = CommentSlot ? -CommentSlot->Offset.Get().Y : 0.f;
 					Context.bSelected = bSelected;
 					TArray<FGraphInformationPopupInfo> Popups;
 
diff --git a/Engine/Source/Editor/GraphEditor/Public/SNodePanel.h b/Engine/Source/Editor/GraphEditor/Public/SNodePanel.h
--- a/Engine/Source/Editor/GraphEditor/Public/SNodePanel.h
+++ b/Engine/Source/Editor/GraphEditor/Public/SNodePanel.h
@@ -421,6 +421,33 @@ public:
 			return NewSlot;
 		}
 
+		FNodeSlot* GetSlot( const ENodeZone::Type SlotId )
+		{
+			FNodeSlot* Result = nullptr;
+			// Return existing
+			for( int32 ChildIndex = 0; ChildIndex < Children.Num(); ++ChildIndex )
+			{
+				if( Children[ ChildIndex ].Zone == SlotId )
+				{
+					Result = &Children[ ChildIndex ];
+					break;
+				}
+			}
+			return Result;
+		}
+
+		void RemoveSlot( const ENodeZone::Type SlotId )
+		{
+			for( int32 ChildIndex = 0; ChildIndex < Children.Num(); ++ChildIndex )
+			{
+				if( Children[ ChildIndex ].Zone == SlotId )
+				{
+					Children.RemoveAt( ChildIndex );
+					break;
+				}
+			}
+		}
+
 		/**
 		* @param NewPosition	The Node should be relocated to this position in the graph panel
 		* @param NodeFilter		Set of nodes to prevent movement on, after moving successfully a node is added to this set.
