Editor: Fixed order-of-initialization warning in SRotatorInputBox
Animation: Checking in missing file

diff --git a/Engine/Source/Editor/AnimGraph/Public/AnimGraphNode_ObserveBone.h b/Engine/Source/Editor/AnimGraph/Public/AnimGraphNode_ObserveBone.h
--- a/Engine/Source/Editor/AnimGraph/Public/AnimGraphNode_ObserveBone.h
+++ b/Engine/Source/Editor/AnimGraph/Public/AnimGraphNode_ObserveBone.h
@@ -7,6 +7,7 @@
 #include "EdGraph/EdGraphNodeUtils.h" // for FNodeTitleTextTable
 #include "AnimGraphNode_ObserveBone.generated.h"
 
+// This allows you to observe the state of a bone at a particular point in the graph, showing it in any space and optionally relative to the reference pose
 UCLASS()
 class ANIMGRAPH_API UAnimGraphNode_ObserveBone : public UAnimGraphNode_SkeletalControlBase
 {
@@ -27,6 +28,7 @@ protected:
 	// UAnimGraphNode_Base interface
 	virtual void ValidateAnimNodeDuringCompilation(USkeleton* ForSkeleton, FCompilerResultsLog& MessageLog) override;
 	virtual void CustomizePinData(UEdGraphPin* Pin, FName SourcePropertyName, int32 ArrayIndex) const override;
+	virtual void CustomizeDetails(IDetailLayoutBuilder& DetailBuilder) override;
 	// End of UAnimGraphNode_Base interface
 
 	// UAnimGraphNode_SkeletalControlBase interface
