Fix for crash undoing state nodes
#UE-18993

diff --git a/Engine/Source/Editor/AnimGraph/Private/AnimationStateMachineSchema.cpp b/Engine/Source/Editor/AnimGraph/Private/AnimationStateMachineSchema.cpp
--- a/Engine/Source/Editor/AnimGraph/Private/AnimationStateMachineSchema.cpp
+++ b/Engine/Source/Editor/AnimGraph/Private/AnimationStateMachineSchema.cpp
@@ -51,8 +51,6 @@ UEdGraphNode* FEdGraphSchemaAction_NewStateNode::PerformAction(class UEdGraph* P
 			FromPin->Modify();
 		}
 
-		NodeTemplate->SetFlags(RF_Transactional);
-
 		// set outer to be the graph so it doesn't go away
 		NodeTemplate->Rename(NULL, ParentGraph);
 		ParentGraph->AddNode(NodeTemplate, true, bSelectNewNode);
diff --git a/Engine/Source/Runtime/Engine/Private/EdGraph/EdGraphNode.cpp b/Engine/Source/Runtime/Engine/Private/EdGraph/EdGraphNode.cpp
--- a/Engine/Source/Runtime/Engine/Private/EdGraph/EdGraphNode.cpp
+++ b/Engine/Source/Runtime/Engine/Private/EdGraph/EdGraphNode.cpp
@@ -284,7 +284,10 @@ void UEdGraphNode::PostEditChangeProperty(FPropertyChangedEvent& PropertyChanged
 {
 	Super::PostEditChangeProperty(PropertyChangedEvent);
 
-	GetSchema()->ForceVisualizationCacheClear();
+	if(const UEdGraphSchema* Schema = GetSchema())
+	{
+		Schema->ForceVisualizationCacheClear();
+	}
 }
 
 void UEdGraphNode::CreateNewGuid()
