Fix for anim graph node properties with names in their names failing to update

diff --git a/Engine/Source/Editor/AnimGraph/Private/AnimGraphNode_Base.cpp b/Engine/Source/Editor/AnimGraph/Private/AnimGraphNode_Base.cpp
--- a/Engine/Source/Editor/AnimGraph/Private/AnimGraphNode_Base.cpp
+++ b/Engine/Source/Editor/AnimGraph/Private/AnimGraphNode_Base.cpp
@@ -248,7 +248,7 @@ FText UAnimGraphNode_Base::GetMenuCategory() const
 
 void UAnimGraphNode_Base::GetPinAssociatedProperty(const UScriptStruct* NodeType, UEdGraphPin* InputPin, UProperty*& OutProperty, int32& OutIndex)
 {
-	OutProperty = NULL;
+	OutProperty = nullptr;
 	OutIndex = INDEX_NONE;
 
 	//@TODO: Name-based hackery, avoid the roundtrip and better indicate when it's an array pose pin
@@ -264,7 +264,9 @@ void UAnimGraphNode_Base::GetPinAssociatedProperty(const UScriptStruct* NodeType
 			OutIndex = ArrayIndex;
 		}
 	}
-	else
+	
+	// If the array check failed or we have no underscores
+	if(OutProperty == nullptr)
 	{
 		if (UProperty* Property = FindField<UProperty>(NodeType, *(InputPin->PinName)))
 		{
