[INTEGRATE] Compile times: Working on eliminating <Module>Classes.h header files (part 2)
- Updated DataStreamCore and DataStreamEditor modules to not use *Classes.h files

diff --git a/Engine/Source/Editor/AnimGraph/Classes/AnimationConduitGraphSchema.h b/Engine/Source/Editor/AnimGraph/Classes/AnimationConduitGraphSchema.h
--- a/Engine/Source/Editor/AnimGraph/Classes/AnimationConduitGraphSchema.h
+++ b/Engine/Source/Editor/AnimGraph/Classes/AnimationConduitGraphSchema.h
@@ -2,6 +2,7 @@
 
 #pragma once
 
+#include "EdGraphSchema_K2.h"
 #include "AnimationConduitGraphSchema.generated.h"
 
 //@TODO: Should this derive from AnimationTransitionSchema (with appropriate suppression of state-based queries)
