Copying //UE4/Dev-Framework to Dev-Main (//UE4/Dev-Main)

diff --git a/Engine/Source/Developer/SimplygonMeshReduction/Private/SimplygonMeshReduction.cpp b/Engine/Source/Developer/SimplygonMeshReduction/Private/SimplygonMeshReduction.cpp
--- a/Engine/Source/Developer/SimplygonMeshReduction/Private/SimplygonMeshReduction.cpp
+++ b/Engine/Source/Developer/SimplygonMeshReduction/Private/SimplygonMeshReduction.cpp
@@ -371,25 +371,21 @@ public:
 				}
 			}
 
-			// if material hasn't been customized
-			if (SkeletalMesh->LODInfo[LODIndex].LODMaterialMap.Num() == 0)
+			// If base lod has a customized LODMaterialMap and this LOD doesn't (could have if changes are applied instead of freshly generated, copy over the data into new new LOD
+			if (SkeletalMesh->LODInfo[LODIndex].LODMaterialMap.Num() == 0 && SkeletalMesh->LODInfo[BaseLOD].LODMaterialMap.Num() != 0)
 			{
-				// copy parent index
-				for (int32 SectionId = 0; SectionId < NewModel->Sections.Num() && SectionId < SrcModel->Sections.Num(); ++SectionId)
-				{
-					NewModel->Sections[SectionId].MaterialIndex = SrcModel->Sections[SectionId].MaterialIndex;
-				}
-
-				// also if BaseLOD has material customized, follow it. 
 				SkeletalMesh->LODInfo[LODIndex].LODMaterialMap = SkeletalMesh->LODInfo[BaseLOD].LODMaterialMap;
 			}
 			else
 			{
+				// Assuming the reducing step has set all material indices correctly, we double check if something went wrong
 				// make sure we don't have more materials
 				int32 TotalSectionCount = NewModel->Sections.Num();
 				if (SkeletalMesh->LODInfo[LODIndex].LODMaterialMap.Num() > TotalSectionCount)
 				{
-					SkeletalMesh->LODInfo[LODIndex].LODMaterialMap.SetNum(TotalSectionCount);
+					SkeletalMesh->LODInfo[LODIndex].LODMaterialMap = SkeletalMesh->LODInfo[BaseLOD].LODMaterialMap;
+					// Something went wrong during the reduce step during regenerate 					
+					check(SkeletalMesh->LODInfo[BaseLOD].LODMaterialMap.Num() == TotalSectionCount);
 				}
 			}
 
@@ -2144,6 +2140,8 @@ private:
 
 		// Emissive
 		GetMaterialChannelData(SGMaterial, SimplygonSDK::SG_MATERIAL_CHANNEL_EMISSIVE, OutMaterial.EmissiveSamples, OutMaterial.EmissiveSize);
+		
+		GetMaterialChannelData(SGMaterial, USER_MATERIAL_CHANNEL_SUBSURFACE_COLOR, OutMaterial.SubSurfaceSamples, OutMaterial.SubSurfaceSize);
 	}
 
 	static FIntPoint ComputeMappingImageSize(const FMaterialSimplificationSettings& Settings)
@@ -2178,6 +2176,7 @@ private:
 			if (!SGMaterial->HasUserChannel(USER_MATERIAL_CHANNEL_OPACITY))
 				SGMaterial->AddUserChannel(USER_MATERIAL_CHANNEL_OPACITY);
 #endif
+			SGMaterial->AddUserChannel(USER_MATERIAL_CHANNEL_SUBSURFACE_COLOR);
 
 			SGMaterial->SetName(TCHAR_TO_ANSI(*FString::Printf(TEXT("Material%d"), MaterialIndex)));
 
@@ -2189,52 +2188,57 @@ private:
 					SGMaterial->SetVertexColorChannel(SimplygonSDK::SG_MATERIAL_CHANNEL_BASECOLOR, 0);
 				}
 
-				SetMaterialChannelData(FlattenMaterial.DiffuseSamples, FlattenMaterial.DiffuseSize, SGMaterial, SimplygonSDK::SG_MATERIAL_CHANNEL_BASECOLOR);
+				SetMaterialChannelData(FlattenMaterial.DiffuseSamples, FlattenMaterial.RenderSize, SGMaterial, SimplygonSDK::SG_MATERIAL_CHANNEL_BASECOLOR);
 			}
 
 			// Does current material have Metallic?
 			if (FlattenMaterial.MetallicSamples.Num())
 			{
-				SetMaterialChannelData(FlattenMaterial.MetallicSamples, FlattenMaterial.MetallicSize, SGMaterial, SimplygonSDK::SG_MATERIAL_CHANNEL_METALLIC);
+				SetMaterialChannelData(FlattenMaterial.MetallicSamples, FlattenMaterial.RenderSize, SGMaterial, SimplygonSDK::SG_MATERIAL_CHANNEL_METALLIC);
 			}
 
 			// Does current material have Specular?
 			if (FlattenMaterial.SpecularSamples.Num())
 			{
-				SetMaterialChannelData(FlattenMaterial.SpecularSamples, FlattenMaterial.SpecularSize, SGMaterial, SimplygonSDK::SG_MATERIAL_CHANNEL_SPECULAR);
+				SetMaterialChannelData(FlattenMaterial.SpecularSamples, FlattenMaterial.RenderSize, SGMaterial, SimplygonSDK::SG_MATERIAL_CHANNEL_SPECULAR);
 			}
 
 			// Does current material have Roughness?
 			if (FlattenMaterial.RoughnessSamples.Num())
 			{
-				SetMaterialChannelData(FlattenMaterial.RoughnessSamples, FlattenMaterial.RoughnessSize, SGMaterial, SimplygonSDK::SG_MATERIAL_CHANNEL_ROUGHNESS);
+				SetMaterialChannelData(FlattenMaterial.RoughnessSamples, FlattenMaterial.RenderSize, SGMaterial, SimplygonSDK::SG_MATERIAL_CHANNEL_ROUGHNESS);
 			}
 
 			//Does current material have a normalmap?
 			if (FlattenMaterial.NormalSamples.Num())
 			{
-				SetMaterialChannelData(FlattenMaterial.NormalSamples, FlattenMaterial.NormalSize, SGMaterial, SimplygonSDK::SG_MATERIAL_CHANNEL_NORMALS);
+				SetMaterialChannelData(FlattenMaterial.NormalSamples, FlattenMaterial.RenderSize, SGMaterial, SimplygonSDK::SG_MATERIAL_CHANNEL_NORMALS);
 			}
 
 			// Does current material have Opacity?
 			if (FlattenMaterial.OpacitySamples.Num())
 			{
 #if USE_USER_OPACITY_CHANNEL
-				SetMaterialChannelData(FlattenMaterial.OpacitySamples, FlattenMaterial.OpacitySize, SGMaterial, USER_MATERIAL_CHANNEL_OPACITY);
+				SetMaterialChannelData(FlattenMaterial.OpacitySamples, FlattenMaterial.RenderSize, SGMaterial, USER_MATERIAL_CHANNEL_OPACITY);
 #else
-				SetMaterialChannelData(FlattenMaterial.OpacitySamples, FlattenMaterial.OpacitySize, SGMaterial, SimplygonSDK::SG_MATERIAL_CHANNEL_OPACITY);
+				SetMaterialChannelData(FlattenMaterial.OpacitySamples, FlattenMaterial.RenderSize, SGMaterial, SimplygonSDK::SG_MATERIAL_CHANNEL_OPACITY);
 #endif
 			}
 
 			if (FlattenMaterial.EmissiveSamples.Num())
 			{
-				SetMaterialChannelData(FlattenMaterial.EmissiveSamples, FlattenMaterial.EmissiveSize, SGMaterial, SimplygonSDK::SG_MATERIAL_CHANNEL_EMISSIVE);
+				SetMaterialChannelData(FlattenMaterial.EmissiveSamples, FlattenMaterial.RenderSize, SGMaterial, SimplygonSDK::SG_MATERIAL_CHANNEL_EMISSIVE);
 			}
 			else
 			{
 				TArray<FColor> BlackEmissive;
 				BlackEmissive.AddZeroed(1);
-				SetMaterialChannelData(BlackEmissive, FlattenMaterial.EmissiveSize, SGMaterial, SimplygonSDK::SG_MATERIAL_CHANNEL_EMISSIVE);
+				SetMaterialChannelData(BlackEmissive, FlattenMaterial.RenderSize, SGMaterial, SimplygonSDK::SG_MATERIAL_CHANNEL_EMISSIVE);
+			}
+
+			if (FlattenMaterial.SubSurfaceSamples.Num())
+			{
+				SetMaterialChannelData(FlattenMaterial.SubSurfaceSamples, FlattenMaterial.RenderSize, SGMaterial, USER_MATERIAL_CHANNEL_SUBSURFACE_COLOR);
 			}
 
 			OutSGMaterialTable->AddMaterial(SGMaterial);
@@ -2450,6 +2454,8 @@ private:
 		if(!OutMaterial->HasUserChannel(USER_MATERIAL_CHANNEL_OPACITY))
 			OutMaterial->AddUserChannel(USER_MATERIAL_CHANNEL_OPACITY);
 #endif
+		OutMaterial->AddUserChannel(USER_MATERIAL_CHANNEL_SUBSURFACE_COLOR);
+
 		for(int ChannelIndex=0; ChannelIndex < InMaterialLODSettings.ChannelsToCast.Num(); ChannelIndex++)
 		{
 			FSimplygonChannelCastingSettings CasterSetting = InMaterialLODSettings.ChannelsToCast[ChannelIndex];
diff --git a/Engine/Source/Developer/SimplygonMeshReduction/Public/SimplygonTypes.h b/Engine/Source/Developer/SimplygonMeshReduction/Public/SimplygonTypes.h
--- a/Engine/Source/Developer/SimplygonMeshReduction/Public/SimplygonTypes.h
+++ b/Engine/Source/Developer/SimplygonMeshReduction/Public/SimplygonTypes.h
@@ -10,6 +10,9 @@
 
 #include "SimplygonTypes.generated.h"
 
+// User defined material channel used for baking out sub surface colours
+static const char* USER_MATERIAL_CHANNEL_SUBSURFACE_COLOR = "UserSubSurfaceColor";
+
 UENUM()
 namespace ESimplygonMaterialChannel
 {
@@ -24,7 +27,8 @@ namespace ESimplygonMaterialChannel
 		SG_MATERIAL_CHANNEL_DISPLACEMENT UMETA(DisplayName = "Displacement", DisplayValue = "Displacement"),
 		SG_MATERIAL_CHANNEL_BASECOLOR UMETA(DisplayName = "Basecolor", DisplayValue = "Basecolor"),
 		SG_MATERIAL_CHANNEL_ROUGHNESS UMETA(DisplayName = "Roughness", DisplayValue = "Roughness"),
-		SG_MATERIAL_CHANNEL_METALLIC UMETA(DisplayName = "Metallic", DisplayValue = "Metallic")
+		SG_MATERIAL_CHANNEL_METALLIC UMETA(DisplayName = "Metallic", DisplayValue = "Metallic"),
+		SG_MATERIAL_CHANNEL_SUBSURFACE UMETA(DisplayName = "SubsurfaceColor", DisplayValue = "SubsurfaceColor")
 	};
 
 }
@@ -49,8 +53,13 @@ static const char* GetSimplygonMaterialChannel(ESimplygonMaterialChannel::Type c
 #endif
 	else if (channel == ESimplygonMaterialChannel::SG_MATERIAL_CHANNEL_EMISSIVE)
 		return SimplygonSDK::SG_MATERIAL_CHANNEL_EMISSIVE;
+	else if (channel == ESimplygonMaterialChannel::SG_MATERIAL_CHANNEL_SUBSURFACE)
+		return USER_MATERIAL_CHANNEL_SUBSURFACE_COLOR;
 	else
-		return SimplygonSDK::SG_MATERIAL_CHANNEL_BASECOLOR;
+	{
+		check(0);
+		return 0;
+	}
 
 }
 
@@ -392,7 +401,7 @@ struct FSimplygonMaterialLODSettings
 	FSimplygonMaterialLODSettings(const FMaterialProxySettings& Settings)
 		: bActive(true)
 		, MaterialLODType(EMaterialLODType::BakeTexture)
-		, bUseAutomaticSizes( Settings.bAutomaticTextureSizes )
+		, bUseAutomaticSizes( Settings.TextureSizingType == TextureSizingType_UseSimplygonAutomaticSizing )
 		, TextureWidth(GetResolutionEnum(Settings.TextureSize.X))
 		, TextureHeight(GetResolutionEnum(Settings.TextureSize.Y))
 		, SamplingQuality(ESimplygonTextureSamplingQuality::High)
@@ -401,18 +410,23 @@ struct FSimplygonMaterialLODSettings
 		, bReuseExistingCharts(false)
 	{
 		ChannelsToCast.Add(FSimplygonChannelCastingSettings(ESimplygonMaterialChannel::SG_MATERIAL_CHANNEL_BASECOLOR, ESimplygonCasterType::Color, ESimplygonColorChannels::RGB));
+		ChannelsToCast.Last().bUseSRGB = false;
 		ChannelsToCast.Last().bActive = true;
 
 		ChannelsToCast.Add(FSimplygonChannelCastingSettings(ESimplygonMaterialChannel::SG_MATERIAL_CHANNEL_SPECULAR, ESimplygonCasterType::Color, ESimplygonColorChannels::RGB));
+		ChannelsToCast.Last().bUseSRGB = false;
 		ChannelsToCast.Last().bActive = Settings.bSpecularMap;
 
 		ChannelsToCast.Add(FSimplygonChannelCastingSettings(ESimplygonMaterialChannel::SG_MATERIAL_CHANNEL_ROUGHNESS, ESimplygonCasterType::Color, ESimplygonColorChannels::RGB));
+		ChannelsToCast.Last().bUseSRGB = false;
 		ChannelsToCast.Last().bActive = Settings.bRoughnessMap;
 
 		ChannelsToCast.Add(FSimplygonChannelCastingSettings(ESimplygonMaterialChannel::SG_MATERIAL_CHANNEL_METALLIC, ESimplygonCasterType::Color, ESimplygonColorChannels::RGB));
+		ChannelsToCast.Last().bUseSRGB = false;
 		ChannelsToCast.Last().bActive = Settings.bMetallicMap;
 
 		ChannelsToCast.Add(FSimplygonChannelCastingSettings(ESimplygonMaterialChannel::SG_MATERIAL_CHANNEL_NORMALS, ESimplygonCasterType::Normals, ESimplygonColorChannels::RGB));
+		ChannelsToCast.Last().bUseSRGB = false;
 		ChannelsToCast.Last().bActive = Settings.bNormalMap;
 
 		ChannelsToCast.Add(FSimplygonChannelCastingSettings(ESimplygonMaterialChannel::SG_MATERIAL_CHANNEL_EMISSIVE, ESimplygonCasterType::Color, ESimplygonColorChannels::RGB));
@@ -420,7 +434,13 @@ struct FSimplygonMaterialLODSettings
 		ChannelsToCast.Last().bActive = Settings.bEmissiveMap;
 
 		ChannelsToCast.Add(FSimplygonChannelCastingSettings(ESimplygonMaterialChannel::SG_MATERIAL_CHANNEL_OPACITY, ESimplygonCasterType::Color, ESimplygonColorChannels::RGB));
+		ChannelsToCast.Last().bUseSRGB = false;
 		ChannelsToCast.Last().bActive = Settings.bOpacityMap;
+
+		// TODO, properly render out sub surface data/values
+		//ChannelsToCast.Add(FSimplygonChannelCastingSettings(ESimplygonMaterialChannel::SG_MATERIAL_CHANNEL_SUBSURFACE, ESimplygonCasterType::Color, ESimplygonColorChannels::RGB));
+		//ChannelsToCast.Last().bUseSRGB = false;
+		//ChannelsToCast.Last().bActive = true;
 	}
 
 	static int32 GetTextureResolutionFromEnum(ESimplygonTextureResolution::Type InResolution)
diff --git a/Engine/Source/Developer/SimplygonSwarm/Private/SimplygonRESTClient.cpp b/Engine/Source/Developer/SimplygonSwarm/Private/SimplygonRESTClient.cpp
--- a/Engine/Source/Developer/SimplygonSwarm/Private/SimplygonRESTClient.cpp
+++ b/Engine/Source/Developer/SimplygonSwarm/Private/SimplygonRESTClient.cpp
@@ -1,7 +1,7 @@
 #include "SimplygonSwarmPrivatePCH.h"
 #include "SimplygonRESTClient.h"
 
-#define HOSTNAME "127.0.0.1"
+#define HOSTNAME "http://127.0.0.1"
 #define PORT ":55002"
 
 
@@ -56,7 +56,14 @@ FSimplygonRESTClient::FSimplygonRESTClient(TSharedPtr<FSimplygonSwarmTask> &InTa
 	FString IP = GetDefault<UEditorPerProjectUserSettings>()->SimplygonServerIP;
 	if (IP != "")
 	{
-		HostName = IP;
+		if (!IP.Contains("http://"))
+		{
+			HostName = "http://" + IP;
+		}
+		else
+		{
+			HostName = IP;
+		}
 	}
 	else
 	{
diff --git a/Engine/Source/Developer/SimplygonSwarm/Private/SimplygonSwarm.cpp b/Engine/Source/Developer/SimplygonSwarm/Private/SimplygonSwarm.cpp
--- a/Engine/Source/Developer/SimplygonSwarm/Private/SimplygonSwarm.cpp
+++ b/Engine/Source/Developer/SimplygonSwarm/Private/SimplygonSwarm.cpp
@@ -870,7 +870,7 @@
 		 FString CommandLine = FString::Printf(TEXT("x %s -o%s * -y"), *FPaths::ConvertRelativePathToFull(ZipFileName), *FPaths::ConvertRelativePathToFull(OutputFolderPath));
 
 		 FString CmdExe = TEXT("cmd.exe");
-		 FString ZipToolPath = FPaths::ConvertRelativePathToFull(FPaths::EngineDir() / TEXT("Extras/7-Zip/7z.exe"));
+		 FString ZipToolPath = FPaths::ConvertRelativePathToFull(FPaths::EngineDir() / TEXT("Binaries/ThirdParty/NotForLicensees/7-Zip/7z.exe"));
 
 		 if (!FPaths::FileExists(ZipToolPath))
 		 {
@@ -898,7 +898,7 @@
 		 FString CommandLine = FString::Printf(TEXT("a %s %s/*"), *FPaths::ConvertRelativePathToFull(OutputFileName), *FPaths::ConvertRelativePathToFull(InputDirectoryPath));
 
 		 FString CmdExe = TEXT("cmd.exe");
-		 FString ZipToolPath = FPaths::ConvertRelativePathToFull(FPaths::EngineDir() / TEXT("Extras/7-Zip/7z.exe"));
+		 FString ZipToolPath = FPaths::ConvertRelativePathToFull(FPaths::EngineDir() / TEXT("Binaries/ThirdParty/NotForLicensees/7-Zip/7z.exe"));
 
 		 if (!FPaths::FileExists(ZipToolPath))
 		 {
diff --git a/Engine/Source/Editor/AIGraph/Private/AIGraphEditor.cpp b/Engine/Source/Editor/AIGraph/Private/AIGraphEditor.cpp
--- a/Engine/Source/Editor/AIGraph/Private/AIGraphEditor.cpp
+++ b/Engine/Source/Editor/AIGraph/Private/AIGraphEditor.cpp
@@ -235,7 +235,8 @@ void FAIGraphEditor::CopySelectedNodes()
 	int32 CopySubNodeIndex = 0;
 	for (FGraphPanelSelectionSet::TIterator SelectedIter(SelectedNodes); SelectedIter; ++SelectedIter)
 	{
-		UAIGraphNode* Node = Cast<UAIGraphNode>(*SelectedIter);
+		UEdGraphNode* Node = Cast<UEdGraphNode>(*SelectedIter);
+		UAIGraphNode* AINode = Cast<UAIGraphNode>(Node);
 		if (Node == nullptr)
 		{
 			SelectedIter.RemoveCurrent();
@@ -243,16 +244,20 @@ void FAIGraphEditor::CopySelectedNodes()
 		}
 
 		Node->PrepareForCopying();
-		Node->CopySubNodeIndex = CopySubNodeIndex;
 
-		// append all subnodes for selection
-		for (int32 Idx = 0; Idx < Node->SubNodes.Num(); Idx++)
+		if (AINode)
 		{
-			Node->SubNodes[Idx]->CopySubNodeIndex = CopySubNodeIndex;
-			SubNodes.Add(Node->SubNodes[Idx]);
-		}
+			AINode->CopySubNodeIndex = CopySubNodeIndex;
+
+			// append all subnodes for selection
+			for (int32 Idx = 0; Idx < AINode->SubNodes.Num(); Idx++)
+			{
+				AINode->SubNodes[Idx]->CopySubNodeIndex = CopySubNodeIndex;
+				SubNodes.Add(AINode->SubNodes[Idx]);
+			}
 
-		CopySubNodeIndex++;
+			CopySubNodeIndex++;
+		}
 	}
 
 	for (int32 Idx = 0; Idx < SubNodes.Num(); Idx++)
@@ -309,9 +314,14 @@ void FAIGraphEditor::PasteNodesHere(const FVector2D& Location)
 
 	// Undo/Redo support
 	const FScopedTransaction Transaction(FGenericCommands::Get().Paste->GetDescription());
-	UAIGraph* EdGraph = Cast<UAIGraph>(CurrentGraphEditor->GetCurrentGraph());
+	UEdGraph* EdGraph = CurrentGraphEditor->GetCurrentGraph();
+	UAIGraph* AIGraph = Cast<UAIGraph>(EdGraph);
+
 	EdGraph->Modify();
-	EdGraph->LockUpdates();
+	if (AIGraph)
+	{
+		AIGraph->LockUpdates();
+	}
 
 	UAIGraphNode* SelectedParent = NULL;
 	bool bHasMultipleNodesSelected = false;
@@ -355,11 +365,12 @@ void FAIGraphEditor::PasteNodesHere(const FVector2D& Location)
 
 	for (TSet<UEdGraphNode*>::TIterator It(PastedNodes); It; ++It)
 	{
-		UAIGraphNode* Node = Cast<UAIGraphNode>(*It);
-		if (Node && !Node->IsSubNode())
+		UEdGraphNode* EdNode = *It;
+		UAIGraphNode* AINode = Cast<UAIGraphNode>(EdNode);
+		if (EdNode && (AINode == nullptr || !AINode->IsSubNode()))
 		{
-			AvgNodePosition.X += Node->NodePosX;
-			AvgNodePosition.Y += Node->NodePosY;
+			AvgNodePosition.X += EdNode->NodePosX;
+			AvgNodePosition.Y += EdNode->NodePosY;
 		}
 	}
 
@@ -375,8 +386,10 @@ void FAIGraphEditor::PasteNodesHere(const FVector2D& Location)
 	TMap<int32, UAIGraphNode*> ParentMap;
 	for (TSet<UEdGraphNode*>::TIterator It(PastedNodes); It; ++It)
 	{
-		UAIGraphNode* PasteNode = Cast<UAIGraphNode>(*It);
-		if (PasteNode && !PasteNode->IsSubNode())
+		UEdGraphNode* PasteNode = *It;
+		UAIGraphNode* PasteAINode = Cast<UAIGraphNode>(PasteNode);
+
+		if (PasteNode && (PasteAINode == nullptr || !PasteAINode->IsSubNode()))
 		{
 			bPastedParentNode = true;
 
@@ -391,8 +404,11 @@ void FAIGraphEditor::PasteNodesHere(const FVector2D& Location)
 			// Give new node a different Guid from the old one
 			PasteNode->CreateNewGuid();
 
-			PasteNode->RemoveAllSubNodes();
-			ParentMap.Add(PasteNode->CopySubNodeIndex, PasteNode);
+			if (PasteAINode)
+			{
+				PasteAINode->RemoveAllSubNodes();
+				ParentMap.Add(PasteAINode->CopySubNodeIndex, PasteAINode);
+			}
 		}
 	}
 
@@ -420,10 +436,13 @@ void FAIGraphEditor::PasteNodesHere(const FVector2D& Location)
 		}
 	}
 
-	EdGraph->UpdateClassData();
+	if (AIGraph)
+	{
+		AIGraph->UpdateClassData();
+		AIGraph->UnlockUpdates();
+	}
 
 	// Update UI
-	EdGraph->UnlockUpdates();
 	CurrentGraphEditor->NotifyGraphChanged();
 
 	UObject* GraphOwner = EdGraph->GetOuter();
diff --git a/Engine/Source/Editor/AnimGraph/Classes/AnimPreviewInstance.h b/Engine/Source/Editor/AnimGraph/Classes/AnimPreviewInstance.h
--- a/Engine/Source/Editor/AnimGraph/Classes/AnimPreviewInstance.h
+++ b/Engine/Source/Editor/AnimGraph/Classes/AnimPreviewInstance.h
@@ -6,6 +6,7 @@
 #include "AnimGraphNode_ModifyBone.h"
 #include "Animation/AnimSingleNodeInstance.h"
 #include "BoneControllers/AnimNode_ModifyBone.h"
+#include "Animation/AnimSingleNodeInstanceProxy.h"
 #include "AnimPreviewInstance.generated.h"
 
 /** Enum to know how montage is being played */
@@ -19,6 +20,111 @@ enum EMontagePreviewType
 	EMPT_MAX,
 };
 
+
+/** Proxy override for this UAnimInstance-derived class */
+USTRUCT()
+struct FAnimPreviewInstanceProxy : public FAnimSingleNodeInstanceProxy
+{
+	GENERATED_BODY()
+
+public:
+	FAnimPreviewInstanceProxy()
+	{
+	}
+
+	FAnimPreviewInstanceProxy(UAnimInstance* InAnimInstance)
+		: FAnimSingleNodeInstanceProxy(InAnimInstance)
+		, SkeletalControlAlpha(1.0f)
+#if WITH_EDITORONLY_DATA
+		, bForceRetargetBasePose(false)
+#endif
+		, bEnableControllers(true)
+		, bSetKey(false)
+	{
+	}
+
+	virtual void Initialize(UAnimInstance* InAnimInstance) override;
+	virtual void Update(float DeltaSeconds) override;
+	virtual bool Evaluate(FPoseContext& Output) override;
+
+	void ResetModifiedBone(bool bCurveController = false);
+
+	FAnimNode_ModifyBone* FindModifiedBone(const FName& InBoneName, bool bCurveController = false);	
+
+	FAnimNode_ModifyBone& ModifyBone(const FName& InBoneName, bool bCurveController = false);
+
+	void RemoveBoneModification(const FName& InBoneName, bool bCurveController = false);
+
+	void SetForceRetargetBasePose(bool bInForceRetargetBasePose)
+	{ 
+		bForceRetargetBasePose = bInForceRetargetBasePose; 
+	}
+
+	bool GetForceRetargetBasePose() const 
+	{ 
+		return bForceRetargetBasePose; 
+	}
+
+	void EnableControllers(bool bEnable)
+	{
+		bEnableControllers = bEnable;
+	}
+
+	void SetSkeletalControlAlpha(float InSkeletalControlAlpha)
+	{
+		SkeletalControlAlpha = FMath::Clamp<float>(InSkeletalControlAlpha, 0.f, 1.f);
+	}
+
+	void SetKey(FSimpleDelegate InOnSetKeyCompleteDelegate)
+	{
+#if WITH_EDITOR
+		bSetKey = true;
+		OnSetKeyCompleteDelegate = InOnSetKeyCompleteDelegate;
+#endif
+	}
+
+	void RefreshCurveBoneControllers();
+
+private:
+	void UpdateCurveController();
+
+	void ApplyBoneControllers(USkeletalMeshComponent* Component, TArray<FAnimNode_ModifyBone> &InBoneControllers, FCSPose<FCompactPose>& OutMeshPose);
+
+	void SetKeyImplementation(const FCompactPose& PreControllerInLocalSpace, const FCompactPose& PostControllerInLocalSpace);
+
+	void AddKeyToSequence(UAnimSequence* Sequence, float Time, const FName& BoneName, const FTransform& AdditiveTransform);
+
+private:
+	/** Controllers for individual bones */
+	TArray<FAnimNode_ModifyBone> BoneControllers;
+
+	/** Curve modifiers */
+	TArray<FAnimNode_ModifyBone> CurveBoneControllers;
+
+	/**
+	 * Delegate to call after Key is set
+	 */
+	FSimpleDelegate OnSetKeyCompleteDelegate;
+
+	/** Shared parameters for previewing blendspace or animsequence **/
+	float SkeletalControlAlpha;
+
+#if WITH_EDITORONLY_DATA
+	bool bForceRetargetBasePose;
+#endif
+
+	/*
+	 * Used to determine if controller has to be applied or not
+	 * Used to disable controller during editing
+	 */
+	bool bEnableControllers;
+
+	/* 
+	 * When this flag is true, it sets key
+	 */
+	bool bSetKey;
+};
+
 /**
  * This Instance only contains one AnimationAsset, and produce poses
  * Used by Preview in AnimGraph, Playing single animation in Kismet2 and etc
@@ -29,16 +135,21 @@ class ANIMGRAPH_API UAnimPreviewInstance : public UAnimSingleNodeInstance
 {
 	GENERATED_UCLASS_BODY()
 
+	// Disable compiler-generated deprecation warnings by implementing our own destructor
+	PRAGMA_DISABLE_DEPRECATION_WARNINGS
+	~UAnimPreviewInstance() {}
+	PRAGMA_ENABLE_DEPRECATION_WARNINGS
+
 	/** Controllers for individual bones */
-	UPROPERTY(transient)
+	DEPRECATED(4.11, "This cannot be accessed directly as it is potentially in use on worker threads")
 	TArray<FAnimNode_ModifyBone> BoneControllers;
 
 	/** Curve modifiers */
-	UPROPERTY(transient)
+	DEPRECATED(4.11, "This cannot be accessed directly as it is potentially in use on worker threads")
 	TArray<FAnimNode_ModifyBone> CurveBoneControllers;
 
 	/** Shared parameters for previewing blendspace or animsequence **/
-	UPROPERTY(transient)
+	DEPRECATED(4.11, "This cannot be accessed directly as it is potentially in use on worker threads")
 	float SkeletalControlAlpha;
 
 	/** Shared parameters for previewing blendspace or animsequence **/
@@ -50,8 +161,7 @@ class ANIMGRAPH_API UAnimPreviewInstance : public UAnimSingleNodeInstance
 
 	//~ Begin UAnimInstance Interface
 	virtual void NativeInitializeAnimation() override;
-	virtual void UpdateAnimationNode(float DeltaTimeX) override;
-	virtual bool NativeEvaluateAnimation(FPoseContext& Output) override;
+	virtual FAnimInstanceProxy* CreateAnimInstanceProxy() override;
 	//~ End UAnimInstance Interface
 
 	/** Set SkeletalControl Alpha**/
@@ -112,7 +222,12 @@ class ANIMGRAPH_API UAnimPreviewInstance : public UAnimSingleNodeInstance
 	 */
 	void ResetModifiedBone(bool bCurveController=false);
 
+	void SetForceRetargetBasePose(bool ForceRetargetBasePose);
+
+	bool GetForceRetargetBasePose() const;
+
 #if WITH_EDITORONLY_DATA
+	DEPRECATED(4.11, "This cannot be accessed directly as it is potentially in use on worker threads, use SetForceRetargetBasePose/GetForceRetargetBasePose")
 	bool bForceRetargetBasePose;
 #endif
 
@@ -141,46 +256,6 @@ class ANIMGRAPH_API UAnimPreviewInstance : public UAnimSingleNodeInstance
 	 * This is used by when editing, when controller has to be disabled
 	 */
 	void EnableControllers(bool bEnable);
-
-private:
-	/** 
-	 * Apply Bone Controllers to the Outpose
-	 *
-	 * @param	Component	Component to apply bone controller to
-	 * @param	BoneControllers	 List of Bone Controllers to apply
-	 * @param 	OutMeshPose	Outpose in Mesh Space once applied
-	 */
-	void ApplyBoneControllers(USkeletalMeshComponent* Component, TArray<FAnimNode_ModifyBone> &BoneControllers, FCSPose<FCompactPose>& OutMeshPose);
-	/** 
-	 * Update CurveControllers based on TransformCurves of Animation
-	 */
-	void UpdateCurveController();
-
-	/* 
-	 * Set Key Implementation function
-	 * It gets Pre Controller Local Space and gets Post Controller Local Space, and add the key to the curve 
-	 */
-	void SetKeyImplementation(const FCompactPose& PreControllerInLocalSpace, const FCompactPose& PostControllerInLocalSpace);
-	/** 
-	 * Add Key to the Sequence
-	 * Now Additive Key is generated, add to the curves
-	 */
-	void AddKeyToSequence(UAnimSequence* Sequence, float Time, const FName& BoneName, const FTransform& AdditiveTransform);
-	/* 
-	 * When this flag is true, it sets key
-	 */
-	bool bSetKey;
-
-	/*
-	 * Used to determine if controller has to be applied or not
-	 * Used to disable controller during editing
-	 */
-	bool bEnableControllers;
-	
-	/**
-	 * Delegate to call after Key is set
-	 */
-	FSimpleDelegate OnSetKeyCompleteDelegate;
 };
 
 
diff --git a/Engine/Source/Editor/AnimGraph/Private/AnimGraphNode_StateMachineBase.cpp b/Engine/Source/Editor/AnimGraph/Private/AnimGraphNode_StateMachineBase.cpp
--- a/Engine/Source/Editor/AnimGraph/Private/AnimGraphNode_StateMachineBase.cpp
+++ b/Engine/Source/Editor/AnimGraph/Private/AnimGraphNode_StateMachineBase.cpp
@@ -111,6 +111,7 @@ void UAnimGraphNode_StateMachineBase::PostPlacedNewNode()
 	Schema->CreateDefaultNodesForGraph(*EditorStateMachineGraph);
 
 	// Add the new graph as a child of our parent graph
+	GetGraph()->Modify();
 	GetGraph()->SubGraphs.Add(EditorStateMachineGraph);
 }
 
diff --git a/Engine/Source/Editor/AnimGraph/Private/AnimPreviewInstance.cpp b/Engine/Source/Editor/AnimGraph/Private/AnimPreviewInstance.cpp
--- a/Engine/Source/Editor/AnimGraph/Private/AnimPreviewInstance.cpp
+++ b/Engine/Source/Editor/AnimGraph/Private/AnimPreviewInstance.cpp
@@ -10,37 +10,20 @@
 
 #define LOCTEXT_NAMESPACE "AnimPreviewInstance"
 
-UAnimPreviewInstance::UAnimPreviewInstance(const FObjectInitializer& ObjectInitializer)
-	: Super(ObjectInitializer)
-	, SkeletalControlAlpha(1.0f)
-#if WITH_EDITORONLY_DATA
-	, bForceRetargetBasePose(false)
-#endif
-	, bSetKey(false)
+void FAnimPreviewInstanceProxy::Initialize(UAnimInstance* InAnimInstance)
 {
-	RootMotionMode = ERootMotionMode::RootMotionFromEverything;
-	bEnableControllers = true;
-}
+	FAnimSingleNodeInstanceProxy::Initialize(InAnimInstance);
 
-void UAnimPreviewInstance::NativeInitializeAnimation()
-{
-	// Cache our play state from the previous animation otherwise set to play
-	bool bCachedIsPlaying = (CurrentAsset != NULL) ? bPlaying : true;
 	bSetKey = false;
-	Super::NativeInitializeAnimation();
-
-	SetPlaying(bCachedIsPlaying);
-
-	RefreshCurveBoneControllers();
 }
 
-void UAnimPreviewInstance::ResetModifiedBone(bool bCurveController/*=false*/)
+void FAnimPreviewInstanceProxy::ResetModifiedBone(bool bCurveController)
 {
 	TArray<FAnimNode_ModifyBone>& Controllers = (bCurveController)?CurveBoneControllers : BoneControllers;
 	Controllers.Empty();
 }
 
-FAnimNode_ModifyBone* UAnimPreviewInstance::FindModifiedBone(const FName& InBoneName, bool bCurveController/*=false*/)
+FAnimNode_ModifyBone* FAnimPreviewInstanceProxy::FindModifiedBone(const FName& InBoneName, bool bCurveController)
 {
 	TArray<FAnimNode_ModifyBone>& Controllers = (bCurveController)?CurveBoneControllers : BoneControllers;
 
@@ -52,7 +35,7 @@ FAnimNode_ModifyBone* UAnimPreviewInstance::FindModifiedBone(const FName& InBone
 	);
 }
 
-FAnimNode_ModifyBone& UAnimPreviewInstance::ModifyBone(const FName& InBoneName, bool bCurveController/*=false*/)
+FAnimNode_ModifyBone& FAnimPreviewInstanceProxy::ModifyBone(const FName& InBoneName, bool bCurveController)
 {
 	FAnimNode_ModifyBone* SingleBoneController = FindModifiedBone(InBoneName, bCurveController);
 	TArray<FAnimNode_ModifyBone>& Controllers = (bCurveController)?CurveBoneControllers : BoneControllers;
@@ -91,7 +74,7 @@ FAnimNode_ModifyBone& UAnimPreviewInstance::ModifyBone(const FName& InBoneName,
 	return *SingleBoneController;
 }
 
-void UAnimPreviewInstance::RemoveBoneModification(const FName& InBoneName, bool bCurveController/*=false*/)
+void FAnimPreviewInstanceProxy::RemoveBoneModification(const FName& InBoneName, bool bCurveController)
 {
 	TArray<FAnimNode_ModifyBone>& Controllers = (bCurveController)?CurveBoneControllers : BoneControllers;
 	Controllers.RemoveAll(
@@ -102,8 +85,11 @@ void UAnimPreviewInstance::RemoveBoneModification(const FName& InBoneName, bool
 	);
 }
 
-void UAnimPreviewInstance::UpdateAnimationNode(float DeltaTimeX)
+void FAnimPreviewInstanceProxy::Update(float DeltaSeconds)
 {
+	// we cant update on a worker thread here because of the key delegate needing to be fired
+	check(IsInGameThread());
+
 #if WITH_EDITORONLY_DATA
 	if(bForceRetargetBasePose)
 	{
@@ -112,15 +98,18 @@ void UAnimPreviewInstance::UpdateAnimationNode(float DeltaTimeX)
 	}
 #endif // #if WITH_EDITORONLY_DATA
 
-	Super::UpdateAnimationNode(DeltaTimeX);
+	FAnimSingleNodeInstanceProxy::Update(DeltaSeconds);
 }
 
-bool UAnimPreviewInstance::NativeEvaluateAnimation(FPoseContext& Output)
+bool FAnimPreviewInstanceProxy::Evaluate(FPoseContext& Output)
 {
+	// we cant evaluate on a worker thread here because of the key delegate needing to be fired
+	check(IsInGameThread());
+
 #if WITH_EDITORONLY_DATA
 	if(bForceRetargetBasePose)
 	{
-		USkeletalMeshComponent* MeshComponent = GetSkelMeshComponent();
+		USkeletalMeshComponent* MeshComponent = Output.AnimInstanceProxy->GetSkelMeshComponent();
 		if(MeshComponent && MeshComponent->SkeletalMesh)
 		{
 			FAnimationRuntime::FillWithRetargetBaseRefPose(Output.Pose, GetSkelMeshComponent()->SkeletalMesh);
@@ -134,14 +123,14 @@ bool UAnimPreviewInstance::NativeEvaluateAnimation(FPoseContext& Output)
 	else
 #endif // #if WITH_EDITORONLY_DATA
 	{
-		Super::NativeEvaluateAnimation(Output);
+		FAnimSingleNodeInstanceProxy::Evaluate(Output);
 	}
 
 	if (bEnableControllers)
 	{
 		UDebugSkelMeshComponent* Component = Cast<UDebugSkelMeshComponent>(GetSkelMeshComponent());
 
-		if(Component && CurrentSkeleton)
+		if(Component && GetSkeleton())
 		{
 			// update curve controllers
 			UpdateCurveController();
@@ -198,11 +187,105 @@ bool UAnimPreviewInstance::NativeEvaluateAnimation(FPoseContext& Output)
 	return true;
 }
 
-void UAnimPreviewInstance::ApplyBoneControllers(USkeletalMeshComponent* Component, TArray<FAnimNode_ModifyBone> &InBoneControllers, FCSPose<FCompactPose>& OutMeshPose)
+void FAnimPreviewInstanceProxy::RefreshCurveBoneControllers()
+{
+	// go through all curves and see if it has Transform Curve
+	// if so, find what bone that belong to and create BoneMOdifier for them
+	UAnimSequence* CurrentSequence = Cast<UAnimSequence>(GetCurrentAsset());
+
+	CurveBoneControllers.Empty();
+
+	// do not apply if BakedAnimation is on
+	if(CurrentSequence)
+	{
+		// make sure if this needs source update
+		if ( !CurrentSequence->DoesContainTransformCurves() )
+		{
+			return;
+		}
+
+		GetRequiredBones().SetUseSourceData(true);
+
+		TArray<FTransformCurve>& Curves = CurrentSequence->RawCurveData.TransformCurves;
+		FSmartNameMapping* NameMapping = GetSkeleton()->SmartNames.GetContainer(USkeleton::AnimTrackCurveMappingName);
+
+		for (auto& Curve : Curves)
+		{
+			// skip if disabled
+			if (Curve.GetCurveTypeFlag(ACF_Disabled))
+			{
+				continue;
+			}
+
+			// add bone modifier
+			FName CurveName;
+			NameMapping->GetName(Curve.CurveUid, CurveName);
+
+			// @TODO: this is going to be issue. If they don't save skeleton with it, we don't have name at all?
+ 			if (CurveName == NAME_None)
+ 			{
+				FSmartNameMapping::UID NewUID;
+ 				NameMapping->AddOrFindName(Curve.LastObservedName, NewUID);
+				Curve.CurveUid = NewUID;
+
+				CurveName = Curve.LastObservedName;
+ 			}
+
+			FName BoneName = CurveName;
+			if (BoneName != NAME_None && GetSkeleton()->GetReferenceSkeleton().FindBoneIndex(BoneName) != INDEX_NONE)
+			{
+				ModifyBone(BoneName, true);
+			}
+		}
+	}
+}
+
+void FAnimPreviewInstanceProxy::UpdateCurveController()
+{
+	// evaluate the curve data first
+	UAnimSequenceBase* CurrentSequence = Cast<UAnimSequenceBase>(GetCurrentAsset());
+
+	if (CurrentSequence && GetSkeleton())
+	{
+		TMap<FName, FTransform> ActiveCurves;
+		CurrentSequence->RawCurveData.EvaluateTransformCurveData(GetSkeleton(), ActiveCurves, GetCurrentTime(), 1.f);
+
+		// make sure those curves exists in the bone controller, otherwise problem
+		if ( ActiveCurves.Num() > 0 )
+		{
+			for(auto& SingleBoneController : CurveBoneControllers)
+			{
+				// make sure the curve exists
+				FName CurveName = SingleBoneController.BoneToModify.BoneName;
+
+				// we should add extra key to front and back whenever animation length changes or so. 
+				// animation length change requires to bake down animation first 
+				// this will make sure all the keys that were embedded at the start/end will automatically be backed to the data
+				const FTransform* Value = ActiveCurves.Find(CurveName);
+				if (Value)
+				{
+					// apply this change
+					SingleBoneController.Translation = Value->GetTranslation();
+					SingleBoneController.Scale = Value->GetScale3D();
+					// sasd we're converting twice
+					SingleBoneController.Rotation = Value->GetRotation().Rotator();
+				}
+			}
+		}
+		else
+		{
+			// should match
+			ensure (CurveBoneControllers.Num() == 0);
+			CurveBoneControllers.Empty();
+		}
+	}
+}
+
+void FAnimPreviewInstanceProxy::ApplyBoneControllers(USkeletalMeshComponent* Component, TArray<FAnimNode_ModifyBone> &InBoneControllers, FCSPose<FCompactPose>& OutMeshPose)
 {
 	for(auto& SingleBoneController : InBoneControllers)
 	{
-		SingleBoneController.BoneToModify.BoneIndex = RequiredBones.GetPoseBoneIndexForBoneName(SingleBoneController.BoneToModify.BoneName);
+		SingleBoneController.BoneToModify.BoneIndex = GetRequiredBones().GetPoseBoneIndexForBoneName(SingleBoneController.BoneToModify.BoneName);
 		if(SingleBoneController.BoneToModify.BoneIndex != INDEX_NONE)
 		{
 			TArray<FBoneTransform> BoneTransforms;
@@ -215,18 +298,18 @@ void UAnimPreviewInstance::ApplyBoneControllers(USkeletalMeshComponent* Componen
 	}
 }
 
-void UAnimPreviewInstance::SetKeyImplementation(const FCompactPose& PreControllerInLocalSpace, const FCompactPose& PostControllerInLocalSpace)
+void FAnimPreviewInstanceProxy::SetKeyImplementation(const FCompactPose& PreControllerInLocalSpace, const FCompactPose& PostControllerInLocalSpace)
 {
 #if WITH_EDITOR
 	// evaluate the curve data first
-	UAnimSequence* CurrentSequence = Cast<UAnimSequence>(CurrentAsset);
+	UAnimSequence* CurrentSequence = Cast<UAnimSequence>(GetCurrentAsset());
 	UDebugSkelMeshComponent* Component = Cast<UDebugSkelMeshComponent> (GetSkelMeshComponent());
 
-	if(CurrentSequence && CurrentSkeleton && Component && Component->SkeletalMesh)
+	if(CurrentSequence && GetSkeleton() && Component && Component->SkeletalMesh)
 	{
 		FScopedTransaction ScopedTransaction(LOCTEXT("SetKey", "Set Key"));
 		CurrentSequence->Modify(true);
-		Modify();
+		GetAnimInstanceObject()->Modify();
 
 		TArray<FName> BonesToModify;
 		// need to get component transform first. Depending on when this gets called, the transform is not up-to-date. 
@@ -241,12 +324,12 @@ void UAnimPreviewInstance::SetKeyImplementation(const FCompactPose& PreControlle
 			FName BoneName = SingleBoneController.BoneToModify.BoneName;
 			// now convert data
 			const FMeshPoseBoneIndex MeshBoneIndex(Component->GetBoneIndex(BoneName));
-			const FCompactPoseBoneIndex BoneIndex = RequiredBones.MakeCompactPoseIndex(MeshBoneIndex);
+			const FCompactPoseBoneIndex BoneIndex = GetRequiredBones().MakeCompactPoseIndex(MeshBoneIndex);
 			FTransform  LocalTransform = PostControllerInLocalSpace[BoneIndex];
 
 			// now we have LocalTransform and get additive data
 			FTransform AdditiveTransform = LocalTransform.GetRelativeTransform(PreControllerInLocalSpace[BoneIndex]);
-			AddKeyToSequence(CurrentSequence, CurrentTime, BoneName, AdditiveTransform);
+			AddKeyToSequence(CurrentSequence, GetCurrentTime(), BoneName, AdditiveTransform);
 
 			BonesToModify.Add(BoneName);
 		}
@@ -263,7 +346,7 @@ void UAnimPreviewInstance::SetKeyImplementation(const FCompactPose& PreControlle
 				// if it's not on BonesToModify, add identity here. 
 				if (!BonesToModify.Contains(BoneName))
 				{
-					AddKeyToSequence(CurrentSequence, CurrentTime, BoneName, FTransform::Identity);
+					AddKeyToSequence(CurrentSequence, GetCurrentTime(), BoneName, FTransform::Identity);
 				}
 			}
 		}
@@ -275,7 +358,7 @@ void UAnimPreviewInstance::SetKeyImplementation(const FCompactPose& PreControlle
 #endif
 }
 
-void UAnimPreviewInstance::AddKeyToSequence(UAnimSequence* Sequence, float Time, const FName& BoneName, const FTransform& AdditiveTransform) 
+void FAnimPreviewInstanceProxy::AddKeyToSequence(UAnimSequence* Sequence, float Time, const FName& BoneName, const FTransform& AdditiveTransform) 
 {
 	Sequence->AddKeyToSequence(Time, BoneName, AdditiveTransform);
 
@@ -284,165 +367,119 @@ void UAnimPreviewInstance::AddKeyToSequence(UAnimSequence* Sequence, float Time,
 	// make sure you add it there
 	ModifyBone(BoneName, true);
 
-	RequiredBones.SetUseSourceData(true);
+	GetRequiredBones().SetUseSourceData(true);
 }
 
-void UAnimPreviewInstance::SetKey(FSimpleDelegate InOnSetKeyCompleteDelegate)
+UAnimPreviewInstance::UAnimPreviewInstance(const FObjectInitializer& ObjectInitializer)
+	: Super(ObjectInitializer)
 {
-#if WITH_EDITOR
-	bSetKey = true;
-	OnSetKeyCompleteDelegate = InOnSetKeyCompleteDelegate;
-#endif
+	RootMotionMode = ERootMotionMode::RootMotionFromEverything;
+	bCanUseParallelUpdateAnimation = false;
 }
 
-void UAnimPreviewInstance::RefreshCurveBoneControllers()
+void UAnimPreviewInstance::NativeInitializeAnimation()
 {
-	// go through all curves and see if it has Transform Curve
-	// if so, find what bone that belong to and create BoneMOdifier for them
-	UAnimSequence* CurrentSequence = Cast<UAnimSequence>(CurrentAsset);
-
-	CurveBoneControllers.Empty();
-
-	// do not apply if BakedAnimation is on
-	if(CurrentSequence)
-	{
-		// make sure if this needs source update
-		if ( !CurrentSequence->DoesContainTransformCurves() )
-		{
-			return;
-		}
+	FAnimPreviewInstanceProxy& Proxy = GetProxyOnGameThread<FAnimPreviewInstanceProxy>();
 
-		RequiredBones.SetUseSourceData(true);
-
-		TArray<FTransformCurve>& Curves = CurrentSequence->RawCurveData.TransformCurves;
-		FSmartNameMapping* NameMapping = CurrentSkeleton->SmartNames.GetContainer(USkeleton::AnimTrackCurveMappingName);
+	// Cache our play state from the previous animation otherwise set to play
+	bool bCachedIsPlaying = (Proxy.GetCurrentAsset() != NULL) ? Proxy.IsPlaying() : true;
 
-		for (auto& Curve : Curves)
-		{
-			// skip if disabled
-			if (Curve.GetCurveTypeFlag(ACF_Disabled))
-			{
-				continue;
-			}
+	Super::NativeInitializeAnimation();
 
-			// add bone modifier
-			FName CurveName;
-			NameMapping->GetName(Curve.CurveUid, CurveName);
+	Proxy.SetPlaying(bCachedIsPlaying);
 
-			// @TODO: this is going to be issue. If they don't save skeleton with it, we don't have name at all?
- 			if (CurveName == NAME_None)
- 			{
-				FSmartNameMapping::UID NewUID;
- 				NameMapping->AddOrFindName(Curve.LastObservedName, NewUID);
-				Curve.CurveUid = NewUID;
+	Proxy.RefreshCurveBoneControllers();
+}
 
-				CurveName = Curve.LastObservedName;
- 			}
+FAnimNode_ModifyBone* UAnimPreviewInstance::FindModifiedBone(const FName& InBoneName, bool bCurveController/*=false*/)
+{
+	return GetProxyOnGameThread<FAnimPreviewInstanceProxy>().FindModifiedBone(InBoneName, bCurveController);
+}
 
-			FName BoneName = CurveName;
-			if (BoneName != NAME_None && CurrentSkeleton->GetReferenceSkeleton().FindBoneIndex(BoneName) != INDEX_NONE)
-			{
-				ModifyBone(BoneName, true);
-			}
-		}
-	}
+FAnimNode_ModifyBone& UAnimPreviewInstance::ModifyBone(const FName& InBoneName, bool bCurveController/*=false*/)
+{
+	return GetProxyOnGameThread<FAnimPreviewInstanceProxy>().ModifyBone(InBoneName, bCurveController);
 }
 
-void UAnimPreviewInstance::UpdateCurveController()
+void UAnimPreviewInstance::RemoveBoneModification(const FName& InBoneName, bool bCurveController/*=false*/)
 {
-	// evaluate the curve data first
-	UAnimSequenceBase* CurrentSequence = Cast<UAnimSequenceBase>(CurrentAsset);
+	GetProxyOnGameThread<FAnimPreviewInstanceProxy>().RemoveBoneModification(InBoneName, bCurveController);
+}
 
-	if (CurrentSequence && CurrentSkeleton)
-	{
-		TMap<FName, FTransform> ActiveCurves;
-		CurrentSequence->RawCurveData.EvaluateTransformCurveData(CurrentSkeleton, ActiveCurves, CurrentTime, 1.f);
+void UAnimPreviewInstance::ResetModifiedBone(bool bCurveController/*=false*/)
+{
+	GetProxyOnGameThread<FAnimPreviewInstanceProxy>().ResetModifiedBone(bCurveController);
+}
 
-		// make sure those curves exists in the bone controller, otherwise problem
-		if ( ActiveCurves.Num() > 0 )
-		{
-			for(auto& SingleBoneController : CurveBoneControllers)
-			{
-				// make sure the curve exists
-				FName CurveName = SingleBoneController.BoneToModify.BoneName;
+void UAnimPreviewInstance::SetKey(FSimpleDelegate InOnSetKeyCompleteDelegate)
+{
+	GetProxyOnGameThread<FAnimPreviewInstanceProxy>().SetKey(InOnSetKeyCompleteDelegate);
+}
 
-				// we should add extra key to front and back whenever animation length changes or so. 
-				// animation length change requires to bake down animation first 
-				// this will make sure all the keys that were embedded at the start/end will automatically be backed to the data
-				const FTransform* Value = ActiveCurves.Find(CurveName);
-				if (Value)
-				{
-					// apply this change
-					SingleBoneController.Translation = Value->GetTranslation();
-					SingleBoneController.Scale = Value->GetScale3D();
-					// sasd we're converting twice
-					SingleBoneController.Rotation = Value->GetRotation().Rotator();
-				}
-			}
-		}
-		else
-		{
-			// should match
-			ensure (CurveBoneControllers.Num() == 0);
-			CurveBoneControllers.Empty();
-		}
-	}
+void UAnimPreviewInstance::RefreshCurveBoneControllers()
+{
+	GetProxyOnGameThread<FAnimPreviewInstanceProxy>().RefreshCurveBoneControllers();
 }
 
 /** Set SkeletalControl Alpha**/
 void UAnimPreviewInstance::SetSkeletalControlAlpha(float InSkeletalControlAlpha)
 {
-	SkeletalControlAlpha = FMath::Clamp<float>(InSkeletalControlAlpha, 0.f, 1.f);
+	GetProxyOnGameThread<FAnimPreviewInstanceProxy>().SetSkeletalControlAlpha(InSkeletalControlAlpha);
 }
 
 UAnimSequence* UAnimPreviewInstance::GetAnimSequence()
 {
-	return Cast<UAnimSequence>(CurrentAsset);
+	return Cast<UAnimSequence>(GetCurrentAsset());
 }
 
 void UAnimPreviewInstance::RestartMontage(UAnimMontage* Montage, FName FromSection)
 {
-	if (Montage == CurrentAsset)
+	FAnimPreviewInstanceProxy& Proxy = GetProxyOnGameThread<FAnimPreviewInstanceProxy>();
+
+	if (Montage == Proxy.GetCurrentAsset())
 	{
 		MontagePreviewType = EMPT_Normal;
 		// since this is preview, we would like not to blend in
 		// just hard stop here
 		Montage_Stop(0.0f, Montage);
-		Montage_Play(Montage, PlayRate);
+		Montage_Play(Montage, Proxy.GetPlayRate());
 		if (FromSection != NAME_None)
 		{
 			Montage_JumpToSection(FromSection);
 		}
-		MontagePreview_SetLoopNormal(bLooping, Montage->GetSectionIndex(FromSection));
+		MontagePreview_SetLoopNormal(Proxy.IsLooping(), Montage->GetSectionIndex(FromSection));
 	}
 }
 
 void UAnimPreviewInstance::SetAnimationAsset(UAnimationAsset* NewAsset, bool bIsLooping, float InPlayRate)
 {
+	FAnimPreviewInstanceProxy& Proxy = GetProxyOnGameThread<FAnimPreviewInstanceProxy>();
+
 	// make sure to turn that off before setting new asset
-	RequiredBones.SetUseSourceData(false);
+	Proxy.GetRequiredBones().SetUseSourceData(false);
 
 	Super::SetAnimationAsset(NewAsset, bIsLooping, InPlayRate);
-	RootMotionMode = Cast<UAnimMontage>(CurrentAsset) != NULL ? ERootMotionMode::RootMotionFromMontagesOnly : ERootMotionMode::RootMotionFromEverything;
+	RootMotionMode = Cast<UAnimMontage>(Proxy.GetCurrentAsset()) != NULL ? ERootMotionMode::RootMotionFromMontagesOnly : ERootMotionMode::RootMotionFromEverything;
 
 	// should re sync up curve bone controllers from new asset
-	RefreshCurveBoneControllers();
+	Proxy.RefreshCurveBoneControllers();
 }
 
 void UAnimPreviewInstance::MontagePreview_SetLooping(bool bIsLooping)
 {
-	bLooping = bIsLooping;
+	FAnimPreviewInstanceProxy& Proxy = GetProxyOnGameThread<FAnimPreviewInstanceProxy>();
+	Proxy.SetLooping(bIsLooping);
 
-	if (UAnimMontage* Montage = Cast<UAnimMontage>(CurrentAsset))
+	if (UAnimMontage* Montage = Cast<UAnimMontage>(Proxy.GetCurrentAsset()))
 	{
 		switch (MontagePreviewType)
 		{
 		case EMPT_AllSections:
-			MontagePreview_SetLoopAllSections(bLooping);
+			MontagePreview_SetLoopAllSections(Proxy.IsLooping());
 			break;
 		case EMPT_Normal:
 		default:
-			MontagePreview_SetLoopNormal(bLooping);
+			MontagePreview_SetLoopNormal(Proxy.IsLooping());
 			break;
 		}
 	}
@@ -450,15 +487,16 @@ void UAnimPreviewInstance::MontagePreview_SetLooping(bool bIsLooping)
 
 void UAnimPreviewInstance::MontagePreview_SetPlaying(bool bIsPlaying)
 {
-	bPlaying = bIsPlaying;
+	FAnimPreviewInstanceProxy& Proxy = GetProxyOnGameThread<FAnimPreviewInstanceProxy>();
+	Proxy.SetPlaying(bIsPlaying);
 
 	if (FAnimMontageInstance* CurMontageInstance = GetActiveMontageInstance())
 	{
-		CurMontageInstance->bPlaying = bIsPlaying;
+		CurMontageInstance->bPlaying = Proxy.IsPlaying();
 	}
-	else if (bPlaying)
+	else if (Proxy.IsPlaying())
 	{
-		UAnimMontage* Montage = Cast<UAnimMontage>(CurrentAsset);
+		UAnimMontage* Montage = Cast<UAnimMontage>(Proxy.GetCurrentAsset());
 		if (Montage)
 		{
 			switch (MontagePreviewType)
@@ -477,18 +515,20 @@ void UAnimPreviewInstance::MontagePreview_SetPlaying(bool bIsPlaying)
 
 void UAnimPreviewInstance::MontagePreview_SetReverse(bool bInReverse)
 {
+	FAnimPreviewInstanceProxy& Proxy = GetProxyOnGameThread<FAnimPreviewInstanceProxy>();
 	Super::SetReverse(bInReverse);
 
 	if (FAnimMontageInstance* CurMontageInstance = GetActiveMontageInstance())
 	{
 		// copy the current playrate
-		CurMontageInstance->SetPlayRate(PlayRate);
+		CurMontageInstance->SetPlayRate(Proxy.GetPlayRate());
 	}
 }
 
 void UAnimPreviewInstance::MontagePreview_Restart()
 {
-	if (UAnimMontage* Montage = Cast<UAnimMontage>(CurrentAsset))
+	FAnimPreviewInstanceProxy& Proxy = GetProxyOnGameThread<FAnimPreviewInstanceProxy>();
+	if (UAnimMontage* Montage = Cast<UAnimMontage>(Proxy.GetCurrentAsset()))
 	{
 		switch (MontagePreviewType)
 		{
@@ -505,15 +545,17 @@ void UAnimPreviewInstance::MontagePreview_Restart()
 
 void UAnimPreviewInstance::MontagePreview_StepForward()
 {
-	if (UAnimMontage* Montage = Cast<UAnimMontage>(CurrentAsset))
+	FAnimPreviewInstanceProxy& Proxy = GetProxyOnGameThread<FAnimPreviewInstanceProxy>();
+
+	if (UAnimMontage* Montage = Cast<UAnimMontage>(Proxy.GetCurrentAsset()))
 	{
-		bool bWasPlaying = IsPlayingMontage() && (bLooping || bPlaying); // we need to handle non-looped case separately, even if paused during playthrough
+		bool bWasPlaying = IsPlayingMontage() && (Proxy.IsLooping() || Proxy.IsPlaying()); // we need to handle non-looped case separately, even if paused during playthrough
 		MontagePreview_SetReverse(false);
 		if (! bWasPlaying)
 		{
-			if (! bLooping)
+			if (! Proxy.IsLooping())
 			{
-				float StoppedAt = CurrentTime;
+				float StoppedAt = Proxy.GetCurrentTime();
 				if (! bWasPlaying)
 				{
 					// play montage but at last known location
@@ -521,7 +563,7 @@ void UAnimPreviewInstance::MontagePreview_StepForward()
 					SetPosition(StoppedAt, false);
 				}
 				int32 LastPreviewSectionIdx = MontagePreview_FindLastSection(MontagePreviewStartSectionIdx);
-				if (FMath::Abs(CurrentTime - (Montage->CompositeSections[LastPreviewSectionIdx].GetTime() + Montage->GetSectionLength(LastPreviewSectionIdx))) <= MontagePreview_CalculateStepLength())
+				if (FMath::Abs(Proxy.GetCurrentTime() - (Montage->CompositeSections[LastPreviewSectionIdx].GetTime() + Montage->GetSectionLength(LastPreviewSectionIdx))) <= MontagePreview_CalculateStepLength())
 				{
 					// we're at the end, jump right to the end
 					Montage_JumpToSectionsEnd(Montage->GetSectionName(LastPreviewSectionIdx));
@@ -542,12 +584,12 @@ void UAnimPreviewInstance::MontagePreview_StepForward()
 		// Advance a single frame, leaving it paused afterwards
 		int32 NumFrames = Montage->GetNumberOfFrames();
 		// Add DELTA to prefer next frame when we're close to the boundary
-		float CurrentFraction = CurrentTime / Montage->SequenceLength + DELTA;
+		float CurrentFraction = Proxy.GetCurrentTime() / Montage->SequenceLength + DELTA;
 		float NextFrame = FMath::Clamp<float>(FMath::FloorToFloat(CurrentFraction * NumFrames) + 1.0f, 0, NumFrames);
 		float NewTime = Montage->SequenceLength * (NextFrame / NumFrames);
 
 		GetSkelMeshComponent()->GlobalAnimRateScale = 1.0f;
-		GetSkelMeshComponent()->TickAnimation(NewTime - CurrentTime, false);
+		GetSkelMeshComponent()->TickAnimation(NewTime - Proxy.GetCurrentTime(), false);
 
 		MontagePreview_SetPlaying(false);
 	}
@@ -555,15 +597,17 @@ void UAnimPreviewInstance::MontagePreview_StepForward()
 
 void UAnimPreviewInstance::MontagePreview_StepBackward()
 {
-	if (UAnimMontage* Montage = Cast<UAnimMontage>(CurrentAsset))
+	FAnimPreviewInstanceProxy& Proxy = GetProxyOnGameThread<FAnimPreviewInstanceProxy>();
+
+	if (UAnimMontage* Montage = Cast<UAnimMontage>(Proxy.GetCurrentAsset()))
 	{
-		bool bWasPlaying = IsPlayingMontage() && (bLooping || bPlaying); // we need to handle non-looped case separately, even if paused during playthrough
+		bool bWasPlaying = IsPlayingMontage() && (Proxy.IsLooping() || Proxy.IsPlaying()); // we need to handle non-looped case separately, even if paused during playthrough
 		MontagePreview_SetReverse(true);
 		if (! bWasPlaying)
 		{
-			if (! bLooping)
+			if (! Proxy.IsLooping())
 			{
-				float StoppedAt = CurrentTime;
+				float StoppedAt = Proxy.GetCurrentTime();
 				if (! bWasPlaying)
 				{
 					// play montage but at last known location
@@ -571,12 +615,12 @@ void UAnimPreviewInstance::MontagePreview_StepBackward()
 					SetPosition(StoppedAt, false);
 				}
 				int32 LastPreviewSectionIdx = MontagePreview_FindLastSection(MontagePreviewStartSectionIdx);
-				if (FMath::Abs(CurrentTime - (Montage->CompositeSections[LastPreviewSectionIdx].GetTime() + Montage->GetSectionLength(LastPreviewSectionIdx))) <= MontagePreview_CalculateStepLength())
+				if (FMath::Abs(Proxy.GetCurrentTime() - (Montage->CompositeSections[LastPreviewSectionIdx].GetTime() + Montage->GetSectionLength(LastPreviewSectionIdx))) <= MontagePreview_CalculateStepLength())
 				{
 					// special case as we could stop at the end of our last section which is also beginning of following section - we don't want to get stuck there, but be inside of our starting section
 					Montage_JumpToSection(Montage->GetSectionName(LastPreviewSectionIdx));
 				}
-				else if (FMath::Abs(CurrentTime - Montage->CompositeSections[MontagePreviewStartSectionIdx].GetTime()) <= MontagePreview_CalculateStepLength())
+				else if (FMath::Abs(Proxy.GetCurrentTime() - Montage->CompositeSections[MontagePreviewStartSectionIdx].GetTime()) <= MontagePreview_CalculateStepLength())
 				{
 					// we're at the end of playing backward, jump right to the end
 					Montage_JumpToSectionsEnd(Montage->GetSectionName(MontagePreviewStartSectionIdx));
@@ -597,12 +641,12 @@ void UAnimPreviewInstance::MontagePreview_StepBackward()
 		// Advance a single frame, leaving it paused afterwards
 		int32 NumFrames = Montage->GetNumberOfFrames();
 		// Add DELTA to prefer next frame when we're close to the boundary
-		float CurrentFraction = CurrentTime / Montage->SequenceLength + DELTA;
+		float CurrentFraction = Proxy.GetCurrentTime() / Montage->SequenceLength + DELTA;
 		float NextFrame = FMath::Clamp<float>(FMath::FloorToFloat(CurrentFraction * NumFrames) - 1.0f, 0, NumFrames);
 		float NewTime = Montage->SequenceLength * (NextFrame / NumFrames);
 
 		GetSkelMeshComponent()->GlobalAnimRateScale = 1.0f;
-		GetSkelMeshComponent()->TickAnimation(FMath::Abs(NewTime - CurrentTime), false);
+		GetSkelMeshComponent()->TickAnimation(FMath::Abs(NewTime - Proxy.GetCurrentTime()), false);
 
 		MontagePreview_SetPlaying(false);
 	}
@@ -615,7 +659,8 @@ float UAnimPreviewInstance::MontagePreview_CalculateStepLength()
 
 void UAnimPreviewInstance::MontagePreview_JumpToStart()
 {
-	if (UAnimMontage* Montage = Cast<UAnimMontage>(CurrentAsset))
+	FAnimPreviewInstanceProxy& Proxy = GetProxyOnGameThread<FAnimPreviewInstanceProxy>();
+	if (UAnimMontage* Montage = Cast<UAnimMontage>(Proxy.GetCurrentAsset()))
 	{
 		int32 SectionIdx = 0;
 		if (MontagePreviewType == EMPT_Normal)
@@ -628,7 +673,7 @@ void UAnimPreviewInstance::MontagePreview_JumpToStart()
 		{
 			MontagePreview_Restart();
 		}
-		if (PlayRate < 0.f)
+		if (Proxy.GetPlayRate() < 0.f)
 		{
 			Montage_JumpToSectionsEnd(Montage->GetSectionName(SectionIdx));
 		}
@@ -645,7 +690,8 @@ void UAnimPreviewInstance::MontagePreview_JumpToStart()
 
 void UAnimPreviewInstance::MontagePreview_JumpToEnd()
 {
-	if (UAnimMontage* Montage = Cast<UAnimMontage>(CurrentAsset))
+	FAnimPreviewInstanceProxy& Proxy = GetProxyOnGameThread<FAnimPreviewInstanceProxy>();
+	if (UAnimMontage* Montage = Cast<UAnimMontage>(Proxy.GetCurrentAsset()))
 	{
 		int32 SectionIdx = 0;
 		if (MontagePreviewType == EMPT_Normal)
@@ -658,7 +704,7 @@ void UAnimPreviewInstance::MontagePreview_JumpToEnd()
 		{
 			MontagePreview_Restart();
 		}
-		if (PlayRate < 0.f)
+		if (Proxy.GetPlayRate() < 0.f)
 		{
 			Montage_JumpToSection(Montage->GetSectionName(MontagePreview_FindLastSection(SectionIdx)));
 		}
@@ -675,7 +721,9 @@ void UAnimPreviewInstance::MontagePreview_JumpToEnd()
 
 void UAnimPreviewInstance::MontagePreview_JumpToPreviewStart()
 {
-	if (UAnimMontage* Montage = Cast<UAnimMontage>(CurrentAsset))
+	FAnimPreviewInstanceProxy& Proxy = GetProxyOnGameThread<FAnimPreviewInstanceProxy>();
+
+	if (UAnimMontage* Montage = Cast<UAnimMontage>(Proxy.GetCurrentAsset()))
 	{
 		int32 SectionIdx = 0;
 		if (MontagePreviewType == EMPT_Normal)
@@ -688,7 +736,7 @@ void UAnimPreviewInstance::MontagePreview_JumpToPreviewStart()
 		{
 			MontagePreview_Restart();
 		}
-		Montage_JumpToSection(Montage->GetSectionName(PlayRate > 0.f? SectionIdx : MontagePreview_FindLastSection(SectionIdx)));
+		Montage_JumpToSection(Montage->GetSectionName(Proxy.GetPlayRate() > 0.f? SectionIdx : MontagePreview_FindLastSection(SectionIdx)));
 		if (! bWasPlaying)
 		{
 			MontagePreview_SetPlaying(false);
@@ -698,8 +746,10 @@ void UAnimPreviewInstance::MontagePreview_JumpToPreviewStart()
 
 void UAnimPreviewInstance::MontagePreview_JumpToPosition(float NewPosition)
 {
+	FAnimPreviewInstanceProxy& Proxy = GetProxyOnGameThread<FAnimPreviewInstanceProxy>();
+
 	SetPosition(NewPosition, false);
-	if (UAnimMontage* Montage = Cast<UAnimMontage>(CurrentAsset))
+	if (UAnimMontage* Montage = Cast<UAnimMontage>(Proxy.GetCurrentAsset()))
 	{
 		// this section will be first
 		int32 NewMontagePreviewStartSectionIdx = MontagePreview_FindFirstSectionAsInMontage(Montage->GetSectionIndexFromPosition(NewPosition));
@@ -709,7 +759,7 @@ void UAnimPreviewInstance::MontagePreview_JumpToPosition(float NewPosition)
 			MontagePreviewStartSectionIdx = NewMontagePreviewStartSectionIdx;
 		}
 		// setup looping to match normal playback
-		MontagePreview_SetLooping(bLooping);
+		MontagePreview_SetLooping(Proxy.IsLooping());
 	}
 }
 
@@ -723,7 +773,9 @@ void UAnimPreviewInstance::MontagePreview_RemoveBlendOut()
 
 void UAnimPreviewInstance::MontagePreview_PreviewNormal(int32 FromSectionIdx)
 {
-	if (UAnimMontage* Montage = Cast<UAnimMontage>(CurrentAsset))
+	FAnimPreviewInstanceProxy& Proxy = GetProxyOnGameThread<FAnimPreviewInstanceProxy>();
+
+	if (UAnimMontage* Montage = Cast<UAnimMontage>(Proxy.GetCurrentAsset()))
 	{
 		int32 PreviewFromSection = FromSectionIdx;
 		if (FromSectionIdx != INDEX_NONE)
@@ -739,40 +791,44 @@ void UAnimPreviewInstance::MontagePreview_PreviewNormal(int32 FromSectionIdx)
 		// since this is preview, we would like not to blend in
 		// just hard stop here
 		Montage_Stop(0.0f, Montage);
-		Montage_Play(Montage, PlayRate);
-		MontagePreview_SetLoopNormal(bLooping, FromSectionIdx);
+		Montage_Play(Montage, Proxy.GetPlayRate());
+		MontagePreview_SetLoopNormal(Proxy.IsLooping(), FromSectionIdx);
 		Montage_JumpToSection(Montage->GetSectionName(PreviewFromSection));
 		MontagePreview_RemoveBlendOut();
-		bPlaying = true;
+		Proxy.SetPlaying(true);
 	}
 }
 
 void UAnimPreviewInstance::MontagePreview_PreviewAllSections()
 {
-	UAnimMontage* Montage = Cast<UAnimMontage>(CurrentAsset);
+	FAnimPreviewInstanceProxy& Proxy = GetProxyOnGameThread<FAnimPreviewInstanceProxy>();
+
+	UAnimMontage* Montage = Cast<UAnimMontage>(Proxy.GetCurrentAsset());
 	if (Montage && Montage->SequenceLength > 0.f)
 	{
 		MontagePreviewType = EMPT_AllSections;
 		// since this is preview, we would like not to blend in
 		// just hard stop here
 		Montage_Stop(0.0f, Montage);
-		Montage_Play(Montage, PlayRate);
-		MontagePreview_SetLoopAllSections(bLooping);
+		Montage_Play(Montage, Proxy.GetPlayRate());
+		MontagePreview_SetLoopAllSections(Proxy.IsLooping());
 		MontagePreview_JumpToPreviewStart();
 		MontagePreview_RemoveBlendOut();
-		bPlaying = true;
+		Proxy.SetPlaying(true);
 	}
 }
 
 void UAnimPreviewInstance::MontagePreview_SetLoopNormal(bool bIsLooping, int32 PreferSectionIdx)
 {
-	if (UAnimMontage* Montage = Cast<UAnimMontage>(CurrentAsset))
+	FAnimPreviewInstanceProxy& Proxy = GetProxyOnGameThread<FAnimPreviewInstanceProxy>();
+
+	if (UAnimMontage* Montage = Cast<UAnimMontage>(Proxy.GetCurrentAsset()))
 	{
 		MontagePreview_ResetSectionsOrder();
 
 		if (PreferSectionIdx == INDEX_NONE)
 		{
-			PreferSectionIdx = Montage->GetSectionIndexFromPosition(CurrentTime);
+			PreferSectionIdx = Montage->GetSectionIndexFromPosition(Proxy.GetCurrentTime());
 		}
 		int32 TotalSection = Montage->CompositeSections.Num();
 		if (TotalSection > 0)
@@ -818,11 +874,11 @@ void UAnimPreviewInstance::MontagePreview_SetLoopNormal(bool bIsLooping, int32 P
 					++ InChain;
 				}
 				// loop this section
-				SetMontageLoop(Montage, bIsLooping, Montage->CompositeSections[NotUsedIdx].SectionName);
+				SetMontageLoop(Montage, Proxy.IsLooping(), Montage->CompositeSections[NotUsedIdx].SectionName);
 			}
 			if (Montage->CompositeSections.IsValidIndex(PreferSectionIdx))
 			{
-				SetMontageLoop(Montage, bIsLooping, Montage->CompositeSections[PreferSectionIdx].SectionName);
+				SetMontageLoop(Montage, Proxy.IsLooping(), Montage->CompositeSections[PreferSectionIdx].SectionName);
 			}
 		}
 	}
@@ -830,7 +886,9 @@ void UAnimPreviewInstance::MontagePreview_SetLoopNormal(bool bIsLooping, int32 P
 
 void UAnimPreviewInstance::MontagePreview_SetLoopAllSetupSections(bool bIsLooping)
 {
-	if (UAnimMontage* Montage = Cast<UAnimMontage>(CurrentAsset))
+	FAnimPreviewInstanceProxy& Proxy = GetProxyOnGameThread<FAnimPreviewInstanceProxy>();
+
+	if (UAnimMontage* Montage = Cast<UAnimMontage>(Proxy.GetCurrentAsset()))
 	{
 		MontagePreview_ResetSectionsOrder();
 
@@ -874,7 +932,7 @@ void UAnimPreviewInstance::MontagePreview_SetLoopAllSetupSections(bool bIsLoopin
 					}
 				}
 			}
-			if (bIsLooping)
+			if (Proxy.IsLooping())
 			{
 				// and loop all
 				Montage_SetNextSection(PreviousSection, FirstSection);
@@ -885,12 +943,14 @@ void UAnimPreviewInstance::MontagePreview_SetLoopAllSetupSections(bool bIsLoopin
 
 void UAnimPreviewInstance::MontagePreview_SetLoopAllSections(bool bIsLooping)
 {
-	if (UAnimMontage* Montage = Cast<UAnimMontage>(CurrentAsset))
+	FAnimPreviewInstanceProxy& Proxy = GetProxyOnGameThread<FAnimPreviewInstanceProxy>();
+
+	if (UAnimMontage* Montage = Cast<UAnimMontage>(Proxy.GetCurrentAsset()))
 	{
 		int32 TotalSection = Montage->CompositeSections.Num();
 		if (TotalSection > 0)
 		{
-			if (bIsLooping)
+			if (Proxy.IsLooping())
 			{
 				for (int i = 0; i < TotalSection; ++ i)
 				{
@@ -911,7 +971,8 @@ void UAnimPreviewInstance::MontagePreview_SetLoopAllSections(bool bIsLooping)
 
 void UAnimPreviewInstance::MontagePreview_ResetSectionsOrder()
 {
-	if (UAnimMontage* Montage = Cast<UAnimMontage>(CurrentAsset))
+	FAnimPreviewInstanceProxy& Proxy = GetProxyOnGameThread<FAnimPreviewInstanceProxy>();
+	if (UAnimMontage* Montage = Cast<UAnimMontage>(Proxy.GetCurrentAsset()))
 	{
 		int32 TotalSection = Montage->CompositeSections.Num();
 		// restore to default
@@ -924,9 +985,10 @@ void UAnimPreviewInstance::MontagePreview_ResetSectionsOrder()
 
 int32 UAnimPreviewInstance::MontagePreview_FindFirstSectionAsInMontage(int32 ForSectionIdx)
 {
+	FAnimPreviewInstanceProxy& Proxy = GetProxyOnGameThread<FAnimPreviewInstanceProxy>();
 	int32 ResultIdx = ForSectionIdx;
 	// Montage does not have looping set up, so it should be valid and it gets
-	if (UAnimMontage* Montage = Cast<UAnimMontage>(CurrentAsset))
+	if (UAnimMontage* Montage = Cast<UAnimMontage>(Proxy.GetCurrentAsset()))
 	{
 		TArray<bool> AlreadyVisited;
 		AlreadyVisited.AddZeroed(Montage->CompositeSections.Num());
@@ -967,8 +1029,9 @@ int32 UAnimPreviewInstance::MontagePreview_FindFirstSectionAsInMontage(int32 For
 
 int32 UAnimPreviewInstance::MontagePreview_FindLastSection(int32 StartSectionIdx)
 {
+	FAnimPreviewInstanceProxy& Proxy = GetProxyOnGameThread<FAnimPreviewInstanceProxy>();
 	int32 ResultIdx = StartSectionIdx;
-	if (UAnimMontage* Montage = Cast<UAnimMontage>(CurrentAsset))
+	if (UAnimMontage* Montage = Cast<UAnimMontage>(Proxy.GetCurrentAsset()))
 	{
 		if (FAnimMontageInstance* CurMontageInstance = GetActiveMontageInstance())
 		{
@@ -992,7 +1055,8 @@ int32 UAnimPreviewInstance::MontagePreview_FindLastSection(int32 StartSectionIdx
 
 void UAnimPreviewInstance::BakeAnimation()
 {
-	if (UAnimSequence* Sequence = Cast<UAnimSequence>(CurrentAsset))
+	FAnimPreviewInstanceProxy& Proxy = GetProxyOnGameThread<FAnimPreviewInstanceProxy>();
+	if (UAnimSequence* Sequence = Cast<UAnimSequence>(Proxy.GetCurrentAsset()))
 	{
 		FScopedTransaction ScopedTransaction(LOCTEXT("BakeAnimation", "Bake Animation"));
 		Sequence->Modify(true);
@@ -1002,6 +1066,22 @@ void UAnimPreviewInstance::BakeAnimation()
 
 void UAnimPreviewInstance::EnableControllers(bool bEnable)
 {
-	bEnableControllers = bEnable;
+	GetProxyOnGameThread<FAnimPreviewInstanceProxy>().EnableControllers(bEnable);
 }
+
+void UAnimPreviewInstance::SetForceRetargetBasePose(bool bInForceRetargetBasePose)
+{
+	GetProxyOnGameThread<FAnimPreviewInstanceProxy>().SetForceRetargetBasePose(bInForceRetargetBasePose);
+}
+
+bool UAnimPreviewInstance::GetForceRetargetBasePose() const
+{
+	return GetProxyOnGameThread<FAnimPreviewInstanceProxy>().GetForceRetargetBasePose();
+}
+
+FAnimInstanceProxy* UAnimPreviewInstance::CreateAnimInstanceProxy()
+{
+	return new FAnimPreviewInstanceProxy(this);
+}
+
 #undef LOCTEXT_NAMESPACE
\ No newline at end of file
diff --git a/Engine/Source/Editor/AnimGraph/Private/AnimationStateMachineSchema.cpp b/Engine/Source/Editor/AnimGraph/Private/AnimationStateMachineSchema.cpp
--- a/Engine/Source/Editor/AnimGraph/Private/AnimationStateMachineSchema.cpp
+++ b/Engine/Source/Editor/AnimGraph/Private/AnimationStateMachineSchema.cpp
@@ -66,6 +66,8 @@ UEdGraphNode* FEdGraphSchemaAction_NewStateNode::PerformAction(class UEdGraph* P
 
 		ResultNode = NodeTemplate;
 
+		ResultNode->SetFlags(RF_Transactional);
+
 		UBlueprint* Blueprint = FBlueprintEditorUtils::FindBlueprintForGraphChecked(ParentGraph);
 		FBlueprintEditorUtils::MarkBlueprintAsStructurallyModified(Blueprint);
 	}
diff --git a/Engine/Source/Editor/BehaviorTreeEditor/BehaviorTreeEditor.Build.cs b/Engine/Source/Editor/BehaviorTreeEditor/BehaviorTreeEditor.Build.cs
--- a/Engine/Source/Editor/BehaviorTreeEditor/BehaviorTreeEditor.Build.cs
+++ b/Engine/Source/Editor/BehaviorTreeEditor/BehaviorTreeEditor.Build.cs
@@ -43,7 +43,6 @@ public class BehaviorTreeEditor : ModuleRules
 				"BlueprintGraph",
                 "AIGraph",
                 "AIModule",
-                "GameplayDebugger",
 				"ClassViewer"
 			}
 		);
diff --git a/Engine/Source/Editor/BehaviorTreeEditor/Classes/BehaviorTreeEditorTypes.h b/Engine/Source/Editor/BehaviorTreeEditor/Classes/BehaviorTreeEditorTypes.h
--- a/Engine/Source/Editor/BehaviorTreeEditor/Classes/BehaviorTreeEditorTypes.h
+++ b/Engine/Source/Editor/BehaviorTreeEditor/Classes/BehaviorTreeEditorTypes.h
@@ -19,7 +19,15 @@ struct FCompareNodeXLocation
 {
 	FORCEINLINE bool operator()(const UEdGraphPin& A, const UEdGraphPin& B) const
 	{
-		return A.GetOwningNode()->NodePosX < B.GetOwningNode()->NodePosX;
+		const UEdGraphNode* NodeA = A.GetOwningNode();
+		const UEdGraphNode* NodeB = B.GetOwningNode();
+
+		if (NodeA->NodePosX == NodeB->NodePosX)
+		{
+			return NodeA->NodePosY < NodeB->NodePosY;
+		}
+
+		return NodeA->NodePosX < NodeB->NodePosX;
 	}
 };
 
diff --git a/Engine/Source/Editor/BehaviorTreeEditor/Classes/BehaviorTreeGraph.h b/Engine/Source/Editor/BehaviorTreeEditor/Classes/BehaviorTreeGraph.h
--- a/Engine/Source/Editor/BehaviorTreeEditor/Classes/BehaviorTreeGraph.h
+++ b/Engine/Source/Editor/BehaviorTreeEditor/Classes/BehaviorTreeGraph.h
@@ -44,6 +44,7 @@ class UBehaviorTreeGraph : public UAIGraph
 	class UEdGraphNode* FindInjectedNode(int32 Index);
 	void ReplaceNodeConnections(UEdGraphNode* OldNode, UEdGraphNode* NewNode);
 	void RebuildExecutionOrder();
+	void RebuildChildOrder(UEdGraphNode* ParentNode);
 	void SpawnMissingNodesForParallel();
 	void RemoveUnknownSubNodes();
 
diff --git a/Engine/Source/Editor/BehaviorTreeEditor/Private/BehaviorTreeDebugger.cpp b/Engine/Source/Editor/BehaviorTreeEditor/Private/BehaviorTreeDebugger.cpp
--- a/Engine/Source/Editor/BehaviorTreeEditor/Private/BehaviorTreeDebugger.cpp
+++ b/Engine/Source/Editor/BehaviorTreeEditor/Private/BehaviorTreeDebugger.cpp
@@ -6,8 +6,6 @@
 #include "BehaviorTree/BTTaskNode.h"
 #include "BehaviorTree/BTAuxiliaryNode.h"
 #include "BehaviorTreeDelegates.h"
-#include "GameplayDebuggingReplicator.h"
-#include "GameplayDebuggingComponent.h"
 #include "Engine/Selection.h"
 #include "EngineUtils.h"
 
@@ -684,23 +682,23 @@ void FBehaviorTreeDebugger::FindLockedDebugActor(UWorld* World)
 	APlayerController* LocalPC = GEngine->GetFirstLocalPlayerController(World);
 	if (LocalPC && LocalPC->GetHUD() && LocalPC->GetPawnOrSpectator())
 	{
-		AGameplayDebuggingReplicator* DebuggingReplicator = NULL;
-		for (TActorIterator<AGameplayDebuggingReplicator> It(World); It; ++It)
+		APawn* SelectedPawn = NULL;
+#if WITH_ENGINE
+		const UEditorEngine* EEngine = Cast<UEditorEngine>(GEngine);
+		for (FSelectionIterator It = EEngine->GetSelectedActorIterator(); It; ++It)
 		{
-			AGameplayDebuggingReplicator* A = *It;
-			if (!A->IsPendingKill())
+			SelectedPawn = Cast<APawn>(*It);
+			if (SelectedPawn)
 			{
-				DebuggingReplicator = A;
 				break;
 			}
 		}
+#endif //WITH_ENGINE
 
-		const APawn* LockedPawn = DebuggingReplicator != NULL ? Cast<APawn>(DebuggingReplicator->GetSelectedActorToDebug()) : NULL;
-		UBehaviorTreeComponent* TestInstance = FindInstanceInActor((APawn*)LockedPawn);
+		UBehaviorTreeComponent* TestInstance = FindInstanceInActor((APawn*)SelectedPawn);
 		if (TestInstance)
 		{
 			TreeInstance = TestInstance;
-
 #if USE_BEHAVIORTREE_DEBUGGER
 			ActiveStepIndex = TestInstance->DebuggerSteps.Num() - 1;
 #endif
@@ -1040,7 +1038,14 @@ void FBehaviorTreeDebugger::UpdateDebuggerViewOnInstanceChange()
 
 	OnDebuggedBlackboardChangedEvent.Broadcast(BBAsset);
 
-	Refresh();
+	if (DebuggerInstanceIndex != INDEX_NONE)
+	{
+		Refresh();
+	}
+	else
+	{
+		ClearDebuggerState();
+	}
 #endif
 }
 
@@ -1136,9 +1141,9 @@ void FBehaviorTreeDebugger::OnInstanceSelectedInDropdown(UBehaviorTreeComponent*
 		AController* OldController = TreeInstance.IsValid() ? Cast<AController>(TreeInstance->GetOwner()) : NULL;
 		APawn* OldPawn = OldController != NULL ? OldController->GetPawn() : NULL;
 		USelection* SelectedActors = GEditor ? GEditor->GetSelectedActors() : NULL;
-		if (SelectedActors && OldPawn)
+		if (SelectedActors)
 		{
-			SelectedActors->Deselect(OldPawn);
+			SelectedActors->DeselectAll();
 		}
 
 		TreeInstance = SelectedInstance;
@@ -1149,9 +1154,7 @@ void FBehaviorTreeDebugger::OnInstanceSelectedInDropdown(UBehaviorTreeComponent*
 			APawn* Pawn = TestController != NULL ? TestController->GetPawn() : NULL;
 			if (Pawn)
 			{
-				SelectedActors = GEditor->GetSelectedActors();
 				SelectedActors->Select(Pawn);
-				AGameplayDebuggingReplicator::OnSelectionChangedDelegate.Broadcast(Pawn);
 			}
 		}
 
diff --git a/Engine/Source/Editor/BehaviorTreeEditor/Private/BehaviorTreeEditorModule.cpp b/Engine/Source/Editor/BehaviorTreeEditor/Private/BehaviorTreeEditorModule.cpp
--- a/Engine/Source/Editor/BehaviorTreeEditor/Private/BehaviorTreeEditorModule.cpp
+++ b/Engine/Source/Editor/BehaviorTreeEditor/Private/BehaviorTreeEditorModule.cpp
@@ -1,6 +1,7 @@
 // Copyright 1998-2015 Epic Games, Inc. All Rights Reserved.
 
 #include "BehaviorTreeEditorPrivatePCH.h"
+#include "BehaviorTreeEditorModule.h"
 #include "DetailCustomizations/BehaviorDecoratorDetails.h"
 #include "DetailCustomizations/BlackboardDecoratorDetails.h"
 #include "DetailCustomizations/BlackboardSelectorDetails.h"
diff --git a/Engine/Source/Editor/BehaviorTreeEditor/Private/BehaviorTreeGraph.cpp b/Engine/Source/Editor/BehaviorTreeEditor/Private/BehaviorTreeGraph.cpp
--- a/Engine/Source/Editor/BehaviorTreeEditor/Private/BehaviorTreeGraph.cpp
+++ b/Engine/Source/Editor/BehaviorTreeEditor/Private/BehaviorTreeGraph.cpp
@@ -1064,6 +1064,21 @@ void UBehaviorTreeGraph::RebuildExecutionOrder()
 	}
 }
 
+void UBehaviorTreeGraph::RebuildChildOrder(UEdGraphNode* ParentNode)
+{
+	if (ParentNode)
+	{
+		for (int32 PinIdx = 0; PinIdx < ParentNode->Pins.Num(); PinIdx++)
+		{
+			UEdGraphPin* Pin = ParentNode->Pins[PinIdx];
+			if (Pin->Direction == EGPD_Output)
+			{
+				Pin->LinkedTo.Sort(FCompareNodeXLocation());
+			}
+		}
+	}
+}
+
 namespace BTAutoArrangeHelpers
 {
 	struct FNodeBoundsInfo
diff --git a/Engine/Source/Editor/BehaviorTreeEditor/Private/SGraphNode_BehaviorTree.cpp b/Engine/Source/Editor/BehaviorTreeEditor/Private/SGraphNode_BehaviorTree.cpp
--- a/Engine/Source/Editor/BehaviorTreeEditor/Private/SGraphNode_BehaviorTree.cpp
+++ b/Engine/Source/Editor/BehaviorTreeEditor/Private/SGraphNode_BehaviorTree.cpp
@@ -1063,4 +1063,32 @@ TSharedRef<SGraphNode> SGraphNode_BehaviorTree::GetNodeUnderMouse(const FGeometr
 	return SubNode.IsValid() ? SubNode.ToSharedRef() : StaticCastSharedRef<SGraphNode>(AsShared());
 }
 
+void SGraphNode_BehaviorTree::MoveTo(const FVector2D& NewPosition, FNodeSet& NodeFilter)
+{
+	SGraphNodeAI::MoveTo(NewPosition, NodeFilter);
+
+	// keep node order (defined by linked pins) up to date with actual positions
+	// this function will keep spamming on every mouse move update
+	UBehaviorTreeGraphNode* BTGraphNode = Cast<UBehaviorTreeGraphNode>(GraphNode);
+	if (BTGraphNode && !BTGraphNode->IsSubNode())
+	{
+		UBehaviorTreeGraph* BTGraph = BTGraphNode->GetBehaviorTreeGraph();
+		if (BTGraph)
+		{
+			for (int32 Idx = 0; Idx < BTGraphNode->Pins.Num(); Idx++)
+			{
+				UEdGraphPin* Pin = BTGraphNode->Pins[Idx];
+				if (Pin && Pin->Direction == EGPD_Input && Pin->LinkedTo.Num() == 1) 
+				{
+					UEdGraphPin* ParentPin = Pin->LinkedTo[0];
+					if (ParentPin)
+					{
+						BTGraph->RebuildChildOrder(ParentPin->GetOwningNode());
+					}
+				}
+			}
+		}
+	}
+}
+
 #undef LOCTEXT_NAMESPACE
\ No newline at end of file
diff --git a/Engine/Source/Editor/BehaviorTreeEditor/Private/SGraphNode_BehaviorTree.h b/Engine/Source/Editor/BehaviorTreeEditor/Private/SGraphNode_BehaviorTree.h
--- a/Engine/Source/Editor/BehaviorTreeEditor/Private/SGraphNode_BehaviorTree.h
+++ b/Engine/Source/Editor/BehaviorTreeEditor/Private/SGraphNode_BehaviorTree.h
@@ -20,6 +20,7 @@ public:
 	virtual void GetOverlayBrushes(bool bSelected, const FVector2D WidgetSize, TArray<FOverlayBrushInfo>& Brushes) const override;
 	virtual TArray<FOverlayWidgetInfo> GetOverlayWidgets(bool bSelected, const FVector2D& WidgetSize) const override;
 	virtual TSharedRef<SGraphNode> GetNodeUnderMouse(const FGeometry& MyGeometry, const FPointerEvent& MouseEvent) override;
+	virtual void MoveTo(const FVector2D& NewPosition, FNodeSet& NodeFilter) override;
 	// End of SGraphNode interface
 
 	/** handle double click */
diff --git a/Engine/Source/Editor/BehaviorTreeEditor/Public/BehaviorTreeEditorModule.h b/Engine/Source/Editor/BehaviorTreeEditor/Public/BehaviorTreeEditorModule.h
--- a/Engine/Source/Editor/BehaviorTreeEditor/Public/BehaviorTreeEditorModule.h
+++ b/Engine/Source/Editor/BehaviorTreeEditor/Public/BehaviorTreeEditorModule.h
@@ -4,6 +4,7 @@
 
 #include "ModuleInterface.h"
 #include "Toolkits/IToolkit.h"	// For EAssetEditorMode
+#include "Toolkits/AssetEditorToolkit.h"
 
 DECLARE_LOG_CATEGORY_EXTERN(LogBehaviorTreeEditor, Log, All);
 
diff --git a/Engine/Source/Editor/BlueprintGraph/Private/BlueprintActionDatabase.cpp b/Engine/Source/Editor/BlueprintGraph/Private/BlueprintActionDatabase.cpp
--- a/Engine/Source/Editor/BlueprintGraph/Private/BlueprintActionDatabase.cpp
+++ b/Engine/Source/Editor/BlueprintGraph/Private/BlueprintActionDatabase.cpp
@@ -946,7 +946,7 @@ static bool BlueprintActionDatabaseImpl::IsObjectValidForDatabase(UObject const*
 	{
 		bReturn = false;
 	}
-	else if(Object->GetOutermost()->PackageFlags & PKG_PlayInEditor)
+	else if(Object->GetOutermost()->HasAnyPackageFlags(PKG_PlayInEditor))
 	{
 		// Do not keep track of any PIE objects as we may prevent them from being cleaned up when ending PIE
 		bReturn = false;
diff --git a/Engine/Source/Editor/ContentBrowser/Private/AssetContextMenu.cpp b/Engine/Source/Editor/ContentBrowser/Private/AssetContextMenu.cpp
--- a/Engine/Source/Editor/ContentBrowser/Private/AssetContextMenu.cpp
+++ b/Engine/Source/Editor/ContentBrowser/Private/AssetContextMenu.cpp
@@ -1828,19 +1828,20 @@ bool FAssetContextMenu::CanExecuteDuplicate() const
 {
 	const TArray< FAssetData > AssetViewSelectedAssets = AssetView.Pin()->GetSelectedAssets();
 	uint32 NumNonRedirectors = 0;
-	for(auto& AssetData : AssetViewSelectedAssets)
+
+	for (const FAssetData& AssetData : AssetViewSelectedAssets)
 	{
-		if(!AssetData.IsValid())
+		if (!AssetData.IsValid())
 		{
 			continue;
 		}
 
-		if(AssetData.AssetClass == NAME_Class)
+		if (AssetData.AssetClass == NAME_Class || AssetData.AssetClass == UWorld::StaticClass()->GetFName())
 		{
 			return false;
 		}
 
-		if(AssetData.AssetClass != UObjectRedirector::StaticClass()->GetFName())
+		if (AssetData.AssetClass != UObjectRedirector::StaticClass()->GetFName())
 		{
 			++NumNonRedirectors;
 		}
diff --git a/Engine/Source/Editor/DetailCustomizations/Private/AnimMontageSegmentDetails.cpp b/Engine/Source/Editor/DetailCustomizations/Private/AnimMontageSegmentDetails.cpp
--- a/Engine/Source/Editor/DetailCustomizations/Private/AnimMontageSegmentDetails.cpp
+++ b/Engine/Source/Editor/DetailCustomizations/Private/AnimMontageSegmentDetails.cpp
@@ -226,7 +226,7 @@ void SAnimationSegmentViewport::InitSkeleton()
 		if (PreviewMesh)
 		{
 			UAnimSingleNodeInstance * Preview = PreviewComponent->PreviewInstance;
-			if((Preview == NULL || Preview->CurrentAsset != AnimSequence) ||
+			if((Preview == NULL || Preview->GetCurrentAsset() != AnimSequence) ||
 				(PreviewComponent->SkeletalMesh != PreviewMesh))
 			{
 				PreviewComponent->SetSkeletalMesh(PreviewMesh);
@@ -250,7 +250,7 @@ void SAnimationSegmentViewport::OnTickPreview( double InCurrentTime, float InDel
 	float Start, End;
 	StartTimePropertyHandle->GetValue( Start );
 	EndTimePropertyHandle->GetValue( End );
-	if ( PreviewComponent->PreviewInstance->CurrentTime > End || PreviewComponent->PreviewInstance->CurrentTime < Start )
+	if ( PreviewComponent->PreviewInstance->GetCurrentTime() > End || PreviewComponent->PreviewInstance->GetCurrentTime() < Start )
 	{
 		PreviewComponent->PreviewInstance->SetPosition( Start, false );
 	}
@@ -417,8 +417,8 @@ FReply SAnimationSegmentScrubPanel::OnClick_Forward()
 	UAnimSingleNodeInstance* PreviewInst = GetPreviewInstance();
 	if (PreviewInst)
 	{
-		bool bIsReverse = PreviewInst->bReverse;
-		bool bIsPlaying = PreviewInst->bPlaying;
+		bool bIsReverse = PreviewInst->IsReverse();
+		bool bIsPlaying = PreviewInst->IsPlaying();
 		// if current bIsReverse and bIsPlaying, we'd like to just turn off reverse
 		if (bIsReverse && bIsPlaying)
 		{
@@ -443,9 +443,9 @@ FReply SAnimationSegmentScrubPanel::OnClick_Forward()
 EPlaybackMode::Type SAnimationSegmentScrubPanel::GetPlaybackMode() const
 {
 	UAnimSingleNodeInstance* PreviewInst = GetPreviewInstance();
-	if (PreviewInst && PreviewInst->bPlaying)
+	if (PreviewInst && PreviewInst->IsPlaying())
 	{
-		return PreviewInst->bReverse ? EPlaybackMode::PlayingReverse : EPlaybackMode::PlayingForward;
+		return PreviewInst->IsReverse() ? EPlaybackMode::PlayingReverse : EPlaybackMode::PlayingForward;
 	}
 	return EPlaybackMode::Stopped;
 }
@@ -491,9 +491,9 @@ uint32 SAnimationSegmentScrubPanel::GetNumOfFrames() const
 		float Length = PreviewInst->GetLength();
 		// if anim sequence, use correct num frames
 		int32 NumFrames = (int32) (Length/0.0333f); 
-		if (PreviewInst->CurrentAsset && PreviewInst->CurrentAsset->IsA(UAnimSequenceBase::StaticClass()))
+		if (PreviewInst->GetCurrentAsset() && PreviewInst->GetCurrentAsset()->IsA(UAnimSequenceBase::StaticClass()))
 		{
-			NumFrames = CastChecked<UAnimSequenceBase>(PreviewInst->CurrentAsset)->GetNumberOfFrames();
+			NumFrames = CastChecked<UAnimSequenceBase>(PreviewInst->GetCurrentAsset())->GetNumberOfFrames();
 		}
 		return NumFrames;
 	}
@@ -521,7 +521,7 @@ bool SAnimationSegmentScrubPanel::DoesSyncViewport() const
 {
 	UAnimSingleNodeInstance* PreviewInst = GetPreviewInstance();
 
-	return (( LockedSequence==NULL && PreviewInst ) || ( LockedSequence && PreviewInst && PreviewInst->CurrentAsset == LockedSequence ));
+	return (( LockedSequence==NULL && PreviewInst ) || ( LockedSequence && PreviewInst && PreviewInst->GetCurrentAsset() == LockedSequence ));
 }
 
 class UAnimSingleNodeInstance* SAnimationSegmentScrubPanel::GetPreviewInstance() const
@@ -536,7 +536,7 @@ float SAnimationSegmentScrubPanel::GetScrubValue() const
 		UAnimSingleNodeInstance* PreviewInst = GetPreviewInstance();
 		if (PreviewInst)
 		{
-			return PreviewInst->CurrentTime; 
+			return PreviewInst->GetCurrentTime();
 		}
 	}
 	return 0.f;
diff --git a/Engine/Source/Editor/DetailCustomizations/Private/AnimSequenceDetails.cpp b/Engine/Source/Editor/DetailCustomizations/Private/AnimSequenceDetails.cpp
--- a/Engine/Source/Editor/DetailCustomizations/Private/AnimSequenceDetails.cpp
+++ b/Engine/Source/Editor/DetailCustomizations/Private/AnimSequenceDetails.cpp
@@ -493,7 +493,7 @@ void SAnimationRefPoseViewport::InitSkeleton()
 		{
 			UAnimSingleNodeInstance * Preview = PreviewComponent->PreviewInstance;
 			USkeletalMesh* PreviewSkeletalMesh = Skeleton->GetPreviewMesh();
-			if((Preview == NULL || Preview->CurrentAsset != AnimRef) || PreviewComponent->SkeletalMesh != PreviewSkeletalMesh)
+			if((Preview == NULL || Preview->GetCurrentAsset() != AnimRef) || PreviewComponent->SkeletalMesh != PreviewSkeletalMesh)
 			{
 				PreviewComponent->SetSkeletalMesh(PreviewSkeletalMesh);
 				PreviewComponent->EnablePreview(true, AnimRef, NULL);
diff --git a/Engine/Source/Editor/DetailCustomizations/Private/AttenuationSettingsCustomizations.cpp b/Engine/Source/Editor/DetailCustomizations/Private/AttenuationSettingsCustomizations.cpp
--- a/Engine/Source/Editor/DetailCustomizations/Private/AttenuationSettingsCustomizations.cpp
+++ b/Engine/Source/Editor/DetailCustomizations/Private/AttenuationSettingsCustomizations.cpp
@@ -77,7 +77,6 @@ void FAttenuationSettingsCustomization::CustomizeChildren( TSharedRef<IPropertyH
 
 	ChildBuilder.AddChildProperty(PropertyHandles.FindChecked(GET_MEMBER_NAME_CHECKED(FAttenuationSettings, bAttenuate)).ToSharedRef());
 	ChildBuilder.AddChildProperty(PropertyHandles.FindChecked(GET_MEMBER_NAME_CHECKED(FAttenuationSettings, bSpatialize)).ToSharedRef());
-	ChildBuilder.AddChildProperty(DistanceAlgorithmHandle.ToSharedRef() );
 
 	// Check to see if a spatialization plugin is enabled
 	if (IsAudioPluginEnabled(EAudioPlugin::SPATIALIZATION))
@@ -85,6 +84,8 @@ void FAttenuationSettingsCustomization::CustomizeChildren( TSharedRef<IPropertyH
 		ChildBuilder.AddChildProperty(SpatializationAlgorithmHandle.ToSharedRef());
 	}
 
+	ChildBuilder.AddChildProperty(DistanceAlgorithmHandle.ToSharedRef() );
+
 	IDetailPropertyRow& CustomCurveRow = ChildBuilder.AddChildProperty(PropertyHandles.FindChecked(GET_MEMBER_NAME_CHECKED(FAttenuationSettings, CustomAttenuationCurve)).ToSharedRef());
 	CustomCurveRow.Visibility(TAttribute<EVisibility>(this, &FAttenuationSettingsCustomization::IsCustomCurveSelected));
 
@@ -191,8 +192,32 @@ void FAttenuationSettingsCustomization::CustomizeChildren( TSharedRef<IPropertyH
 	ChildBuilder.AddChildProperty(PropertyHandles.FindChecked(GET_MEMBER_NAME_CHECKED(FAttenuationSettings, bAttenuateWithLPF)).ToSharedRef());
 	ChildBuilder.AddChildProperty(PropertyHandles.FindChecked(GET_MEMBER_NAME_CHECKED(FAttenuationSettings, LPFRadiusMin)).ToSharedRef());
 	ChildBuilder.AddChildProperty(PropertyHandles.FindChecked(GET_MEMBER_NAME_CHECKED(FAttenuationSettings, LPFRadiusMax)).ToSharedRef());
+	
+
+	bIsSpatializedHandle = PropertyHandles.FindChecked(GET_MEMBER_NAME_CHECKED(FAttenuationSettings, bSpatialize)).ToSharedRef();
+	bIsFocusedHandle = PropertyHandles.FindChecked(GET_MEMBER_NAME_CHECKED(FAttenuationSettings, bEnableListenerFocus)).ToSharedRef();
+
+	ChildBuilder.AddChildProperty(PropertyHandles.FindChecked(GET_MEMBER_NAME_CHECKED(FAttenuationSettings, bEnableListenerFocus)).ToSharedRef());
+
+	ChildBuilder.AddChildProperty(PropertyHandles.FindChecked(GET_MEMBER_NAME_CHECKED(FAttenuationSettings, FocusAzimuth)).ToSharedRef())
+		.EditCondition(GetIsFocusEnabledAttribute(), nullptr);
+
+	ChildBuilder.AddChildProperty(PropertyHandles.FindChecked(GET_MEMBER_NAME_CHECKED(FAttenuationSettings, NonFocusAzimuth)).ToSharedRef())
+		.EditCondition(GetIsFocusEnabledAttribute(), nullptr);
+
+	ChildBuilder.AddChildProperty(PropertyHandles.FindChecked(GET_MEMBER_NAME_CHECKED(FAttenuationSettings, FocusDistanceScale)).ToSharedRef())
+		.EditCondition(GetIsFocusEnabledAttribute(), nullptr);
+
+	ChildBuilder.AddChildProperty(PropertyHandles.FindChecked(GET_MEMBER_NAME_CHECKED(FAttenuationSettings, FocusPriorityScale)).ToSharedRef())
+		.EditCondition(GetIsFocusEnabledAttribute(), nullptr);
+
+	ChildBuilder.AddChildProperty(PropertyHandles.FindChecked(GET_MEMBER_NAME_CHECKED(FAttenuationSettings, NonFocusPriorityScale)).ToSharedRef())
+		.EditCondition(GetIsFocusEnabledAttribute(), nullptr);
+
+	ChildBuilder.AddChildProperty(PropertyHandles.FindChecked(GET_MEMBER_NAME_CHECKED(FAttenuationSettings, NonFocusVolumeAttenuation)).ToSharedRef())
+		.EditCondition(GetIsFocusEnabledAttribute(), nullptr);
 
-	if (PropertyHandles.Num() != 15)
+	if (PropertyHandles.Num() != 23)
 	{
 		FString PropertyList;
 		for (auto It(PropertyHandles.CreateConstIterator()); It; ++It)
@@ -203,6 +228,27 @@ void FAttenuationSettingsCustomization::CustomizeChildren( TSharedRef<IPropertyH
 	}
 }
 
+bool FAttenuationSettingsCustomization::IsFocusedEnabled() const
+{
+	bool bIsFocusEnabled;
+	bIsFocusedHandle->GetValue(bIsFocusEnabled);
+	if (!bIsFocusEnabled)
+	{
+		return false;
+	}
+
+	bool bIsSpatialized;
+	bIsSpatializedHandle->GetValue(bIsSpatialized);
+
+	return bIsSpatialized;
+}
+
+TAttribute<bool> FAttenuationSettingsCustomization::GetIsFocusEnabledAttribute() const
+{
+	return TAttribute<bool>(this, &FAttenuationSettingsCustomization::IsFocusedEnabled);
+}
+
+
 EVisibility FAttenuationSettingsCustomization::IsConeSelected() const
 {
 	uint8 AttenuationShapeValue;
diff --git a/Engine/Source/Editor/DetailCustomizations/Private/AttenuationSettingsCustomizations.h b/Engine/Source/Editor/DetailCustomizations/Private/AttenuationSettingsCustomizations.h
--- a/Engine/Source/Editor/DetailCustomizations/Private/AttenuationSettingsCustomizations.h
+++ b/Engine/Source/Editor/DetailCustomizations/Private/AttenuationSettingsCustomizations.h
@@ -23,4 +23,12 @@ protected:
 	TSharedPtr< IPropertyHandle > AttenuationShapeHandle;
 	TSharedPtr< IPropertyHandle > DistanceAlgorithmHandle;
 	TSharedPtr< IPropertyHandle > SpatializationAlgorithmHandle;
+
+	TSharedPtr< IPropertyHandle > bIsSpatializedHandle;
+	TSharedPtr< IPropertyHandle > bIsFocusedHandle;
+
+	bool IsFocusedEnabled() const;
+	TAttribute<bool> GetIsFocusEnabledAttribute() const;
+
+	TAttribute<bool> IsFocusEnabledAttribute;
 };
diff --git a/Engine/Source/Editor/DetailCustomizations/Private/ComponentTransformDetails.cpp b/Engine/Source/Editor/DetailCustomizations/Private/ComponentTransformDetails.cpp
--- a/Engine/Source/Editor/DetailCustomizations/Private/ComponentTransformDetails.cpp
+++ b/Engine/Source/Editor/DetailCustomizations/Private/ComponentTransformDetails.cpp
@@ -25,7 +25,7 @@ public:
 		if( GUnrealEd->PlayWorld && !GIsPlayInEditorWorld )
 		{
 			UPackage* ObjectPackage = Object->GetOutermost();
-			bRequiresPlayWorld = !!(ObjectPackage->PackageFlags & PKG_PlayInEditor);
+			bRequiresPlayWorld = ObjectPackage->HasAnyPackageFlags(PKG_PlayInEditor);
 		}
 
 		if( bRequiresPlayWorld )
diff --git a/Engine/Source/Editor/DetailCustomizations/Private/DetailCustomizations.cpp b/Engine/Source/Editor/DetailCustomizations/Private/DetailCustomizations.cpp
--- a/Engine/Source/Editor/DetailCustomizations/Private/DetailCustomizations.cpp
+++ b/Engine/Source/Editor/DetailCustomizations/Private/DetailCustomizations.cpp
@@ -103,6 +103,7 @@
 #include "CaptureResolutionCustomization.h"
 #include "TextCustomization.h"
 #include "AnimTrailNodeDetails.h"
+#include "MaterialProxySettingsCustomizations.h"
 
 
 IMPLEMENT_MODULE( FDetailCustomizationsModule, DetailCustomizations );
@@ -211,6 +212,7 @@ void FDetailCustomizationsModule::RegisterPropertyTypeCustomizations()
 	RegisterCustomPropertyTypeLayout("AssetImportInfo", FOnGetPropertyTypeCustomizationInstance::CreateStatic(&FAssetImportDataCustomization::MakeInstance));
 	RegisterCustomPropertyTypeLayout("CaptureResolution", FOnGetPropertyTypeCustomizationInstance::CreateStatic(&FCaptureResolutionCustomization::MakeInstance));
 	RegisterCustomPropertyTypeLayout("WeightedBlendable", FOnGetPropertyTypeCustomizationInstance::CreateStatic(&FWeightedBlendableCustomization::MakeInstance));
+	RegisterCustomPropertyTypeLayout("MaterialProxySettings", FOnGetPropertyTypeCustomizationInstance::CreateStatic(&FMaterialProxySettingsCustomizations::MakeInstance));
 }
 
 
diff --git a/Engine/Source/Editor/DetailCustomizations/Private/MaterialProxySettingsCustomizations.cpp b/Engine/Source/Editor/DetailCustomizations/Private/MaterialProxySettingsCustomizations.cpp
new file mode 100644
--- /dev/null
+++ b/Engine/Source/Editor/DetailCustomizations/Private/MaterialProxySettingsCustomizations.cpp
@@ -0,0 +1,106 @@
+// Copyright 1998-2015 Epic Games, Inc. All Rights Reserved.
+
+#include "DetailCustomizationsPrivatePCH.h"
+#include "MaterialProxySettingsCustomizations.h"
+#include "Engine/EngineTypes.h"
+
+#define LOCTEXT_NAMESPACE "MaterialProxySettingsCustomizations"
+
+TSharedRef<IPropertyTypeCustomization> FMaterialProxySettingsCustomizations::MakeInstance()
+{
+	return MakeShareable(new FMaterialProxySettingsCustomizations);
+}
+
+void FMaterialProxySettingsCustomizations::CustomizeHeader(TSharedRef<IPropertyHandle> StructPropertyHandle, FDetailWidgetRow& HeaderRow, IPropertyTypeCustomizationUtils& CustomizationUtils)
+{
+	HeaderRow.
+		NameContent()
+		[
+			StructPropertyHandle->CreatePropertyNameWidget()
+		]
+		.ValueContent()
+		[
+			StructPropertyHandle->CreatePropertyValueWidget()
+		];
+}
+
+void FMaterialProxySettingsCustomizations::CustomizeChildren(TSharedRef<IPropertyHandle> StructPropertyHandle, class IDetailChildrenBuilder& ChildBuilder, IPropertyTypeCustomizationUtils& StructCustomizationUtils)
+{
+	// Retrieve structure's child properties
+	uint32 NumChildren;
+	StructPropertyHandle->GetNumChildren( NumChildren );	
+	TMap<FName, TSharedPtr< IPropertyHandle > > PropertyHandles;	
+	for( uint32 ChildIndex = 0; ChildIndex < NumChildren; ++ChildIndex )
+	{
+		TSharedRef<IPropertyHandle> ChildHandle = StructPropertyHandle->GetChildHandle( ChildIndex ).ToSharedRef();
+		const FName PropertyName = ChildHandle->GetProperty()->GetFName();
+
+		PropertyHandles.Add(PropertyName, ChildHandle);
+	}
+	
+	// Retrieve special case properties
+	EnumHandle = PropertyHandles.FindChecked(GET_MEMBER_NAME_CHECKED(FMaterialProxySettings, TextureSizingType));
+	TextureSizeHandle = PropertyHandles.FindChecked(GET_MEMBER_NAME_CHECKED(FMaterialProxySettings, TextureSize));
+	DiffuseTextureSizeHandle = PropertyHandles.FindChecked(GET_MEMBER_NAME_CHECKED(FMaterialProxySettings, DiffuseTextureSize));
+	NormalTextureSizeHandle = PropertyHandles.FindChecked(GET_MEMBER_NAME_CHECKED(FMaterialProxySettings, NormalTextureSize));
+	MetallicTextureSizeHandle = PropertyHandles.FindChecked(GET_MEMBER_NAME_CHECKED(FMaterialProxySettings, MetallicTextureSize));
+	RoughnessTextureSizeHandle = PropertyHandles.FindChecked(GET_MEMBER_NAME_CHECKED(FMaterialProxySettings, RoughnessTextureSize));
+	SpecularTextureSizeHandle = PropertyHandles.FindChecked(GET_MEMBER_NAME_CHECKED(FMaterialProxySettings, SpecularTextureSize));
+	EmissiveTextureSizeHandle = PropertyHandles.FindChecked(GET_MEMBER_NAME_CHECKED(FMaterialProxySettings, EmissiveTextureSize));
+	OpacityTextureSizeHandle = PropertyHandles.FindChecked(GET_MEMBER_NAME_CHECKED(FMaterialProxySettings, OpacityTextureSize));
+	
+	for( auto Iter(PropertyHandles.CreateConstIterator()); Iter; ++Iter  )
+	{
+		// Handle special property cases (done inside the loop to maintain order according to the struct
+		if (Iter.Value() == DiffuseTextureSizeHandle 
+			|| Iter.Value() == DiffuseTextureSizeHandle
+			|| Iter.Value() == NormalTextureSizeHandle
+			|| Iter.Value() == MetallicTextureSizeHandle
+			|| Iter.Value() == RoughnessTextureSizeHandle
+			|| Iter.Value() == SpecularTextureSizeHandle
+			|| Iter.Value() == EmissiveTextureSizeHandle
+			|| Iter.Value() == OpacityTextureSizeHandle
+			)
+		{
+			IDetailPropertyRow& SizeRow = ChildBuilder.AddChildProperty(Iter.Value().ToSharedRef());
+			SizeRow.Visibility(TAttribute<EVisibility>(this, &FMaterialProxySettingsCustomizations::AreManualOverrideTextureSizesEnabled));
+		}
+		else if (Iter.Value() == TextureSizeHandle)
+		{
+			IDetailPropertyRow& SettingsRow = ChildBuilder.AddChildProperty(Iter.Value().ToSharedRef());
+			SettingsRow.Visibility(TAttribute<EVisibility>(this, &FMaterialProxySettingsCustomizations::IsTextureSizeEnabled));
+		}
+		else
+		{
+			IDetailPropertyRow& SettingsRow = ChildBuilder.AddChildProperty(Iter.Value().ToSharedRef());
+		}
+	}	
+}
+
+EVisibility FMaterialProxySettingsCustomizations::AreManualOverrideTextureSizesEnabled() const
+{
+	uint8 TypeValue;
+	EnumHandle->GetValue(TypeValue);
+
+	if (TypeValue == TextureSizingType_UseManualOverrideTextureSize)
+	{
+		return EVisibility::Visible;
+	}
+
+	return EVisibility::Hidden;
+}
+
+EVisibility FMaterialProxySettingsCustomizations::IsTextureSizeEnabled() const
+{
+	uint8 TypeValue;
+	EnumHandle->GetValue(TypeValue);
+
+	if (TypeValue == TextureSizingType_UseSimplygonAutomaticSizing || TypeValue == TextureSizingType_UseManualOverrideTextureSize)
+	{		
+		return EVisibility::Hidden;
+	}
+
+	return EVisibility::Visible;	
+}
+
+#undef LOCTEXT_NAMESPACE // MaterialProxySettingsCustomizations
\ No newline at end of file
diff --git a/Engine/Source/Editor/DetailCustomizations/Private/MaterialProxySettingsCustomizations.h b/Engine/Source/Editor/DetailCustomizations/Private/MaterialProxySettingsCustomizations.h
new file mode 100644
--- /dev/null
+++ b/Engine/Source/Editor/DetailCustomizations/Private/MaterialProxySettingsCustomizations.h
@@ -0,0 +1,28 @@
+// Copyright 1998-2015 Epic Games, Inc. All Rights Reserved.
+
+#pragma once
+
+class FMaterialProxySettingsCustomizations : public IPropertyTypeCustomization
+{
+public:
+	static TSharedRef<IPropertyTypeCustomization> MakeInstance();
+
+	/** IPropertyTypeCustomization instance */
+	virtual void CustomizeHeader(TSharedRef<IPropertyHandle> PropertyHandle, FDetailWidgetRow& HeaderRow, IPropertyTypeCustomizationUtils& CustomizationUtils) override;
+	virtual void CustomizeChildren( TSharedRef<IPropertyHandle> StructPropertyHandle, class IDetailChildrenBuilder& ChildBuilder, IPropertyTypeCustomizationUtils& StructCustomizationUtils ) override;
+
+protected:
+	EVisibility AreManualOverrideTextureSizesEnabled() const;
+	EVisibility IsTextureSizeEnabled() const;
+	
+	TSharedPtr< IPropertyHandle > EnumHandle;
+
+	TSharedPtr< IPropertyHandle > TextureSizeHandle;
+	TSharedPtr< IPropertyHandle > DiffuseTextureSizeHandle;
+	TSharedPtr< IPropertyHandle > NormalTextureSizeHandle;
+	TSharedPtr< IPropertyHandle > MetallicTextureSizeHandle;
+	TSharedPtr< IPropertyHandle > RoughnessTextureSizeHandle;
+	TSharedPtr< IPropertyHandle > SpecularTextureSizeHandle;
+	TSharedPtr< IPropertyHandle > EmissiveTextureSizeHandle;
+	TSharedPtr< IPropertyHandle > OpacityTextureSizeHandle;
+};
diff --git a/Engine/Source/Editor/HierarchicalLODOutliner/Private/HLODOutliner.cpp b/Engine/Source/Editor/HierarchicalLODOutliner/Private/HLODOutliner.cpp
--- a/Engine/Source/Editor/HierarchicalLODOutliner/Private/HLODOutliner.cpp
+++ b/Engine/Source/Editor/HierarchicalLODOutliner/Private/HLODOutliner.cpp
@@ -157,7 +157,7 @@ namespace HLODOutliner
 					.HAlign(HAlign_Center)
 					.Text(LOCTEXT("BuildMeshes", "Build LOD Meshes for Clusters"))
 					.OnClicked(this, &SHLODOutliner::HandleBuildLODActors)
-				]
+				]				
 			];
 		
 	}
@@ -205,7 +205,6 @@ namespace HLODOutliner
 				.FillWidth(0.5f)
 				[
 					SNew(SSlider)
-					.IsEnabled(this, &SHLODOutliner::HandleForcedLevelSliderIsEnabled)
 					.OnValueChanged(this, &SHLODOutliner::HandleForcedLevelSliderValueChanged)
 					.OnMouseCaptureBegin(this, &SHLODOutliner::HandleForcedLevelSliderCaptureBegin)
 					.OnMouseCaptureEnd(this, &SHLODOutliner::HandleForcedLevelSliderCaptureEnd)
@@ -418,24 +417,9 @@ namespace HLODOutliner
 		return FReply::Handled();
 	}
 
-	FReply SHLODOutliner::HandleTestFunction()
-	{
-		if (SelectedNodes.Num() > 0)
-		{
-			if (SelectedNodes[0]->GetTreeItemType() == ITreeItem::HierarchicalLODActor)
-			{
-				FLODActorItem* Item = static_cast<FLODActorItem*>(SelectedNodes[0].Get());
-				ALODActor* LODActor = Item->LODActor.Get();
-				
-				FHierarchicalLODUtilities::CreateVolumeForLODActor(LODActor, CurrentWorld);
-			}
-		}
-		return FReply::Handled();
-	}
-
 	END_SLATE_FUNCTION_BUILD_OPTIMIZATION
 
-		void SHLODOutliner::RegisterDelegates()
+	void SHLODOutliner::RegisterDelegates()
 	{
 		FEditorDelegates::MapChange.AddSP(this, &SHLODOutliner::OnMapChange);
 		FEditorDelegates::NewCurrentLevel.AddSP(this, &SHLODOutliner::OnNewCurrentLevel);
@@ -507,7 +491,7 @@ namespace HLODOutliner
 		}
 	}
 
-	bool SHLODOutliner::HandleForcedLevelSliderIsEnabled() const
+	bool SHLODOutliner::AreHLODsBuild() const
 	{
 		bool bHLODsBuild = true;
 		for (bool Build : LODLevelBuildFlags)
@@ -638,13 +622,16 @@ namespace HLODOutliner
 			for (auto SelectedItem : SelectedItems )
 			{
 				FLODActorItem* ActorItem = (FLODActorItem*)(SelectedItem.Get());
-				auto Parent = ActorItem->GetParent();
-
-				ITreeItem::TreeItemType Type = Parent->GetTreeItemType();
-				if (Type == ITreeItem::HierarchicalLODLevel)
+				if (ActorItem->LODActor->SubActors.Num() > 1)
 				{
-					FLODLevelItem* LevelItem = (FLODLevelItem*)(Parent.Get());
-					CurrentWorld->HierarchicalLODBuilder->BuildMeshForLODActor(ActorItem->LODActor.Get(), LevelItem->LODLevelIndex);
+					auto Parent = ActorItem->GetParent();
+
+					ITreeItem::TreeItemType Type = Parent->GetTreeItemType();
+					if (Type == ITreeItem::HierarchicalLODLevel)
+					{
+						FLODLevelItem* LevelItem = (FLODLevelItem*)(Parent.Get());
+						CurrentWorld->HierarchicalLODBuilder->BuildMeshForLODActor(ActorItem->LODActor.Get(), LevelItem->LODLevelIndex);
+					}
 				}
 			}
 
@@ -664,15 +651,17 @@ namespace HLODOutliner
 			for (auto SelectedItem : SelectedItems)
 			{
 				FLODActorItem* ActorItem = (FLODActorItem*)(SelectedItem.Get());
-
-				auto Parent = ActorItem->GetParent();
-
-				ITreeItem::TreeItemType Type = Parent->GetTreeItemType();
-				if (Type == ITreeItem::HierarchicalLODLevel)
+				if (ActorItem->LODActor->SubActors.Num() > 1)
 				{
-					FLODLevelItem* LevelItem = (FLODLevelItem*)(Parent.Get());
-					ActorItem->LODActor->SetIsDirty(true);
-					CurrentWorld->HierarchicalLODBuilder->BuildMeshForLODActor(ActorItem->LODActor.Get(), LevelItem->LODLevelIndex);			
+					auto Parent = ActorItem->GetParent();
+
+					ITreeItem::TreeItemType Type = Parent->GetTreeItemType();
+					if (Type == ITreeItem::HierarchicalLODLevel)
+					{
+						FLODLevelItem* LevelItem = (FLODLevelItem*)(Parent.Get());
+						ActorItem->LODActor->SetIsDirty(true);
+						CurrentWorld->HierarchicalLODBuilder->BuildMeshForLODActor(ActorItem->LODActor.Get(), LevelItem->LODLevelIndex);
+					}
 				}
 			}
 
@@ -1195,6 +1184,8 @@ namespace HLODOutliner
 						{
 							DestroySelectionActors();
 						}
+
+						TreeView->RequestScrollIntoView(*Item);
 					}
 					else
 					{
@@ -1214,6 +1205,8 @@ namespace HLODOutliner
 				{
 					DestroySelectionActors();
 				}
+
+				TreeView->RequestScrollIntoView(*Item);
 			}	
 			else
 			{
diff --git a/Engine/Source/Editor/HierarchicalLODOutliner/Private/HLODOutliner.h b/Engine/Source/Editor/HierarchicalLODOutliner/Private/HLODOutliner.h
--- a/Engine/Source/Editor/HierarchicalLODOutliner/Private/HLODOutliner.h
+++ b/Engine/Source/Editor/HierarchicalLODOutliner/Private/HLODOutliner.h
@@ -103,7 +103,6 @@ namespace HLODOutliner
 		FReply RetrieveActors();
 		FReply HandleBuildLODActors();
 		FReply HandleForceRefresh();
-		FReply HandleTestFunction();
 		/** End button handlers */
 
 	private:
@@ -118,11 +117,11 @@ namespace HLODOutliner
 		void ForceViewLODActor();
 
 		/**
-		* Determines if the HLOD slider is enabled, only if all HLODs are buil
+		* Returns whether or not all HLODs in the level are build
 		*
 		* @return bool
 		*/
-		bool HandleForcedLevelSliderIsEnabled() const;
+		bool AreHLODsBuild() const;
 		
 		/**
 		* Handles changes in slider value, maps/snaps it to the corresponding HLOD level
diff --git a/Engine/Source/Editor/HierarchicalLODOutliner/Private/LODActorItem.cpp b/Engine/Source/Editor/HierarchicalLODOutliner/Private/LODActorItem.cpp
--- a/Engine/Source/Editor/HierarchicalLODOutliner/Private/LODActorItem.cpp
+++ b/Engine/Source/Editor/HierarchicalLODOutliner/Private/LODActorItem.cpp
@@ -88,8 +88,8 @@ FText HLODOutliner::FLODActorItem::GetNumTrianglesAsText() const
 	{
 		const uint32 SubActorCount = LODActor->GetNumTrianglesInSubActors();
 		const uint32 MergedCount = LODActor->GetNumTrianglesInMergedMesh();
-		const uint32 PercentageOfOriginal = ((float)MergedCount / (float)SubActorCount) * 100;
-		return FText::FromString(FString::FromInt(SubActorCount) + TEXT(" (Original) - ") + FString::FromInt(MergedCount) + TEXT(" (Merged) - ") + FString::FromInt(PercentageOfOriginal) + "% (of Orginal)" );
+		const float PercentageOfOriginal = ((float)MergedCount / (float)SubActorCount) * 100.0f;
+		return FText::FromString(FString::FromInt(SubActorCount) + TEXT(" (Original) - ") + FString::FromInt(MergedCount) + TEXT(" (Merged) - ") + FString::SanitizeFloat(PercentageOfOriginal) + "% (of Original)");
 	}
 	else
 	{
