Copying //UE4/Dev-Framework to //UE4/Dev-Main (Source: //UE4/Dev-Framework @ 3038004)

==========================
MAJOR FEATURES + CHANGES
==========================

Change 3021479 on 2016/06/21 by [REDACTED_PERSON]

	Fix child actor properties set in the parent's construction script from being wiped out (4.12)

Change 3021703 on 2016/06/21 by [REDACTED_PERSON]

	Fix crash due to copying properties to registered components and then reregistering them. (4.12)

Change 3022105 on 2016/06/21 by [REDACTED_PERSON]

	-new test assets for James for PSD node

Change 3022621 on 2016/06/22 by [REDACTED_PERSON]

	Add AnimBP for testing PSD

Change 3022622 on 2016/06/22 by [REDACTED_PERSON]

	Only restrict anim asset selection for UAnimGraphNode_AssetPlayerBase derived nodes

Change 3022656 on 2016/06/22 by [REDACTED_PERSON]

	UE-30537 Fix solid collision geom drawing not working when mirrored

Change 3022657 on 2016/06/22 by [REDACTED_PERSON]

	Don't crash in FAnimGraphNodeDetails::OnShouldFilterAnimAsset if AnimAsset doesn't have Skeleton asset registry tag (shouldn't happen, but shouldn't crash)

Change 3022663 on 2016/06/22 by [REDACTED_PERSON]

	UE-31283 Additional extensibility for anim and physics
	PR #2434: Morpheme integration changes (Contributed by NaturalMotionTechnology)

Change 3022683 on 2016/06/22 by [REDACTED_PERSON]

	- Change OrientationDriver to always use PoseAsset for target poses
	- Remove NumPoses from PoseAsset and use GetAssetRegistryTags instead

Change 3022891 on 2016/06/22 by mason.seay

	Test asset for component hit

Change 3023203 on 2016/06/22 by mason.seay

	Updated map to use more noticeable sound assets

Change 3023335 on 2016/06/22 by [REDACTED_PERSON]

	Use AddReferencedObjects instead of iterating array manuallly

Change 3023351 on 2016/06/22 by [REDACTED_PERSON]

	Fix the case where physics hit events were passing the wrong component's bone info to the hit event.

Change 3023368 on 2016/06/22 by mason.seay

	Renamed actors in World Outliner

Change 3023425 on 2016/06/22 by mason.seay

	Moved asset to new folder and fixed deprecated node

Change 3023429 on 2016/06/22 by mason.seay

	Disabled collision on proc mesh

Change 3023553 on 2016/06/22 by [REDACTED_PERSON]

	Fix issue where MaxAngularVelocity resets to default on UPrimitiveComponent->BodyInstance. Replicated from CL 3009477.

Change 3024669 on 2016/06/23 by [REDACTED_PERSON]

	Update PSD test assets (removing unused)

Change 3024864 on 2016/06/23 by [REDACTED_PERSON]

	Audio Threading!

Change 3024877 on 2016/06/23 by [REDACTED_PERSON]

	PR #2375: Allow the creation of custom IStreamingManager (Contributed by [REDACTED_PERSON])

Change 3024880 on 2016/06/23 by [REDACTED_PERSON]

	PR #2209: Fix UGameplayStatics::*Game*Slot documentation (Contributed by [REDACTED_PERSON])

Change 3024939 on 2016/06/23 by [REDACTED_PERSON]

	- Add SwingOnly options to OrientationDriver
	- Move EBoneAxis from AnimNode_RotationMultiplier.h to AnimTypes.h
	- Calculate gaussian radius per pose, not globally

Change 3024940 on 2016/06/23 by [REDACTED_PERSON]

	PoseAsset editor improvements
	- Replace pose edit box with inline-editable style (with validation)
	- Add filter highlight
	- Show curve values for each pose when selected
	- Add different background for curve list
	- Filter box only searches pose list, moved location to indicate that

Change 3024949 on 2016/06/23 by [REDACTED_PERSON]

	Small update to PSD test AnimBP

Change 3025002 on 2016/06/23 by [REDACTED_PERSON]

	Fix the case where fixed frame rate combined with t.maxfps would lead to negative delta time. We now take the min of t.maxfps and fixed frame rate.

Change 3025214 on 2016/06/23 by mason.seay

	Updated Character Movement Map

Change 3025319 on 2016/06/23 by [REDACTED_PERSON]

	Make sure changing skeletal mesh updates the bone index on body instances already created.

Change 3025435 on 2016/06/23 by [REDACTED_PERSON]

	Fix welded bodies not updating their collision profile when calling SetCollisionProfile

Change 3025581 on 2016/06/23 by mason.seay

	Test asset for slicing procedural mesh

Change 3026483 on 2016/06/24 by [REDACTED_PERSON]

	Don't reschedule multiple times tick functions used as prerequisites

Change 3026498 on 2016/06/24 by mason.seay

	Updating blueprint for bug repro

Change 3026547 on 2016/06/24 by [REDACTED_PERSON]

	Fixed crash in FKismetDebugUtilities::GetWatchText()

	Crash reported by this UDN: https://udn.unrealengine.com/questions/300110/crash-in-kismetdebugutilities-when-printing-watchp.html

Change 3026598 on 2016/06/24 by [REDACTED_PERSON]

	Double clicking on poses now toggles them between 1.0 and 0.0 strength

Change 3026768 on 2016/06/24 by [REDACTED_PERSON]

	Change up suspend audio thread cvar sink warning about disabled threading to avoid inappropriate warnings

Change 3026802 on 2016/06/24 by [REDACTED_PERSON]

	#Pose Asset work

	# additive blending change : additive scale is saved to [targetscale/sourcescale - 1] where it used to be [targetscale/sourcescale] since blending doesn't work with it
	 - Blending should work once we save to [targetscale/sourcescale - 1] as normal - i.e. if you blend 0.3, it should not shrink the mesh because you applyed additive to 0.3
	 - When apply the scale to base, it should multiply [additive scale + 1 ]  where additive scale is [targetscale/sourcescale - 1]
	 - Changed FTransform::Blend to FTransform::Lerp since it's literally just Lerp. Name Blend should be used for Accumulate but changing the name now is dangerous, so I'm keeping Accumulate but changed Blend to Lerp

	# pose asset preview fix
	 - made sure it adds to curve, so you don't have to use delegate to apply
	 - PreviewOverride is now added to output curve, so we don't have to apply that curve later
	 - only reason of anim instance delegate is now for normal anim blueprint.

	#pose asset change
	 - Curve extraction happens with ExtractContext, the output curve is stricly output curve
	 - Pose Asset supports Shrink now, but only shrink if full pose
	 - Added PoseHandler to handle most of common stuff between different pose nodes
	 - Still have to work on how to update pose asset - wip
	 - todo: clean up single node player to handle pose asset in a function

	#code review:Martin.Wilson, James.Golding

Change 3026978 on 2016/06/24 by [REDACTED_PERSON]

	- Delete DrivePose Curve type
	- Renamed TriggerEvent to DriveAttribute for consistency
	- Replaced drive pose to drive attribute
	  - right now it can't have 0 curve type flags, so everything is DriveAttribute

Change 3027113 on 2016/06/24 by mason.seay

	Test Pose Assets

Change 3027454 on 2016/06/24 by [REDACTED_PERSON]

	UE-32492 Fix for cleaning up xaudio2 source voices and xaudio2 buffers if the source fails to initialize

	https://answers.unrealengine.com/questions/441080/audio-crash.html
	http://crashreporter/Crashes/Show/5689478

Change 3027519 on 2016/06/24 by [REDACTED_PERSON]

	Reverted FTransform name change as that causes compile errors due to lack of deprecated messages

	- not worth to keep the old functions and add new one

Change 3027887 on 2016/06/25 by [REDACTED_PERSON]

	Fix clang build warning

Change 3028703 on 2016/06/27 by [REDACTED_PERSON]

	gameplay debugger config improvements, categories and extensions can now be toggled while PIE/simulate is active
	#ue4

Change 3028792 on 2016/06/27 by [REDACTED_PERSON]

	compilation fix for gameplay debugger

Change 3028950 on 2016/06/27 by [REDACTED_PERSON]

	compilation fix for gameplay debugger

Change 3029003 on 2016/06/27 by [REDACTED_PERSON]

	Added PhysicalAnimation component that allows us to physically drive skeletal mesh from animation

Change 3029019 on 2016/06/27 by [REDACTED_PERSON]

	Update pose from source asset

Change 3029094 on 2016/06/27 by [REDACTED_PERSON]

	If Player->StartSpot is null disregard ShouldSpawnAtStartPoint returned true.

Change 3029308 on 2016/06/27 by [REDACTED_PERSON]

	-adding test animation for PSD node. Has morphs built in to compare against driver result

Change 3029372 on 2016/06/27 by [REDACTED_PERSON]

	Fix compile error after merge
	Also just fix the logic to be explicit rather than using suppression for static analysis warning

Change 3029493 on 2016/06/27 by [REDACTED_PERSON]

	Move PhysicsAsset.h out of public engine header.

Change 3029550 on 2016/06/27 by [REDACTED_PERSON]

	Fix crash with Nan when additive blending of poses\

Change 3029659 on 2016/06/27 by [REDACTED_PERSON]

	Adding new minor feature to add new concurrency mode

	- stop by lowest priorty but prevent new rather than stop oldest.

Change 3029673 on 2016/06/27 by [REDACTED_PERSON]

	Implementation in CL 3024124

Change 3030470 on 2016/06/28 by [REDACTED_PERSON]

	Fix OnConstraintBrokenWrapper being accidently wrapped with if WITH_CLOTHING

Change 3030586 on 2016/06/28 by [REDACTED_PERSON]

	Preview curve fix from anim curve viewer

Change 3031054 on 2016/06/28 by [REDACTED_PERSON]

Change 3031535 on 2016/06/28 by mason.seay

	Re-saving concurrency asset

Change 3031691 on 2016/06/28 by [REDACTED_PERSON]

	Fix stat sounds not turning on correctly unless a sort was specified

Change 3031883 on 2016/06/28 by [REDACTED_PERSON]

	#ue4 - Prevent bNotifyJumpApex from being editable, and clean up comments.

Change 3031898 on 2016/06/28 by [REDACTED_PERSON]

	#ue4 - Fix mesh smoothing on clients popping briefly when crouching. This was due to the change in 4.12 where we started smoothing Z location rather than always zeroing it (in certain movement modes).

	#udn https://udn.unrealengine.com/questions/300494/networked-crouching-jitter.html

Change 3032539 on 2016/06/29 by [REDACTED_PERSON]

	Don't destroy AudioDevices before draining audio commands and stopping audio thread

Change 3032633 on 2016/06/29 by [REDACTED_PERSON]

	In the same way that SpawnActor doesn't work during world teardown, don't allow new components to be added which could introduce recursion within the destroy logic.

Change 3032644 on 2016/06/29 by [REDACTED_PERSON]

	- Fixed issue where pose node evaluator doesn't show up in the menu with asset
	  - it showed twice of pose node (none) - jira UE-32358

	- Fixed issue where anim evaluator/pose asset by name/blend space evaluator failed to display assets properly
	   - jira UE-32359

	- support create pose menu from create asset - UE-32596
	 - added create pose asset from current pose

	- update source should refresh list - UE-32576

	- fixed blendspace to be in the blendspaces category

Change 3032847 on 2016/06/29 by [REDACTED_PERSON]

	Added PredictProjectilePath and SuggestProjectileVelocity_MediumArc utilities to UGameplayStatics.
	Updated SuggestProjectileVelocity to avoid floating point precision errors on gravity value comparison.

Change 3033124 on 2016/06/29 by [REDACTED_PERSON]

	Fix issue where InstancedStaticMeshComponent InstanceBodies don't move when the mesh is updated.

Change 3033155 on 2016/06/29 by [REDACTED_PERSON]

	- montage is playing and montage is pure
	- made montage parameter to be mostly const (except play), and made it consistently pointer

Change 3033157 on 2016/06/29 by [REDACTED_PERSON]

	Check in missing file

Change 3033456 on 2016/06/29 by [REDACTED_PERSON]

	fixed path following changes broken by merge
	#ue4

Change 3033956 on 2016/06/30 by bruce.nesbit

	PR #2483: Fix/Improvment Move Component To Rotation (Contributed by [REDACTED_PERSON])

	#test PIE

Change 3034019 on 2016/06/30 by [REDACTED_PERSON]

	Anim blueprint sub-instances, allowing anim blueprints to run within anim blueprints and expose parameters back to the "parent" instance.
	Caveats:
	    - Slots and state machine names are unique and boxed per instance, meaning playing a montage on a slot will only affect slots in the outermost instance and state machine getters are local to their instance.

Change 3034085 on 2016/06/30 by [REDACTED_PERSON]

	Missed LOCTEXT_NAMESPACE undefs from the subinstance checkin, for some reason doesn't get caught on windows, likely how the unity files are stuck together.

Change 3034162 on 2016/06/30 by [REDACTED_PERSON]

	Refactor bone reference widget so that selection tree can be used seperately

Change 3034205 on 2016/06/30 by [REDACTED_PERSON]

	#ANIM: fix issue with addiitve blending with non-full weight applying wrong scale

Change 3034339 on 2016/06/30 by [REDACTED_PERSON]

	Moving functionality from Skeleton Curves tab into Anim Curve Viewer tab

Change 3034426 on 2016/06/30 by [REDACTED_PERSON]

	CIS Fix

Change 3034629 on 2016/06/30 by [REDACTED_PERSON]

	Support non-zero curves to be stippred out upon importing

Change 3035863 on 2016/07/01 by [REDACTED_PERSON]

	When pasting components in to a blueprint, make the relative position and rotation of the root 0,0,0

Change 3035916 on 2016/07/01 by [REDACTED_PERSON]

	Fixed PaperGroupedSprite doesn't update InstanceBodies data in physics. This change is related to CL-3033124

Change 3035973 on 2016/07/01 by [REDACTED_PERSON]

	fixed hash function for FRecastDebugPathfindingNode
	#ue4

Change 3036024 on 2016/07/01 by [REDACTED_PERSON]

	#ue4 - Avoid filling in array in AActor::FixupNativeActorComponents() unless we detect a null scene component. Avoid copying TWeakObjectPtr in ValidateDeferredTransformCache().

Change 3036157 on 2016/07/01 by [REDACTED_PERSON]

	Protect against running commands on game thread when the audio device has already been freed

Change 3036178 on 2016/07/01 by [REDACTED_PERSON]

	Don't bitpack the gamethread specific boolean.

Change 3036906 on 2016/07/04 by bruce.nesbit

	Fixed a typo in HasDefaultBuildSettings - (bCompi8leLeanAndMeanUE should be bCompileLeanAndMeanUE)

	#tests Compiled

Change 3036929 on 2016/07/04 by [REDACTED_PERSON]

	UE-32405 Label Rotator components X/Y/Z instead of Roll/Pitch/Yaw

Change 3036930 on 2016/07/04 by [REDACTED_PERSON]

	UE-30414 Move constraint warnings to Message Log

Change 3036931 on 2016/07/04 by [REDACTED_PERSON]

	PR #2427: SkeletalMeshMerge now can transform the UVs of the source meshes. (Contributed by [REDACTED_PERSON])

Change 3037123 on 2016/07/04 by [REDACTED_PERSON]

	Added physical animation preview in PhAT as well as physical animation profiles.

Change 3037420 on 2016/07/05 by Jurre.deBaare

	Moved BodySetup_DEPRECATED out of WITH_EDITORONLY_DATA since it's being used in postload (fixes shipping builds)

Change 3037702 on 2016/07/05 by [REDACTED_PERSON]

	Copying change 3037701 from Release-4.12:

	Fixed crash when viewing uncompressed animation

	Made sure that objects required by the animation evaluation are set up when performing game-thread side work in the editor.

Change 3037837 on 2016/07/05 by [REDACTED_PERSON]

	sound stats will now still be displayed when creating a new audio device

diff --git a/Engine/Plugins/2D/Paper2D/Source/Paper2D/Classes/PaperGroupedSpriteComponent.h b/Engine/Plugins/2D/Paper2D/Source/Paper2D/Classes/PaperGroupedSpriteComponent.h
--- a/Engine/Plugins/2D/Paper2D/Source/Paper2D/Classes/PaperGroupedSpriteComponent.h
+++ b/Engine/Plugins/2D/Paper2D/Source/Paper2D/Classes/PaperGroupedSpriteComponent.h
@@ -76,6 +76,8 @@ public:
 	UFUNCTION(BlueprintCallable, Category = "Components|Sprite")
 	bool GetInstanceTransform(int32 InstanceIndex, FTransform& OutInstanceTransform, bool bWorldSpace = false) const;
 	
+	virtual void OnUpdateTransform(EUpdateTransformFlags UpdateTransformFlags, ETeleportType Teleport) override;
+
 	/** Update the transform for the instance specified. Instance is given in local space of this component unless bWorldSpace is set.  Returns True on success. */
 	UFUNCTION(BlueprintCallable, Category = "Components|Sprite")
 	virtual bool UpdateInstanceTransform(int32 InstanceIndex, const FTransform& NewInstanceTransform, bool bWorldSpace = false, bool bMarkRenderStateDirty = true);
diff --git a/Engine/Plugins/2D/Paper2D/Source/Paper2D/Private/PaperGroupedSpriteComponent.cpp b/Engine/Plugins/2D/Paper2D/Source/Paper2D/Private/PaperGroupedSpriteComponent.cpp
--- a/Engine/Plugins/2D/Paper2D/Source/Paper2D/Private/PaperGroupedSpriteComponent.cpp
+++ b/Engine/Plugins/2D/Paper2D/Source/Paper2D/Private/PaperGroupedSpriteComponent.cpp
@@ -69,6 +69,22 @@ bool UPaperGroupedSpriteComponent::GetInstanceTransform(int32 InstanceIndex, FTr
 	return true;
 }
 
+void UPaperGroupedSpriteComponent::OnUpdateTransform(EUpdateTransformFlags UpdateTransformFlags, ETeleportType Teleport)
+{
+	// We are handling the physics move below, so don't handle it at higher levels
+	Super::OnUpdateTransform(UpdateTransformFlags | EUpdateTransformFlags::SkipPhysicsUpdate, Teleport);
+
+	// Always send new transform to physics
+	if (bPhysicsStateCreated && !(EUpdateTransformFlags::SkipPhysicsUpdate & UpdateTransformFlags))
+	{
+		for (int i = 0; i < PerInstanceSpriteData.Num(); i++)
+		{
+			const FTransform InstanceTransform(PerInstanceSpriteData[i].Transform);
+			UpdateInstanceTransform(i, InstanceTransform * ComponentToWorld, true);
+		}
+	}
+}
+
 bool UPaperGroupedSpriteComponent::UpdateInstanceTransform(int32 InstanceIndex, const FTransform& NewInstanceTransform, bool bWorldSpace, bool bMarkRenderStateDirty)
 {
 	if (!PerInstanceSpriteData.IsValidIndex(InstanceIndex))
diff --git a/Engine/Source/Developer/GameplayDebugger/GameplayDebugger.Build.cs b/Engine/Source/Developer/GameplayDebugger/GameplayDebugger.Build.cs
--- a/Engine/Source/Developer/GameplayDebugger/GameplayDebugger.Build.cs
+++ b/Engine/Source/Developer/GameplayDebugger/GameplayDebugger.Build.cs
@@ -27,9 +27,14 @@ namespace UnrealBuildTool.Rules
 
 			PrivateIncludePaths.Add("Developer/GameplayDebugger/Private");
 
-            if (UEBuildConfiguration.bBuildEditor == true)
+            if (UEBuildConfiguration.bBuildEditor)
 			{
-				PrivateDependencyModuleNames.Add("UnrealEd");
+                PrivateDependencyModuleNames.AddRange(
+                    new string[] {
+                        "EditorStyle",
+                        "UnrealEd",
+                        "PropertyEditor",
+                    });
 			}
 
             Definitions.Add("ENABLE_OLD_GAMEPLAY_DEBUGGER=" + (bEnableDeprecatedDebugger ? "1" : "0"));
diff --git a/Engine/Source/Developer/GameplayDebugger/Private/GameplayDebuggerExtension_HUD.cpp b/Engine/Source/Developer/GameplayDebugger/Private/GameplayDebuggerExtension_HUD.cpp
--- a/Engine/Source/Developer/GameplayDebugger/Private/GameplayDebuggerExtension_HUD.cpp
+++ b/Engine/Source/Developer/GameplayDebugger/Private/GameplayDebuggerExtension_HUD.cpp
@@ -2,7 +2,6 @@
 
 #include "GameplayDebuggerPrivatePCH.h"
 #include "GameplayDebuggerExtension_HUD.h"
-#include "GameplayDebuggerConfig.h"
 #include "GameFramework/HUD.h"
 #include "GameFramework/PlayerController.h"
 
@@ -14,16 +13,13 @@ FGameplayDebuggerExtension_HUD::FGameplayDebuggerExtension_HUD()
 	bPrevDebugMessagesEnabled = GEngine && GEngine->bEnableOnScreenDebugMessages;
 	bIsCachedDescriptionValid = false;
 
-	const UGameplayDebuggerConfig* SettingsCDO = UGameplayDebuggerConfig::StaticClass()->GetDefaultObject<UGameplayDebuggerConfig>();
+	const FGameplayDebuggerInputHandlerConfig HUDKeyConfig(TEXT("ToggleHUD"), EKeys::Tilde.GetFName(), FGameplayDebuggerInputModifier::Ctrl);
+	const FGameplayDebuggerInputHandlerConfig MessagesKeyConfig(TEXT("ToggleMessages"), EKeys::Tab.GetFName(), FGameplayDebuggerInputModifier::Ctrl);
 
-	const bool bHasHUDBinding = BindKeyPress(SettingsCDO->GameHUDKey.GetFName(), 
-		SettingsCDO->bUseGameHUDModifiers ? SettingsCDO->GameHUDModifiers.CreateModifier() : FGameplayDebuggerInputModifier::None,
-		this, &FGameplayDebuggerExtension_HUD::ToggleGameHUD);
+	const bool bHasHUDBinding = BindKeyPress(HUDKeyConfig, this, &FGameplayDebuggerExtension_HUD::ToggleGameHUD);
 	HudBindingIdx = bHasHUDBinding ? (GetNumInputHandlers() - 1) : INDEX_NONE;
 
-	const bool bHasMessagesBinding = BindKeyPress(SettingsCDO->DebugMessagesKey.GetFName(),
-		SettingsCDO->bUseDebugMessagesModifiers ? SettingsCDO->DebugMessagesModifiers.CreateModifier() : FGameplayDebuggerInputModifier::None,
-		this, &FGameplayDebuggerExtension_HUD::ToggleDebugMessages);
+	const bool bHasMessagesBinding = BindKeyPress(MessagesKeyConfig, this, &FGameplayDebuggerExtension_HUD::ToggleDebugMessages);
 	MessagesBindingIdx = bHasMessagesBinding ? (GetNumInputHandlers() - 1) : INDEX_NONE;
 }
 
diff --git a/Engine/Source/Developer/GameplayDebugger/Private/GameplayDebuggerExtension_Spectator.cpp b/Engine/Source/Developer/GameplayDebugger/Private/GameplayDebuggerExtension_Spectator.cpp
--- a/Engine/Source/Developer/GameplayDebugger/Private/GameplayDebuggerExtension_Spectator.cpp
+++ b/Engine/Source/Developer/GameplayDebugger/Private/GameplayDebuggerExtension_Spectator.cpp
@@ -2,7 +2,6 @@
 
 #include "GameplayDebuggerPrivatePCH.h"
 #include "GameplayDebuggerExtension_Spectator.h"
-#include "GameplayDebuggerConfig.h"
 #include "GameplayDebuggerPlayerManager.h"
 #include "Engine/DebugCameraController.h"
 #include "Engine/Player.h"
@@ -11,11 +10,8 @@
 
 FGameplayDebuggerExtension_Spectator::FGameplayDebuggerExtension_Spectator()
 {
-	const UGameplayDebuggerConfig* SettingsCDO = UGameplayDebuggerConfig::StaticClass()->GetDefaultObject<UGameplayDebuggerConfig>();
-
-	bHasInputBinding = BindKeyPress(SettingsCDO->SpectatorKey.GetFName(),
-		SettingsCDO->bUseSpectatorModifiers ? SettingsCDO->SpectatorModifiers.CreateModifier() : FGameplayDebuggerInputModifier::None,
-		this, &FGameplayDebuggerExtension_Spectator::ToggleSpectatorMode);
+	const FGameplayDebuggerInputHandlerConfig KeyConfig(TEXT("Toggle"), EKeys::Tab.GetFName());
+	bHasInputBinding = BindKeyPress(KeyConfig, this, &FGameplayDebuggerExtension_Spectator::ToggleSpectatorMode);
 }
 
 TSharedRef<FGameplayDebuggerExtension> FGameplayDebuggerExtension_Spectator::MakeInstance()
diff --git a/Engine/Source/Developer/GameplayDebugger/Private/GameplayDebuggerModule.cpp b/Engine/Source/Developer/GameplayDebugger/Private/GameplayDebuggerModule.cpp
--- a/Engine/Source/Developer/GameplayDebugger/Private/GameplayDebuggerModule.cpp
+++ b/Engine/Source/Developer/GameplayDebugger/Private/GameplayDebuggerModule.cpp
@@ -10,6 +10,13 @@
 #include "GameplayDebuggerExtension_Spectator.h"
 #include "GameplayDebuggerExtension_HUD.h"
 
+#if WITH_EDITOR
+#include "PropertyEditorModule.h"
+#include "Editor/GameplayDebuggerCategoryConfigCustomization.h"
+#include "Editor/GameplayDebuggerExtensionConfigCustomization.h"
+#include "Editor/GameplayDebuggerInputConfigCustomization.h"
+#endif
+
 #if !ENABLE_OLD_GAMEPLAY_DEBUGGER
 
 class FGameplayDebuggerModule : public IGameplayDebugger
@@ -62,16 +69,15 @@ void FGameplayDebuggerCompat::StartupNewDebugger()
 				SettingsCDO);
 		}
 
-		if (SettingsCDO->bEnableExtension_GameHUD)
-		{
-			AddonManager.RegisterExtension("GameHUD", FOnGetExtension::CreateStatic(&FGameplayDebuggerExtension_HUD::MakeInstance));
-		}
-
-		if (SettingsCDO->bEnableExtension_Spectator)
-		{
-			AddonManager.RegisterExtension("Spectator", FOnGetExtension::CreateStatic(&FGameplayDebuggerExtension_Spectator::MakeInstance));
-		}
+#if WITH_EDITOR
+		FPropertyEditorModule& PropertyEditorModule = FModuleManager::GetModuleChecked<FPropertyEditorModule>("PropertyEditor");
+		PropertyEditorModule.RegisterCustomPropertyTypeLayout("GameplayDebuggerCategoryConfig", FOnGetPropertyTypeCustomizationInstance::CreateStatic(&FGameplayDebuggerCategoryConfigCustomization::MakeInstance));
+		PropertyEditorModule.RegisterCustomPropertyTypeLayout("GameplayDebuggerExtensionConfig", FOnGetPropertyTypeCustomizationInstance::CreateStatic(&FGameplayDebuggerExtensionConfigCustomization::MakeInstance));
+		PropertyEditorModule.RegisterCustomPropertyTypeLayout("GameplayDebuggerInputConfig", FOnGetPropertyTypeCustomizationInstance::CreateStatic(&FGameplayDebuggerInputConfigCustomization::MakeInstance));
+#endif
 
+		AddonManager.RegisterExtension("GameHUD", FOnGetExtension::CreateStatic(&FGameplayDebuggerExtension_HUD::MakeInstance));
+		AddonManager.RegisterExtension("Spectator", FOnGetExtension::CreateStatic(&FGameplayDebuggerExtension_Spectator::MakeInstance));
 		AddonManager.NotifyExtensionsChanged();
 	}
 }
@@ -91,6 +97,13 @@ void FGameplayDebuggerCompat::ShutdownNewDebugger()
 	{
 		SettingsModule->UnregisterSettings("Project", "Engine", "GameplayDebugger");
 	}
+
+#if WITH_EDITOR
+	FPropertyEditorModule& PropertyEditorModule = FModuleManager::GetModuleChecked<FPropertyEditorModule>("PropertyEditor");
+	PropertyEditorModule.UnregisterCustomPropertyTypeLayout("GameplayDebuggerCategoryConfig");
+	PropertyEditorModule.UnregisterCustomPropertyTypeLayout("GameplayDebuggerExtensionConfig");
+	PropertyEditorModule.UnregisterCustomPropertyTypeLayout("GameplayDebuggerInputConfig");
+#endif
 }
 
 void FGameplayDebuggerModule::RegisterCategory(FName CategoryName, IGameplayDebugger::FOnGetCategory MakeInstanceDelegate, EGameplayDebuggerCategoryState CategoryState, int32 SlotIdx)
diff --git a/Engine/Source/Developer/GameplayDebugger/Public/GameplayDebuggerAddonBase.h b/Engine/Source/Developer/GameplayDebugger/Public/GameplayDebuggerAddonBase.h
--- a/Engine/Source/Developer/GameplayDebugger/Public/GameplayDebuggerAddonBase.h
+++ b/Engine/Source/Developer/GameplayDebugger/Public/GameplayDebuggerAddonBase.h
@@ -32,7 +32,7 @@ protected:
 	/** returns replicator actor */
 	AGameplayDebuggerCategoryReplicator* GetReplicator() const;
 
-	/** creates new key binding handler */
+	/** creates new key binding handler: single key press */
 	template<class UserClass>
 	bool BindKeyPress(FName KeyName, UserClass* KeyHandlerObject, typename FGameplayDebuggerInputHandler::FHandler::TRawMethodDelegate< UserClass >::FMethodPtr KeyHandlerFunc, EGameplayDebuggerInputMode InputMode = EGameplayDebuggerInputMode::Local)
 	{
@@ -50,7 +50,7 @@ protected:
 		return false;
 	}
 
-	/** creates new key binding handler */
+	/** creates new key binding handler: key press with modifiers */
 	template<class UserClass>
 	bool BindKeyPress(FName KeyName, FGameplayDebuggerInputModifier KeyModifer, UserClass* KeyHandlerObject, typename FGameplayDebuggerInputHandler::FHandler::TRawMethodDelegate< UserClass >::FMethodPtr KeyHandlerFunc, EGameplayDebuggerInputMode InputMode = EGameplayDebuggerInputMode::Local)
 	{
@@ -69,6 +69,25 @@ protected:
 		return false;
 	}
 
+	/** creates new key binding handler: customizable key press, stored in config files */
+	template<class UserClass>
+	bool BindKeyPress(const FGameplayDebuggerInputHandlerConfig& InputConfig, UserClass* KeyHandlerObject, typename FGameplayDebuggerInputHandler::FHandler::TRawMethodDelegate< UserClass >::FMethodPtr KeyHandlerFunc, EGameplayDebuggerInputMode InputMode = EGameplayDebuggerInputMode::Local)
+	{
+		FGameplayDebuggerInputHandler NewHandler;
+		NewHandler.KeyName = InputConfig.KeyName;
+		NewHandler.Modifier = InputConfig.Modifier;
+		NewHandler.Delegate.BindRaw(KeyHandlerObject, KeyHandlerFunc);
+		NewHandler.Mode = InputMode;
+
+		if (NewHandler.IsValid())
+		{
+			InputHandlers.Add(NewHandler);
+			return true;
+		}
+
+		return false;
+	}
+
 private:
 
 	friend class FGameplayDebuggerAddonManager;
diff --git a/Engine/Source/Developer/GameplayDebugger/Public/GameplayDebuggerConfig.h b/Engine/Source/Developer/GameplayDebugger/Public/GameplayDebuggerConfig.h
--- a/Engine/Source/Developer/GameplayDebugger/Public/GameplayDebuggerConfig.h
+++ b/Engine/Source/Developer/GameplayDebugger/Public/GameplayDebuggerConfig.h
@@ -5,27 +5,79 @@
 #include "GameplayDebuggerTypes.h"
 #include "GameplayDebuggerConfig.generated.h"
 
+UENUM()
+enum class EGameplayDebuggerOverrideMode : uint8
+{
+	Enable,
+	Disable,
+	UseDefault,
+};
+
 USTRUCT()
-struct GAMEPLAYDEBUGGER_API FGameplayDebuggerInputConfig
+struct FGameplayDebuggerInputConfig
 {
-	GENERATED_USTRUCT_BODY();
+	GENERATED_USTRUCT_BODY()
+
+	UPROPERTY(VisibleAnywhere, Category = Input)
+	FString ConfigName;
+
+	UPROPERTY(EditAnywhere, Category = Input)
+	FKey Key;
 
 	UPROPERTY(EditAnywhere, Category = Input)
-	uint32 bShift : 1;
+	uint32 bModShift : 1;
 
 	UPROPERTY(EditAnywhere, Category = Input)
-	uint32 bCtrl : 1;
+	uint32 bModCtrl : 1;
 
 	UPROPERTY(EditAnywhere, Category = Input)
-	uint32 bAlt : 1;
+	uint32 bModAlt : 1;
 
 	UPROPERTY(EditAnywhere, Category = Input)
-	uint32 bCmd : 1;
+	uint32 bModCmd : 1;
+};
+
+USTRUCT()
+struct FGameplayDebuggerCategoryConfig
+{
+	GENERATED_USTRUCT_BODY()
+
+	UPROPERTY(VisibleAnywhere, Category = Settings)
+	FString CategoryName;
 
-	FGameplayDebuggerInputModifier CreateModifier() const
-	{
-		return FGameplayDebuggerInputModifier(bShift, bCtrl, bAlt, bCmd);
-	}
+	UPROPERTY(EditAnywhere, Category = Settings, meta = (EditCondition = "bOverrideSlotIdx", ClampMin = -1, ClampMax = 9, UIMin = -1, UIMax = 9))
+	int32 SlotIdx;
+
+	UPROPERTY(EditAnywhere, Category = Settings)
+	EGameplayDebuggerOverrideMode ActiveInGame;
+
+	UPROPERTY(EditAnywhere, Category = Settings)
+	EGameplayDebuggerOverrideMode ActiveInSimulate;
+
+	UPROPERTY(EditAnywhere, Category = Settings, AdvancedDisplay)
+	uint32 bOverrideSlotIdx : 1;
+
+	UPROPERTY(EditAnywhere, Category = Settings, EditFixedSize)
+	TArray<FGameplayDebuggerInputConfig> InputHandlers;
+
+	FGameplayDebuggerCategoryConfig() : ActiveInGame(EGameplayDebuggerOverrideMode::UseDefault), ActiveInSimulate(EGameplayDebuggerOverrideMode::UseDefault) {}
+};
+
+USTRUCT()
+struct FGameplayDebuggerExtensionConfig
+{
+	GENERATED_USTRUCT_BODY()
+
+	UPROPERTY(VisibleAnywhere, Category = Settings)
+	FString ExtensionName;
+
+	UPROPERTY(EditAnywhere, Category = Settings)
+	EGameplayDebuggerOverrideMode UseExtension;
+
+	UPROPERTY(EditAnywhere, Category = Settings, EditFixedSize)
+	TArray<FGameplayDebuggerInputConfig> InputHandlers;
+
+	FGameplayDebuggerExtensionConfig() : UseExtension(EGameplayDebuggerOverrideMode::UseDefault) {}
 };
 
 UCLASS(config = Engine, defaultconfig)
@@ -85,44 +137,6 @@ class GAMEPLAYDEBUGGER_API UGameplayDebuggerConfig : public UObject
 	UPROPERTY(config, EditAnywhere, Category = Input)
 	FKey CategorySlot9;
 
-	/** toggle game hud */
-	UPROPERTY(config, EditAnywhere, Category = Extensions)
-	FKey GameHUDKey;
-
-	UPROPERTY(config, EditAnywhere, Category = Extensions, meta = (EditCondition = "bUseGameHUDModifiers"))
-	FGameplayDebuggerInputConfig GameHUDModifiers;
-
-	/** toggle onscreen debug messages */
-	UPROPERTY(config, EditAnywhere, Category = Extensions)
-	FKey DebugMessagesKey;
-
-	UPROPERTY(config, EditAnywhere, Category = Extensions, meta = (EditCondition = "bUseDebugMessagesModifiers"))
-	FGameplayDebuggerInputConfig DebugMessagesModifiers;
-
-	/** toggle spectator mode */
-	UPROPERTY(config, EditAnywhere, Category = Extensions)
-	FKey SpectatorKey;
-
-	UPROPERTY(config, EditAnywhere, Category = Extensions, meta = (EditCondition = "bUseSpectatorModifiers"))
-	FGameplayDebuggerInputConfig SpectatorModifiers;
-
-	/** if set, game hud extension will be automatically added */
-	UPROPERTY(config, EditAnywhere, Category = Extensions)
-	uint32 bEnableExtension_GameHUD : 1;
-
-	/** if set, spectator mode extension will be automatically added */
-	UPROPERTY(config, EditAnywhere, Category = Extensions)
-	uint32 bEnableExtension_Spectator : 1;
-
-	UPROPERTY()
-	uint32 bUseGameHUDModifiers : 1;
-	
-	UPROPERTY()
-	uint32 bUseDebugMessagesModifiers : 1;
-	
-	UPROPERTY()
-	uint32 bUseSpectatorModifiers : 1;
-
 	/** additional canvas padding: left */
 	UPROPERTY(config, EditAnywhere, Category = Display)
 	float DebugCanvasPaddingLeft;
@@ -138,4 +152,38 @@ class GAMEPLAYDEBUGGER_API UGameplayDebuggerConfig : public UObject
 	/** additional canvas padding: bottom */
 	UPROPERTY(config, EditAnywhere, Category = Display)
 	float DebugCanvasPaddingBottom;
+
+	UPROPERTY(config, EditAnywhere, Category = AddOns, EditFixedSize)
+	TArray<FGameplayDebuggerCategoryConfig> Categories;
+
+	UPROPERTY(config, EditAnywhere, Category = AddOns, EditFixedSize)
+	TArray<FGameplayDebuggerExtensionConfig> Extensions;
+
+	/** updates entry in Categories array and modifies category creation params */
+	void UpdateCategoryConfig(const FName CategoryName, int32& SlotIdx, uint8& CategoryState);
+
+	/** updates entry in Categories array and modifies input binding params */
+	void UpdateCategoryInputConfig(const FName CategoryName, const FName InputName, FName& KeyName, FGameplayDebuggerInputModifier& KeyModifier);
+
+	/** updates entry in Extensions array and modifies extension creation params */
+	void UpdateExtensionConfig(const FName ExtensionName, uint8& UseExtension);
+
+	/** updates entry in Categories array and modifies input binding params */
+	void UpdateExtensionInputConfig(const FName ExtensionName, const FName InputName, FName& KeyName, FGameplayDebuggerInputModifier& KeyModifier);
+
+	/** remove all category and extension data from unknown sources (outdated entries) */
+	void RemoveUnknownConfigs();
+
+	virtual void Serialize(FArchive& Ar) override;
+#if WITH_EDITOR
+	virtual void PostEditChangeChainProperty(struct FPropertyChangedChainEvent& PropertyChangedEvent) override;
+#endif
+
+private:
+
+	/** used for cleanup */
+	TArray<FName> KnownCategoryNames;
+	TArray<FName> KnownExtensionNames;
+	TMultiMap<FName, FName> KnownCategoryInputNames;
+	TMultiMap<FName, FName> KnownExtensionInputNames;
 };
diff --git a/Engine/Source/Editor/AnimGraph/Classes/AnimGraphNode_AssetPlayerBase.cpp b/Engine/Source/Editor/AnimGraph/Classes/AnimGraphNode_AssetPlayerBase.cpp
--- a/Engine/Source/Editor/AnimGraph/Classes/AnimGraphNode_AssetPlayerBase.cpp
+++ b/Engine/Source/Editor/AnimGraph/Classes/AnimGraphNode_AssetPlayerBase.cpp
@@ -6,9 +6,12 @@
 #include "Animation/AimOffsetBlendSpace1D.h"
 #include "AnimGraphNode_AssetPlayerBase.h"
 #include "AnimGraphNode_SequencePlayer.h"
+#include "AnimGraphNode_SequenceEvaluator.h"
 #include "AnimGraphNode_RotationOffsetBlendSpace.h"
 #include "AnimGraphNode_BlendSpacePlayer.h"
+#include "AnimGraphNode_BlendSpaceEvaluator.h"
 #include "AnimGraphNode_PoseBlendNode.h"
+#include "AnimGraphNode_PoseByName.h"
 
 bool IsAimOffsetBlendSpace(const UClass* BlendSpaceClass)
 {
@@ -44,3 +47,29 @@ UClass* GetNodeClassForAsset(const UClass* AssetClass)
 	return nullptr;
 }
 
+bool SupportNodeClassForAsset(const UClass* AssetClass, const UClass* NodeClass)
+{
+	// we don't want montage to show up, so not checking AnimSequenceBase
+	if (AssetClass->IsChildOf(UAnimSequence::StaticClass()) || AssetClass->IsChildOf(UAnimComposite::StaticClass()))
+	{
+		return (UAnimGraphNode_SequencePlayer::StaticClass() == NodeClass || UAnimGraphNode_SequenceEvaluator::StaticClass() == NodeClass);
+	}
+	else if (AssetClass->IsChildOf(UBlendSpaceBase::StaticClass()))
+	{
+		if (IsAimOffsetBlendSpace(AssetClass))
+		{
+			return (UAnimGraphNode_RotationOffsetBlendSpace::StaticClass() == NodeClass);
+		}
+		else
+		{
+			return (UAnimGraphNode_BlendSpacePlayer::StaticClass() == NodeClass || UAnimGraphNode_BlendSpaceEvaluator::StaticClass() == NodeClass);
+		}
+	}
+	else if (AssetClass->IsChildOf(UPoseAsset::StaticClass()))
+	{
+		return (UAnimGraphNode_PoseBlendNode::StaticClass() == NodeClass || UAnimGraphNode_PoseByName::StaticClass() == NodeClass);
+	}
+	
+	return false;
+}
+
diff --git a/Engine/Source/Editor/AnimGraph/Classes/AnimGraphNode_AssetPlayerBase.h b/Engine/Source/Editor/AnimGraph/Classes/AnimGraphNode_AssetPlayerBase.h
--- a/Engine/Source/Editor/AnimGraph/Classes/AnimGraphNode_AssetPlayerBase.h
+++ b/Engine/Source/Editor/AnimGraph/Classes/AnimGraphNode_AssetPlayerBase.h
@@ -8,6 +8,8 @@
 class UAnimSequenceBase;
 
 ANIMGRAPH_API UClass* GetNodeClassForAsset(const UClass* AssetClass);
+ANIMGRAPH_API bool SupportNodeClassForAsset(const UClass* AssetClass, const UClass* NodeClass);
+
 
 /** Helper / intermediate for asset player graphical nodes */
 UCLASS(Abstract, MinimalAPI)
diff --git a/Engine/Source/Editor/AnimGraph/Classes/AnimGraphNode_BlendSpaceBase.h b/Engine/Source/Editor/AnimGraph/Classes/AnimGraphNode_BlendSpaceBase.h
--- a/Engine/Source/Editor/AnimGraph/Classes/AnimGraphNode_BlendSpaceBase.h
+++ b/Engine/Source/Editor/AnimGraph/Classes/AnimGraphNode_BlendSpaceBase.h
@@ -13,6 +13,7 @@ class UAnimGraphNode_BlendSpaceBase : public UAnimGraphNode_AssetPlayerBase
 
 	// UEdGraphNode interface
 	ANIMGRAPH_API virtual FLinearColor GetNodeTitleColor() const override;
+	ANIMGRAPH_API virtual FText GetMenuCategory() const override;
 	// End of UEdGraphNode interface
 
 	// UAnimGraphNode_Base interface
diff --git a/Engine/Source/Editor/AnimGraph/Classes/AnimGraphNode_OrientationDriver.h b/Engine/Source/Editor/AnimGraph/Classes/AnimGraphNode_OrientationDriver.h
--- a/Engine/Source/Editor/AnimGraph/Classes/AnimGraphNode_OrientationDriver.h
+++ b/Engine/Source/Editor/AnimGraph/Classes/AnimGraphNode_OrientationDriver.h
@@ -41,14 +41,4 @@ protected:
 	virtual FText GetControllerDescription() const override;
 	virtual const FAnimNode_SkeletalControlBase* GetNode() const override { return &Node; }
 	// End of UAnimGraphNode_SkeletalControlBase protected interface
-
-private:
-	/** Called when a pose is selected from asset picker */
-	void OnPoseAssetSelected(const FAssetData& AssetData);
-
-	/** Build pop-up asset picker for choosing PoseAsset to import info from */
-	TSharedRef<SWidget> BuildPoseAssetPicker();
-
-	/** Handle 'import poses from pose asset' button being pressed */
-	FReply ClickedOnImportPosesFromPoseAsset();
 };
diff --git a/Engine/Source/Editor/AnimGraph/Classes/AnimGraphNode_PoseBlendNode.h b/Engine/Source/Editor/AnimGraph/Classes/AnimGraphNode_PoseBlendNode.h
--- a/Engine/Source/Editor/AnimGraph/Classes/AnimGraphNode_PoseBlendNode.h
+++ b/Engine/Source/Editor/AnimGraph/Classes/AnimGraphNode_PoseBlendNode.h
@@ -18,6 +18,7 @@ class UAnimGraphNode_PoseBlendNode : public UAnimGraphNode_AssetPlayerBase
 	// UEdGraphNode interface
 	virtual FText GetNodeTitle(ENodeTitleType::Type TitleType) const override;
 	virtual FText GetTooltipText() const override;
+	virtual FText GetMenuCategory() const override;
 	// End of UEdGraphNode
 
 	// UAnimGraphNode_Base interface
@@ -28,6 +29,8 @@ class UAnimGraphNode_PoseBlendNode : public UAnimGraphNode_AssetPlayerBase
 	virtual UAnimationAsset* GetAnimationAsset() const override;
 	virtual void GetAllAnimationSequencesReferred(TArray<UAnimationAsset*>& AnimationAssets) const override;
 	virtual void ReplaceReferredAnimations(const TMap<UAnimationAsset*, UAnimationAsset*>& AnimAssetReplacementMap) override;
+	virtual void GetMenuActions(FBlueprintActionDatabaseRegistrar& ActionRegistrar) const override;
+	virtual bool IsActionFilteredOut(class FBlueprintActionFilter const& Filter) override;
 	// End of UAnimGraphNode_Base
 
 	// UK2Node interface
@@ -36,6 +39,11 @@ class UAnimGraphNode_PoseBlendNode : public UAnimGraphNode_AssetPlayerBase
 
 	virtual void SetAnimationAsset(UAnimationAsset* Asset) override;
 private:
+	static FText GetTitleGivenAssetInfo(const FText& AssetName);
+
 	/** Constructing FText strings can be costly, so we cache the node's title */
 	FNodeTextCache CachedNodeTitle;
+
+	/** Used for filtering in the Blueprint context menu when the pose asset this node uses is unloaded */
+	FString UnloadedSkeletonName;
 };
diff --git a/Engine/Source/Editor/AnimGraph/Classes/AnimGraphNode_PoseByName.h b/Engine/Source/Editor/AnimGraph/Classes/AnimGraphNode_PoseByName.h
--- a/Engine/Source/Editor/AnimGraph/Classes/AnimGraphNode_PoseByName.h
+++ b/Engine/Source/Editor/AnimGraph/Classes/AnimGraphNode_PoseByName.h
@@ -18,6 +18,7 @@ class UAnimGraphNode_PoseByName : public UAnimGraphNode_AssetPlayerBase
 	// UEdGraphNode interface
 	virtual FText GetNodeTitle(ENodeTitleType::Type TitleType) const override;
 	virtual FText GetTooltipText() const override;
+	virtual FText GetMenuCategory() const override;
 	// End of UEdGraphNode
 
 	// UAnimGraphNode_Base interface
@@ -32,6 +33,7 @@ class UAnimGraphNode_PoseByName : public UAnimGraphNode_AssetPlayerBase
 
 	// UK2Node interface
 	virtual void GetContextMenuActions(const FGraphNodeContextMenuBuilder& Context) const override;
+	virtual void GetMenuActions(FBlueprintActionDatabaseRegistrar& ActionRegistrar) const override;
 	//virtual void GetMenuActions(FBlueprintActionDatabaseRegistrar& ActionRegistrar) const override;
 	// End of UK2Node interface
 
diff --git a/Engine/Source/Editor/AnimGraph/Classes/AnimGraphNode_SubInput.h b/Engine/Source/Editor/AnimGraph/Classes/AnimGraphNode_SubInput.h
new file mode 100644
--- /dev/null
+++ b/Engine/Source/Editor/AnimGraph/Classes/AnimGraphNode_SubInput.h
@@ -0,0 +1,25 @@
+// Copyright 1998-2016 Epic Games, Inc. All Rights Reserved.
+
+#pragma once
+
+#include "AnimGraphNode_Base.h"
+#include "Animation/AnimNode_SubInput.h"
+#include "AnimGraphNode_SubInput.generated.h"
+
+UCLASS(MinimalAPI)
+class UAnimGraphNode_SubInput : public UAnimGraphNode_Base
+{
+	GENERATED_BODY()
+
+	UPROPERTY(EditAnywhere, Category = Settings)
+	FAnimNode_SubInput Node;
+
+	//~ Begin UEdGraphNode Interface.
+	virtual FLinearColor GetNodeTitleColor() const override;
+	virtual FText GetTooltipText() const override;
+	virtual FText GetNodeTitle(ENodeTitleType::Type TitleType) const override;
+	virtual bool CanUserDeleteNode() const override;
+	virtual bool CanDuplicateNode() const override;
+	virtual bool IsActionFilteredOut(class FBlueprintActionFilter const& Filter) override;
+	//~ End UEdGraphNode Interface.
+};
\ No newline at end of file
diff --git a/Engine/Source/Editor/AnimGraph/Classes/AnimGraphNode_SubInstance.h b/Engine/Source/Editor/AnimGraph/Classes/AnimGraphNode_SubInstance.h
new file mode 100644
--- /dev/null
+++ b/Engine/Source/Editor/AnimGraph/Classes/AnimGraphNode_SubInstance.h
@@ -0,0 +1,75 @@
+// Copyright 1998-2016 Epic Games, Inc. All Rights Reserved.
+
+#pragma once
+
+#include "AnimGraphNode_Base.h"
+#include "Animation/AnimNode_SubInstance.h"
+
+#include "AnimGraphNode_SubInstance.generated.h"
+
+
+UCLASS(MinimalAPI)
+class UAnimGraphNode_SubInstance : public UAnimGraphNode_Base
+{
+	GENERATED_BODY()
+
+public:
+
+	UPROPERTY(EditAnywhere, Category = Settings)
+	FAnimNode_SubInstance Node;
+
+	//~ Begin UEdGraphNode Interface.
+	virtual FLinearColor GetNodeTitleColor() const override;
+	virtual FText GetTooltipText() const override;
+	virtual FText GetNodeTitle(ENodeTitleType::Type TitleType) const override;
+	virtual void ValidateAnimNodeDuringCompilation(USkeleton* ForSkeleton, FCompilerResultsLog& MessageLog) override;
+	virtual void ReallocatePinsDuringReconstruction(TArray<UEdGraphPin*>& OldPins) override;
+	virtual void PostEditChangeProperty(FPropertyChangedEvent& PropertyChangedEvent) override;
+	virtual bool HasExternalDependencies(TArray<class UStruct*>* OptionalOutput = NULL) const override;
+	virtual UObject* GetJumpTargetForDoubleClick() const override;
+	//~ End UEdGraphNode Interface.
+
+	// Gets the property on InOwnerInstanceClass that corresponds to InInputPin
+	void ANIMGRAPH_API GetInstancePinProperty(const UClass* InOwnerInstanceClass, UEdGraphPin* InInputPin, UProperty*& OutProperty);
+
+	// Gets the unique name for the property linked to a given pin
+	FString ANIMGRAPH_API GetPinTargetVariableName(const UEdGraphPin* InPin) const;
+
+	// Detail panel customizations
+	virtual void CustomizeDetails(IDetailLayoutBuilder& DetailBuilder) override;
+
+private:
+
+	/** List of property names we know to exist on the target class, so we can detect when
+	 *  Properties are added or removed on reconstruction
+	 */
+	UPROPERTY()
+	TArray<FName> KnownExposableProperties;
+
+	/** Names of properties the user has chosen to expose */
+	UPROPERTY()
+	TArray<FName> ExposedPropertyNames;
+
+	// Searches the instance class for properties that we can expose (public and BP visible)
+	void GetExposableProperties(TArray<UProperty*>& OutExposableProperties) const;
+	// Gets a property's type as FText (for UI)
+	FText GetPropertyTypeText(UProperty* Property);
+	// Given a new class, rebuild the known property list (for tracking class changes and moving pins)
+	void RebuildExposedProperties(UClass* InNewClass);
+
+	// ----- UI CALLBACKS ----- //
+	// If given property exposed on this node
+	ECheckBoxState IsPropertyExposed(FName PropertyName) const;
+	// User chose to expose, or unexpose a property
+	void OnPropertyExposeCheckboxChanged(ECheckBoxState NewState, FName PropertyName);
+	// User changed the instance class
+	void OnInstanceClassChanged(IDetailLayoutBuilder* DetailBuilder);
+
+	// Gets path to the currently selected instance class' blueprint
+	FString GetCurrentInstanceBlueprintPath() const;
+	// Filter callback for blueprints (only accept matching skeletons)
+	bool OnShouldFilterInstanceBlueprint(const FAssetData& AssetData) const;
+	// Instance blueprint was changed by user
+	void OnSetInstanceBlueprint(const FAssetData& AssetData, TSharedRef<IPropertyHandle> InstanceClassPropHandle);
+	// ----- END UI CALLBACKS ----- //
+};
\ No newline at end of file
diff --git a/Engine/Source/Editor/AnimGraph/Classes/AnimPreviewInstance.h b/Engine/Source/Editor/AnimGraph/Classes/AnimPreviewInstance.h
--- a/Engine/Source/Editor/AnimGraph/Classes/AnimPreviewInstance.h
+++ b/Engine/Source/Editor/AnimGraph/Classes/AnimPreviewInstance.h
@@ -267,6 +267,12 @@ class ANIMGRAPH_API UAnimPreviewInstance : public UAnimSingleNodeInstance
 	 * This is used by when editing, when controller has to be disabled
 	 */
 	void EnableControllers(bool bEnable);
+
+	/** Used by GenSpaceBases to set up the anim instance proxy, calling InitializeObjects */
+	void PrePerformAnimationEvaluation();
+
+	/** Used by GenSpaceBases to clear the anim instance proxy, calling ClearObjects */
+	void PostPerformAnimationEvaluation();
 };
 
 
diff --git a/Engine/Source/Editor/AnimGraph/Private/AnimGraphNode_BlendSpaceBase.cpp b/Engine/Source/Editor/AnimGraph/Private/AnimGraphNode_BlendSpaceBase.cpp
--- a/Engine/Source/Editor/AnimGraph/Private/AnimGraphNode_BlendSpaceBase.cpp
+++ b/Engine/Source/Editor/AnimGraph/Private/AnimGraphNode_BlendSpaceBase.cpp
@@ -113,4 +113,8 @@ void UAnimGraphNode_BlendSpaceBase::PostProcessPinName(const UEdGraphPin* Pin, F
 	Super::PostProcessPinName(Pin, DisplayName);
 }
 
+FText UAnimGraphNode_BlendSpaceBase::GetMenuCategory() const
+{
+	return LOCTEXT("BlendSpaceCategory_Label", "BlendSpaces");
+}
 #undef LOCTEXT_NAMESPACE
\ No newline at end of file
diff --git a/Engine/Source/Editor/AnimGraph/Private/AnimGraphNode_OrientationDriver.cpp b/Engine/Source/Editor/AnimGraph/Private/AnimGraphNode_OrientationDriver.cpp
--- a/Engine/Source/Editor/AnimGraph/Private/AnimGraphNode_OrientationDriver.cpp
+++ b/Engine/Source/Editor/AnimGraph/Private/AnimGraphNode_OrientationDriver.cpp
@@ -58,11 +58,10 @@ void UAnimGraphNode_OrientationDriver::DrawCanvas(FViewport& InViewport, FSceneV
 	Canvas.DrawItem(InterpItem);
 	PosY += YDelta;
 
-	// List each pose with weight
 	int32 PoseIndex = 0;
-	for (const FOrientationDriverPose& Pose : DebugNode->Poses)
+	for (const FOrientationDriverPoseInfo& PoseInfo : DebugNode->PoseInfos)
 	{
-		FString PoseItemString = FString::Printf(TEXT("%d : %f"), PoseIndex, Pose.PoseWeight);
+		FString PoseItemString = FString::Printf(TEXT("%s : %f"), *PoseInfo.PoseName.ToString(), PoseInfo.PoseWeight);
 		FCanvasTextItem PoseItem(FVector2D(MarginX + XDelta, PosY), FText::FromString(PoseItemString), GEngine->GetSmallFont(), FLinearColor::White);
 		Canvas.DrawItem(PoseItem);
 		PosY += YDelta;
@@ -80,7 +79,7 @@ void UAnimGraphNode_OrientationDriver::DrawCanvas(FViewport& InViewport, FSceneV
 	// Iterate over Parameters
 	for (const FOrientationDriverParam& Param : DebugNode->ResultParamSet.Params)
 	{
-		FString ParamItemString = FString::Printf(TEXT("%s : %f"), *Param.ParameterName.ToString(), Param.ParameterValue);
+		FString ParamItemString = FString::Printf(TEXT("%s : %f"), *Param.ParamInfo.Name.ToString(), Param.ParamValue);
 		FCanvasTextItem ParamItem(FVector2D(MarginX + XDelta, PosY), FText::FromString(ParamItemString), GEngine->GetSmallFont(), FLinearColor::White);
 		Canvas.DrawItem(ParamItem);
 		PosY += YDelta;
@@ -98,50 +97,9 @@ void UAnimGraphNode_OrientationDriver::ValidateAnimNodeDuringCompilation(USkelet
 	Super::ValidateAnimNodeDuringCompilation(ForSkeleton, MessageLog);
 }
 
-
-void UAnimGraphNode_OrientationDriver::OnPoseAssetSelected(const FAssetData& AssetData)
-{
-
-}
-
-TSharedRef<SWidget> UAnimGraphNode_OrientationDriver::BuildPoseAssetPicker()
-{
-	FContentBrowserModule& ContentBrowserModule = FModuleManager::Get().LoadModuleChecked<FContentBrowserModule>(TEXT("ContentBrowser"));
-
-	FAssetPickerConfig AssetPickerConfig;
-	AssetPickerConfig.Filter.ClassNames.Add(UPoseAsset::StaticClass()->GetFName());
-	AssetPickerConfig.OnAssetSelected = FOnAssetSelected::CreateUObject(this, &UAnimGraphNode_OrientationDriver::OnPoseAssetSelected);
-	AssetPickerConfig.bAllowNullSelection = true;
-	AssetPickerConfig.InitialAssetViewType = EAssetViewType::List;
-	AssetPickerConfig.bFocusSearchBoxWhenOpened = true;
-	AssetPickerConfig.bShowBottomToolbar = false;
-	AssetPickerConfig.SelectionMode = ESelectionMode::Single;
-
-	return 
-		SNew(SBox)
-		.WidthOverride(384)
-		.HeightOverride(768)
-		[
-			ContentBrowserModule.Get().CreateAssetPicker(AssetPickerConfig)
-		];
-}
-
 void UAnimGraphNode_OrientationDriver::CustomizeDetails(IDetailLayoutBuilder& DetailBuilder)
 {
-	IDetailCategoryBuilder& ImportCategory = DetailBuilder.EditCategory("Import Driver Setup");
-
-	ImportCategory.AddCustomRow(LOCTEXT("ImportRow", "Import"))
-	[
-		SNew(SComboButton)
-		.ContentPadding(3)
-		.OnGetMenuContent(FOnGetContent::CreateUObject(this, &UAnimGraphNode_OrientationDriver::BuildPoseAssetPicker))
-		.ButtonContent()
-		[
-			SNew(STextBlock)
-			.Text(LOCTEXT("ImportPosesFromPoseAsset ", "Import Poses From PoseAsset"))
-			.ToolTipText(LOCTEXT("ImportPosesFromPoseAssetTooltip", "Import target poses from a PoseAsset"))
-		]
-	];
+
 }
 
 
diff --git a/Engine/Source/Editor/AnimGraph/Private/AnimGraphNode_PoseBlendNode.cpp b/Engine/Source/Editor/AnimGraph/Private/AnimGraphNode_PoseBlendNode.cpp
--- a/Engine/Source/Editor/AnimGraph/Private/AnimGraphNode_PoseBlendNode.cpp
+++ b/Engine/Source/Editor/AnimGraph/Private/AnimGraphNode_PoseBlendNode.cpp
@@ -4,13 +4,41 @@
 
 #include "CompilerResultsLog.h"
 #include "GraphEditorActions.h"
+#include "EditorCategoryUtils.h"
+#include "AssetRegistryModule.h"
+#include "BlueprintActionDatabaseRegistrar.h"
+#include "BlueprintActionFilter.h"
+#include "BlueprintNodeSpawner.h"
 #include "AnimGraphNode_PoseBlendNode.h"
 
+#define LOCTEXT_NAMESPACE "PoseBlendNode"
+
+// Action to add a pose asset blend node to the graph
+struct FNewPoseBlendNodeAction : public FEdGraphSchemaAction_K2NewNode
+{
+protected:
+	FAssetData AssetInfo;
+public:
+	FNewPoseBlendNodeAction(const FAssetData& InAssetInfo, FText Title)
+		: FEdGraphSchemaAction_K2NewNode(LOCTEXT("PoseAsset", "PosetAssets"), Title, TEXT("Evaluates curves to produce a pose from pose asset"), 0, FText::FromName(InAssetInfo.ObjectPath))
+	{
+		AssetInfo = InAssetInfo;
+
+		UAnimGraphNode_PoseBlendNode* Template = NewObject<UAnimGraphNode_PoseBlendNode>();
+		NodeTemplate = Template;
+	}
+
+	virtual UEdGraphNode* PerformAction(class UEdGraph* ParentGraph, UEdGraphPin* FromPin, const FVector2D Location, bool bSelectNewNode = true) override
+	{
+		UAnimGraphNode_PoseBlendNode* SpawnedNode = CastChecked<UAnimGraphNode_PoseBlendNode>(FEdGraphSchemaAction_K2NewNode::PerformAction(ParentGraph, FromPin, Location, bSelectNewNode));
+		SpawnedNode->Node.PoseAsset = Cast<UPoseAsset>(AssetInfo.GetAsset());
+
+		return SpawnedNode;
+	}
+};
 /////////////////////////////////////////////////////
 // UAnimGraphNode_PoseBlendNode
 
-#define LOCTEXT_NAMESPACE "PoseBlendNode"
-
 UAnimGraphNode_PoseBlendNode::UAnimGraphNode_PoseBlendNode(const FObjectInitializer& ObjectInitializer)
 	: Super(ObjectInitializer)
 {
@@ -64,10 +92,147 @@ FText UAnimGraphNode_PoseBlendNode::GetNodeTitle(ENodeTitleType::Type TitleType)
 	return CachedNodeTitle;
 }
 
-// void UAnimGraphNode_PoseBlendNode::GetMenuActions(FBlueprintActionDatabaseRegistrar& ActionRegistrar) const
-// {
-// 	// Intentionally empty; you can drop down a regular sequence player and convert into a sequence evaluator in the right-click menu.
-// }
+void UAnimGraphNode_PoseBlendNode::GetMenuActions(FBlueprintActionDatabaseRegistrar& ActionRegistrar) const
+{
+	auto LoadedAssetSetup = [](UEdGraphNode* NewNode, bool /*bIsTemplateNode*/, TWeakObjectPtr<UPoseAsset> PoseAssetPtr)
+	{
+		UAnimGraphNode_PoseBlendNode* PoseBlendNodeNode = CastChecked<UAnimGraphNode_PoseBlendNode>(NewNode);
+		PoseBlendNodeNode->Node.PoseAsset = PoseAssetPtr.Get();
+	};
+
+	auto UnloadedAssetSetup = [](UEdGraphNode* NewNode, bool bIsTemplateNode, const FAssetData AssetData)
+	{
+		UAnimGraphNode_PoseBlendNode* PoseBlendNodeNode = CastChecked<UAnimGraphNode_PoseBlendNode>(NewNode);
+		if (bIsTemplateNode)
+		{
+			AssetData.GetTagValue("Skeleton", PoseBlendNodeNode->UnloadedSkeletonName);
+		}
+		else
+		{
+			UPoseAsset* PoseAsset = Cast<UPoseAsset>(AssetData.GetAsset());
+			check(PoseAsset != nullptr);
+			PoseBlendNodeNode->Node.PoseAsset = PoseAsset;
+		}
+	};
+
+	const UObject* QueryObject = ActionRegistrar.GetActionKeyFilter();
+	if (QueryObject == nullptr)
+	{
+		FAssetRegistryModule& AssetRegistryModule = FModuleManager::LoadModuleChecked<FAssetRegistryModule>(TEXT("AssetRegistry"));
+		// define a filter to help in pulling UPoseAsset asset data from the registry
+		FARFilter Filter;
+		Filter.ClassNames.Add(UPoseAsset::StaticClass()->GetFName());
+		Filter.bRecursiveClasses = true;
+		// Find matching assets and add an entry for each one
+		TArray<FAssetData> PoseAssetList;
+		AssetRegistryModule.Get().GetAssets(Filter, /*out*/PoseAssetList);
+
+		for (auto AssetIt = PoseAssetList.CreateConstIterator(); AssetIt; ++AssetIt)
+		{
+			const FAssetData& Asset = *AssetIt;
+
+			UBlueprintNodeSpawner* NodeSpawner = UBlueprintNodeSpawner::Create(GetClass());
+			if (Asset.IsAssetLoaded())
+			{
+				TWeakObjectPtr<UPoseAsset> PoseAsset = Cast<UPoseAsset>(Asset.GetAsset());
+				NodeSpawner->CustomizeNodeDelegate = UBlueprintNodeSpawner::FCustomizeNodeDelegate::CreateStatic(LoadedAssetSetup, PoseAsset);
+				NodeSpawner->DefaultMenuSignature.MenuName = GetTitleGivenAssetInfo(FText::FromName(PoseAsset->GetFName()));
+			}
+			else
+			{
+				NodeSpawner->CustomizeNodeDelegate = UBlueprintNodeSpawner::FCustomizeNodeDelegate::CreateStatic(UnloadedAssetSetup, Asset);
+				NodeSpawner->DefaultMenuSignature.MenuName = GetTitleGivenAssetInfo(FText::FromName(Asset.AssetName));
+			}
+			ActionRegistrar.AddBlueprintAction(Asset, NodeSpawner);
+		}
+	}
+	else if (const UPoseAsset* PoseAsset = Cast<UPoseAsset>(QueryObject))
+	{
+		UBlueprintNodeSpawner* NodeSpawner = UBlueprintNodeSpawner::Create(GetClass());
+
+		TWeakObjectPtr<UPoseAsset> PoseAssetPtr = PoseAsset;
+		NodeSpawner->CustomizeNodeDelegate = UBlueprintNodeSpawner::FCustomizeNodeDelegate::CreateStatic(LoadedAssetSetup, PoseAssetPtr);
+		NodeSpawner->DefaultMenuSignature.MenuName = GetTitleGivenAssetInfo(FText::FromName(PoseAsset->GetFName()));
+
+		ActionRegistrar.AddBlueprintAction(QueryObject, NodeSpawner);
+	}
+	else if (QueryObject == GetClass())
+	{
+		FAssetRegistryModule& AssetRegistryModule = FModuleManager::LoadModuleChecked<FAssetRegistryModule>(TEXT("AssetRegistry"));
+		// define a filter to help in pulling UPoseAsset asset data from the registry
+		FARFilter Filter;
+		Filter.ClassNames.Add(UPoseAsset::StaticClass()->GetFName());
+		Filter.bRecursiveClasses = true;
+		// Find matching assets and add an entry for each one
+		TArray<FAssetData> PoseAssetList;
+		AssetRegistryModule.Get().GetAssets(Filter, /*out*/PoseAssetList);
+
+		for (auto AssetIt = PoseAssetList.CreateConstIterator(); AssetIt; ++AssetIt)
+		{
+			const FAssetData& Asset = *AssetIt;
+			if (Asset.IsAssetLoaded())
+			{
+				continue;
+			}
+
+			UBlueprintNodeSpawner* NodeSpawner = UBlueprintNodeSpawner::Create(GetClass());
+			NodeSpawner->CustomizeNodeDelegate = UBlueprintNodeSpawner::FCustomizeNodeDelegate::CreateStatic(UnloadedAssetSetup, Asset);
+			NodeSpawner->DefaultMenuSignature.MenuName = GetTitleGivenAssetInfo(FText::FromName(Asset.AssetName));
+			ActionRegistrar.AddBlueprintAction(Asset, NodeSpawner);
+		}
+	}
+}
+
+FText UAnimGraphNode_PoseBlendNode::GetTitleGivenAssetInfo(const FText& AssetName)
+{
+	FFormatNamedArguments Args;
+	Args.Add(TEXT("AssetName"), AssetName);
+
+	return FText::Format(LOCTEXT("PoseAssetNodeTitle", "Evaluate Pose {AssetName}"), Args);
+}
+
+bool UAnimGraphNode_PoseBlendNode::IsActionFilteredOut(class FBlueprintActionFilter const& Filter)
+{
+	bool bIsFilteredOut = false;
+	FBlueprintActionContext const& FilterContext = Filter.Context;
+
+	for (UBlueprint* Blueprint : FilterContext.Blueprints)
+	{
+		if (UAnimBlueprint* AnimBlueprint = Cast<UAnimBlueprint>(Blueprint))
+		{
+			if (Node.PoseAsset)
+			{
+				if (Node.PoseAsset->GetSkeleton() != AnimBlueprint->TargetSkeleton)
+				{
+					// PoseAsset does not use the same skeleton as the Blueprint, cannot use
+					bIsFilteredOut = true;
+					break;
+				}
+			}
+			else
+			{
+				FAssetData SkeletonData(AnimBlueprint->TargetSkeleton);
+				if (UnloadedSkeletonName != SkeletonData.GetExportTextName())
+				{
+					bIsFilteredOut = true;
+					break;
+				}
+			}
+		}
+		else
+		{
+			// Not an animation Blueprint, cannot use
+			bIsFilteredOut = true;
+			break;
+		}
+	}
+	return bIsFilteredOut;
+}
+
+FText UAnimGraphNode_PoseBlendNode::GetMenuCategory() const
+{
+	return LOCTEXT("PoseAssetCategory_Label", "Poses");
+}
 
 void UAnimGraphNode_PoseBlendNode::SetAnimationAsset(UAnimationAsset* Asset)
 {
@@ -86,7 +251,7 @@ void UAnimGraphNode_PoseBlendNode::ValidateAnimNodeDuringCompilation(class USkel
 	else
 	{
 		USkeleton* SeqSkeleton = Node.PoseAsset->GetSkeleton();
-		if (SeqSkeleton&& // if anim sequence doesn't have skeleton, it might be due to anim sequence not loaded yet, @todo: wait with anim blueprint compilation until all assets are loaded?
+		if (SeqSkeleton&& // if PoseAsset doesn't have skeleton, it might be due to PoseAsset not loaded yet, @todo: wait with anim blueprint compilation until all assets are loaded?
 			!SeqSkeleton->IsCompatible(ForSkeleton))
 		{
 			MessageLog.Error(TEXT("@@ references poseasset that uses different skeleton @@"), this, SeqSkeleton);
@@ -94,21 +259,6 @@ void UAnimGraphNode_PoseBlendNode::ValidateAnimNodeDuringCompilation(class USkel
 	}
 }
 
-// potentially in the future, give options for the pose name?
-// void UAnimGraphNode_PoseBlendNode::GetContextMenuActions(const FGraphNodeContextMenuBuilder& Context) const
-// {
-// 	if (!Context.bIsDebugging)
-// 	{
-// 		// add an option to convert to a regular sequence player
-// 		Context.MenuBuilder->BeginSection("AnimGraphNodePoseByName", NSLOCTEXT("A3Nodes", "PoseByNameHeading", "Sequence Evaluator"));
-// 		{
-// 			Context.MenuBuilder->AddMenuEntry(FGraphEditorCommands::Get().OpenRelatedAsset);
-// 			Context.MenuBuilder->AddMenuEntry(FGraphEditorCommands::Get().ConvertToSeqPlayer);
-// 		}
-// 		Context.MenuBuilder->EndSection();
-// 	}
-// }
-
 bool UAnimGraphNode_PoseBlendNode::DoesSupportTimeForTransitionGetter() const
 {
 	return false;
diff --git a/Engine/Source/Editor/AnimGraph/Private/AnimGraphNode_PoseByName.cpp b/Engine/Source/Editor/AnimGraph/Private/AnimGraphNode_PoseByName.cpp
--- a/Engine/Source/Editor/AnimGraph/Private/AnimGraphNode_PoseByName.cpp
+++ b/Engine/Source/Editor/AnimGraph/Private/AnimGraphNode_PoseByName.cpp
@@ -4,6 +4,7 @@
 
 #include "CompilerResultsLog.h"
 #include "GraphEditorActions.h"
+#include "EditorCategoryUtils.h"
 #include "AnimGraphNode_PoseByName.h"
 
 /////////////////////////////////////////////////////
@@ -43,6 +44,11 @@ FText UAnimGraphNode_PoseByName::GetTooltipText() const
 	return GetNodeTitle(ENodeTitleType::ListView);
 }
 
+FText UAnimGraphNode_PoseByName::GetMenuCategory() const
+{
+	return LOCTEXT("PoseAssetCategory_Label", "Poses");
+}
+
 FText UAnimGraphNode_PoseByName::GetNodeTitle(ENodeTitleType::Type TitleType) const
 {
 	if (Node.PoseAsset == nullptr)
@@ -95,21 +101,6 @@ void UAnimGraphNode_PoseByName::ValidateAnimNodeDuringCompilation(class USkeleto
 	}
 }
 
-// potentially in the future, give options for the pose name?
-// void UAnimGraphNode_PoseByName::GetContextMenuActions(const FGraphNodeContextMenuBuilder& Context) const
-// {
-// 	if (!Context.bIsDebugging)
-// 	{
-// 		// add an option to convert to a regular sequence player
-// 		Context.MenuBuilder->BeginSection("AnimGraphNodePoseByName", NSLOCTEXT("A3Nodes", "PoseByNameHeading", "Sequence Evaluator"));
-// 		{
-// 			Context.MenuBuilder->AddMenuEntry(FGraphEditorCommands::Get().OpenRelatedAsset);
-// 			Context.MenuBuilder->AddMenuEntry(FGraphEditorCommands::Get().ConvertToSeqPlayer);
-// 		}
-// 		Context.MenuBuilder->EndSection();
-// 	}
-// }
-
 bool UAnimGraphNode_PoseByName::DoesSupportTimeForTransitionGetter() const
 {
 	return false;
@@ -132,4 +123,9 @@ void UAnimGraphNode_PoseByName::GetContextMenuActions(const FGraphNodeContextMen
 		Context.MenuBuilder->EndSection();
 	}
 }
+
+void UAnimGraphNode_PoseByName::GetMenuActions(FBlueprintActionDatabaseRegistrar& ActionRegistrar) const
+{
+	// Intentionally empty; you can drop down a regular pose blend node and convert into a poseasset by name in the right-click menu.
+}
 #undef LOCTEXT_NAMESPACE
\ No newline at end of file
diff --git a/Engine/Source/Editor/AnimGraph/Private/AnimGraphNode_SubInput.cpp b/Engine/Source/Editor/AnimGraph/Private/AnimGraphNode_SubInput.cpp
new file mode 100644
--- /dev/null
+++ b/Engine/Source/Editor/AnimGraph/Private/AnimGraphNode_SubInput.cpp
@@ -0,0 +1,60 @@
+// Copyright 1998-2016 Epic Games, Inc. All Rights Reserved.
+
+#include "AnimGraphPrivatePCH.h"
+#include "AnimGraphNode_SubInput.h"
+#include "GraphEditorSettings.h"
+#include "AnimGraphNode_SubInstance.h"
+#include "BlueprintActionDatabaseRegistrar.h"
+#include "BlueprintNodeSpawner.h"
+#include "BlueprintActionFilter.h"
+
+#define LOCTEXT_NAMESPACE "SubInputNode"
+
+FLinearColor UAnimGraphNode_SubInput::GetNodeTitleColor() const
+{
+	return GetDefault<UGraphEditorSettings>()->ResultNodeTitleColor;
+}
+
+FText UAnimGraphNode_SubInput::GetTooltipText() const
+{
+	return LOCTEXT("ToolTip", "Inputs to a sub-animation graph from a parent instance.");
+}
+
+FText UAnimGraphNode_SubInput::GetNodeTitle(ENodeTitleType::Type TitleType) const
+{
+	return LOCTEXT("Title", "Sub-Graph Input");
+}
+
+bool UAnimGraphNode_SubInput::CanUserDeleteNode() const
+{
+	return true;
+}
+
+bool UAnimGraphNode_SubInput::CanDuplicateNode() const
+{
+	return false;
+}
+
+bool UAnimGraphNode_SubInput::IsActionFilteredOut(class FBlueprintActionFilter const& Filter)
+{
+	bool bFilterOut = false;
+
+	TArray<UAnimGraphNode_SubInput*> SubInputNodes;
+
+	for(UBlueprint* Blueprint : Filter.Context.Blueprints)
+	{
+		SubInputNodes.Empty(SubInputNodes.Num());
+
+		FBlueprintEditorUtils::GetAllNodesOfClass(Blueprint, SubInputNodes);
+
+		if(SubInputNodes.Num() > 0)
+		{
+			bFilterOut = true;
+			break;
+		}
+	}
+
+	return bFilterOut;
+}
+
+#undef LOCTEXT_NAMESPACE
diff --git a/Engine/Source/Editor/AnimGraph/Private/AnimGraphNode_SubInstance.cpp b/Engine/Source/Editor/AnimGraph/Private/AnimGraphNode_SubInstance.cpp
new file mode 100644
--- /dev/null
+++ b/Engine/Source/Editor/AnimGraph/Private/AnimGraphNode_SubInstance.cpp
@@ -0,0 +1,464 @@
+// Copyright 1998-2016 Epic Games, Inc. All Rights Reserved.
+
+#include "AnimGraphPrivatePCH.h"
+#include "AnimGraphNode_SubInstance.h"
+#include "CompilerResultsLog.h"
+
+#include "PropertyEditing.h"
+#include "IPropertyUtilities.h"
+#include "Animation/AnimNode_SubInput.h"
+#include "PropertyCustomizationHelpers.h"
+#include "ScopedTransaction.h"
+#include "Animation/AnimInstance.h"
+
+#define LOCTEXT_NAMESPACE "SubInstanceNode"
+
+namespace SubInstanceGraphNodeConstants
+{
+	FLinearColor TitleColor(0.2f, 0.2f, 0.8f);
+}
+
+FLinearColor UAnimGraphNode_SubInstance::GetNodeTitleColor() const
+{
+	return SubInstanceGraphNodeConstants::TitleColor;
+}
+
+FText UAnimGraphNode_SubInstance::GetTooltipText() const
+{
+	return LOCTEXT("ToolTip", "Runs a sub-anim instance to process animation");
+}
+FText UAnimGraphNode_SubInstance::GetNodeTitle(ENodeTitleType::Type TitleType) const
+{
+	UClass* TargetClass = *Node.InstanceClass;
+
+	FFormatNamedArguments Args;
+	Args.Add(TEXT("NodeTitle"), LOCTEXT("Title", "Sub Anim Instance"));
+	Args.Add(TEXT("TargetClass"), TargetClass ? FText::FromString(TargetClass->GetName()) : LOCTEXT("ClassNone", "None"));
+
+	if(TitleType == ENodeTitleType::MenuTitle)
+	{
+		return LOCTEXT("NodeTitle", "Sub Anim Instance");
+	}
+	if(TitleType == ENodeTitleType::ListView)
+	{
+		return FText::Format(LOCTEXT("TitleFormat", "{NodeTitle} - Target Class: {TargetClass}"), Args);
+	}
+	else
+	{
+		return FText::Format(LOCTEXT("TitleFormat", "{NodeTitle}\nTarget Class: {TargetClass}"), Args);
+	}
+}
+
+void UAnimGraphNode_SubInstance::ValidateAnimNodeDuringCompilation(USkeleton* ForSkeleton, FCompilerResultsLog& MessageLog)
+{
+	UAnimBlueprint* AnimBP = CastChecked<UAnimBlueprint>(GetBlueprint());
+
+	UObject* OriginalNode = MessageLog.FindSourceObject(this);
+
+	// Check we have a class set
+	if(!*Node.InstanceClass)
+	{
+		MessageLog.Error(TEXT("Sub instance node @@ has no valid instance class to spawn."), this);
+	}
+
+	// Check for cycles from other sub instance nodes
+	TArray<UEdGraph*> Graphs;
+	AnimBP->GetAllGraphs(Graphs);
+
+	for(UEdGraph* Graph : Graphs)
+	{
+		TArray<UAnimGraphNode_SubInstance*> SubInstanceNodes;
+		Graph->GetNodesOfClass(SubInstanceNodes);
+
+		for(UAnimGraphNode_SubInstance* SubInstanceNode : SubInstanceNodes)
+		{
+			if(SubInstanceNode == OriginalNode)
+			{
+				continue;
+			}
+
+			FAnimNode_SubInstance& InnerNode = SubInstanceNode->Node;
+
+			if(*InnerNode.InstanceClass && *InnerNode.InstanceClass == *Node.InstanceClass)
+			{
+				MessageLog.Error(TEXT("Node @@ and node @@ both target the same class @@, causing a sub instance loop."), this, SubInstanceNode, *Node.InstanceClass);
+			}
+		}
+	}
+
+	// Check we don't try to spawn our own blueprint
+	if(*Node.InstanceClass == AnimBP->GetAnimBlueprintGeneratedClass())
+	{
+		MessageLog.Error(TEXT("Sub instance node @@ targets instance class @@ which it is inside, this would cause a loop."), this, AnimBP->GetAnimBlueprintGeneratedClass());
+	}
+}
+
+void UAnimGraphNode_SubInstance::ReallocatePinsDuringReconstruction(TArray<UEdGraphPin*>& OldPins)
+{
+	Super::ReallocatePinsDuringReconstruction(OldPins);
+
+	UClass* TargetClass = *Node.InstanceClass;
+
+	if(!TargetClass)
+	{
+		// Nothing to search for properties
+		return;
+	}
+
+	// Need the schema to extract pin types
+	const UEdGraphSchema_K2* Schema = CastChecked<UEdGraphSchema_K2>(GetSchema());
+
+	bool bShowPose = false;
+
+	// Scan the target class for a sub input node, we only want to show the pose input if
+	// we have that node available
+	for(TFieldIterator<UProperty> It(TargetClass); It; ++It)
+	{
+		UProperty* CurrentProp = *It;
+		
+		if(UStructProperty* StructProp = Cast<UStructProperty>(CurrentProp))
+		{
+			if(StructProp->Struct->IsChildOf(FAnimNode_SubInput::StaticStruct()))
+			{
+				// Found a pose input
+				bShowPose = true;
+				break;
+			}
+		}
+	}
+
+	if(bShowPose)
+	{
+		if(UProperty* PoseProperty = FindField<UProperty>(FAnimNode_SubInstance::StaticStruct(), GET_MEMBER_NAME_CHECKED(FAnimNode_SubInstance, InPose)))
+		{
+			FEdGraphPinType PinType;
+			if(Schema->ConvertPropertyToPinType(PoseProperty, PinType))
+			{
+				UEdGraphPin* NewPin = CreatePin(EEdGraphPinDirection::EGPD_Input, PinType, PoseProperty->GetName());
+				NewPin->PinFriendlyName = PoseProperty->GetDisplayNameText();
+
+				CustomizePinData(NewPin, PoseProperty->GetFName(), INDEX_NONE);
+			}
+		}
+	}
+
+	// Grab the list of properties we can expose
+	TArray<UProperty*> ExposablePropeties;
+	GetExposableProperties(ExposablePropeties);
+
+	// We'll track the names we encounter by removing from this list, if anything remains the properties
+	// have been removed from the target class and we should remove them too
+	TArray<FName> BeginExposableNames = KnownExposableProperties;
+
+	for(UProperty* Property : ExposablePropeties)
+	{
+		FName PropertyName = Property->GetFName();
+		BeginExposableNames.Remove(PropertyName);
+
+		if(!KnownExposableProperties.Contains(PropertyName))
+		{
+			// New property added to the target class
+			KnownExposableProperties.Add(PropertyName);
+		}
+
+		if(ExposedPropertyNames.Contains(PropertyName))
+		{
+			FEdGraphPinType PinType;
+
+			verify(Schema->ConvertPropertyToPinType(Property, PinType));
+
+			UEdGraphPin* NewPin = CreatePin(EEdGraphPinDirection::EGPD_Input, PinType, Property->GetName());
+			NewPin->PinFriendlyName = Property->GetDisplayNameText();
+
+			CustomizePinData(NewPin, PropertyName, INDEX_NONE);
+		}
+	}
+
+	// Remove any properties that no longer exist on the target class
+	for(FName& RemovedPropertyName : BeginExposableNames)
+	{
+		KnownExposableProperties.Remove(RemovedPropertyName);
+		ExposedPropertyNames.Remove(RemovedPropertyName);
+	}
+
+}
+
+void UAnimGraphNode_SubInstance::PostEditChangeProperty(FPropertyChangedEvent& PropertyChangedEvent)
+{
+	Super::PostEditChangeProperty(PropertyChangedEvent);
+
+	bool bRequiresNodeReconstruct = false;
+	UProperty* ChangedProperty = PropertyChangedEvent.Property;
+
+	if(ChangedProperty)
+	{
+		if(ChangedProperty->GetFName() == GET_MEMBER_NAME_CHECKED(FAnimNode_SubInstance, InstanceClass))
+		{
+			bRequiresNodeReconstruct = true;
+			RebuildExposedProperties(*Node.InstanceClass);
+		}
+	}
+
+	if(bRequiresNodeReconstruct)
+	{
+		ReconstructNode();
+	}
+}
+
+void UAnimGraphNode_SubInstance::GetInstancePinProperty(const UClass* InOwnerInstanceClass, UEdGraphPin* InInputPin, UProperty*& OutProperty)
+{
+	// The actual name of the instance property
+	FString FullName = GetPinTargetVariableName(InInputPin);
+
+	if(UProperty* Property = FindField<UProperty>(InOwnerInstanceClass, *FullName))
+	{
+		OutProperty = Property;
+	}
+	else
+	{
+		OutProperty = nullptr;
+	}
+}
+
+FString UAnimGraphNode_SubInstance::GetPinTargetVariableName(const UEdGraphPin* InPin) const
+{
+	return TEXT("__SUBINSTANCE_") + InPin->PinName + TEXT("_") + NodeGuid.ToString();
+}
+
+void UAnimGraphNode_SubInstance::CustomizeDetails(IDetailLayoutBuilder& DetailBuilder)
+{
+	Super::CustomizeDetails(DetailBuilder);
+
+	TArray<UProperty*> ExposableProperties;
+	GetExposableProperties(ExposableProperties);
+
+	if(ExposableProperties.Num() > 0)
+	{
+		IDetailCategoryBuilder& CategoryBuilder = DetailBuilder.EditCategory(FName(TEXT("Sub Instance Properties")));
+
+		for(UProperty* Property : ExposableProperties)
+		{
+			FDetailWidgetRow& WidgetRow = CategoryBuilder.AddCustomRow(FText::FromString(Property->GetName()));
+
+			FName PropertyName = Property->GetFName();
+			FText PropertyTypeText = GetPropertyTypeText(Property);
+
+			FFormatNamedArguments Args;
+			Args.Add(TEXT("PropertyName"), FText::FromName(PropertyName));
+			Args.Add(TEXT("PropertyType"), PropertyTypeText);
+
+			FText TooltipText = FText::Format(LOCTEXT("PropertyTooltipText", "{PropertyName}\nType: {PropertyType}"), Args);
+
+			WidgetRow.NameContent()
+			[
+				SNew(STextBlock)
+				.Text(FText::FromString(Property->GetName()))
+				.ToolTipText(TooltipText)
+			];
+
+			WidgetRow.ValueContent()
+			[
+				SNew(SHorizontalBox)
+				+SHorizontalBox::Slot()
+				[
+					SNew(STextBlock)
+					.Text(LOCTEXT("ExposePropertyValue", "Expose: "))
+				]
+				+SHorizontalBox::Slot()
+				[
+					SNew(SCheckBox)
+					.IsChecked_UObject(this, &UAnimGraphNode_SubInstance::IsPropertyExposed, PropertyName)
+					.OnCheckStateChanged_UObject(this, &UAnimGraphNode_SubInstance::OnPropertyExposeCheckboxChanged, PropertyName)
+				]
+			];
+		}
+	}
+
+	TSharedRef<IPropertyHandle> ClassHandle = DetailBuilder.GetProperty(TEXT("Node.InstanceClass"), GetClass());
+	if(ClassHandle->IsValidHandle())
+	{
+		ClassHandle->SetOnPropertyValueChanged(FSimpleDelegate::CreateUObject(this, &UAnimGraphNode_SubInstance::OnInstanceClassChanged, &DetailBuilder));
+	}
+
+	ClassHandle->MarkHiddenByCustomization();
+
+	IDetailCategoryBuilder& CategoryBuilder = DetailBuilder.EditCategory(FName(TEXT("Settings")));
+
+	FDetailWidgetRow& ClassWidgetRow = CategoryBuilder.AddCustomRow(LOCTEXT("FilterString", "Instance Class"));
+	ClassWidgetRow.NameContent()
+		[
+			ClassHandle->CreatePropertyNameWidget()
+		];
+
+	ClassWidgetRow.ValueContent()
+		[
+			SNew(SObjectPropertyEntryBox)
+				.ObjectPath_UObject(this, &UAnimGraphNode_SubInstance::GetCurrentInstanceBlueprintPath)
+				.AllowedClass(UAnimBlueprint::StaticClass())
+				.NewAssetFactories(TArray<UFactory*>())
+				.OnShouldFilterAsset(FOnShouldFilterAsset::CreateUObject(this, &UAnimGraphNode_SubInstance::OnShouldFilterInstanceBlueprint))
+				.OnObjectChanged(FOnSetObject::CreateUObject(this, &UAnimGraphNode_SubInstance::OnSetInstanceBlueprint, ClassHandle))
+		];
+}
+
+FText UAnimGraphNode_SubInstance::GetPropertyTypeText(UProperty* Property)
+{
+	FText PropertyTypeText;
+
+	if(UStructProperty* StructProperty = Cast<UStructProperty>(Property))
+	{
+		PropertyTypeText = StructProperty->Struct->GetDisplayNameText();
+	}
+	else if(UObjectProperty* ObjectProperty = Cast<UObjectProperty>(Property))
+	{
+		PropertyTypeText = ObjectProperty->PropertyClass->GetDisplayNameText();
+	}
+	else if(UClass* PropClass = Property->GetClass())
+	{
+		PropertyTypeText = PropClass->GetDisplayNameText();
+	}
+	else
+	{
+		PropertyTypeText = LOCTEXT("PropertyTypeUnknown", "Unknown");
+	}
+	
+	return PropertyTypeText;
+}
+
+void UAnimGraphNode_SubInstance::RebuildExposedProperties(UClass* InNewClass)
+{
+	ExposedPropertyNames.Empty();
+	KnownExposableProperties.Empty();
+	if(InNewClass)
+	{
+		TArray<UProperty*> ExposableProperties;
+		GetExposableProperties(ExposableProperties);
+
+		for(UProperty* Property : ExposableProperties)
+		{
+			KnownExposableProperties.Add(Property->GetFName());
+		}
+	}
+}
+
+ECheckBoxState UAnimGraphNode_SubInstance::IsPropertyExposed(FName PropertyName) const
+{
+	return ExposedPropertyNames.Contains(PropertyName) ? ECheckBoxState::Checked : ECheckBoxState::Unchecked;
+}
+
+void UAnimGraphNode_SubInstance::OnPropertyExposeCheckboxChanged(ECheckBoxState NewState, FName PropertyName)
+{
+	if(NewState == ECheckBoxState::Checked)
+	{
+		ExposedPropertyNames.AddUnique(PropertyName);
+	}
+	else if(NewState == ECheckBoxState::Unchecked)
+	{
+		ExposedPropertyNames.Remove(PropertyName);
+	}
+
+	ReconstructNode();
+}
+
+void UAnimGraphNode_SubInstance::OnInstanceClassChanged(IDetailLayoutBuilder* DetailBuilder)
+{
+	if(DetailBuilder)
+	{
+		DetailBuilder->ForceRefreshDetails();
+	}
+}
+
+FString UAnimGraphNode_SubInstance::GetCurrentInstanceBlueprintPath() const
+{
+	UClass* InstanceClass = *Node.InstanceClass;
+
+	if(InstanceClass)
+	{
+		UBlueprint* ActualBlueprint = UBlueprint::GetBlueprintFromClass(InstanceClass);
+
+		if(ActualBlueprint)
+		{
+			return ActualBlueprint->GetPathName();
+		}
+	}
+
+	return FString();
+}
+
+bool UAnimGraphNode_SubInstance::OnShouldFilterInstanceBlueprint(const FAssetData& AssetData) const
+{
+	if(const FString* SkeletonName = AssetData.TagsAndValues.Find(TEXT("TargetSkeleton")))
+	{
+		if(UAnimBlueprint* CurrentBlueprint = Cast<UAnimBlueprint>(GetBlueprint()))
+		{
+			FString TargetSkeletonName = FString::Printf(TEXT("%s'%s'"), *CurrentBlueprint->TargetSkeleton->GetClass()->GetName(), *CurrentBlueprint->TargetSkeleton->GetPathName());
+
+			return *SkeletonName != TargetSkeletonName;
+		}
+	}
+
+	return false;
+}
+
+void UAnimGraphNode_SubInstance::OnSetInstanceBlueprint(const FAssetData& AssetData, TSharedRef<IPropertyHandle> InstanceClassPropHandle)
+{
+	if(UAnimBlueprint* Blueprint = Cast<UAnimBlueprint>(AssetData.GetAsset()))
+	{
+		FScopedTransaction Transaction(LOCTEXT("SetBP", "Set Instance Blueprint"));
+
+		Modify();
+		
+		InstanceClassPropHandle->SetValue(Blueprint->GetAnimBlueprintGeneratedClass());
+	}
+}
+
+UObject* UAnimGraphNode_SubInstance::GetJumpTargetForDoubleClick() const
+{
+	UClass* InstanceClass = *Node.InstanceClass;
+	
+	if(InstanceClass)
+	{
+		return InstanceClass->ClassGeneratedBy;
+	}
+
+	return nullptr;
+}
+
+bool UAnimGraphNode_SubInstance::HasExternalDependencies(TArray<class UStruct*>* OptionalOutput /*= NULL*/) const
+{
+	UClass* InstanceClassToUse = *Node.InstanceClass;
+
+	// Add our instance class... If that changes we need a recompile
+	if(InstanceClassToUse && OptionalOutput)
+	{
+		OptionalOutput->AddUnique(InstanceClassToUse);
+	}
+
+	bool bSuperResult = Super::HasExternalDependencies(OptionalOutput);
+	return InstanceClassToUse || bSuperResult;
+}
+
+void UAnimGraphNode_SubInstance::GetExposableProperties(TArray<UProperty*>& OutExposableProperties) const
+{
+	OutExposableProperties.Empty();
+
+	UClass* TargetClass = *Node.InstanceClass;
+
+	if(TargetClass)
+	{
+		const UEdGraphSchema_K2* Schema = CastChecked<UEdGraphSchema_K2>(GetSchema());
+
+		for(TFieldIterator<UProperty> It(TargetClass, EFieldIteratorFlags::IncludeSuper); It; ++It)
+		{
+			UProperty* CurProperty = *It;
+			FEdGraphPinType PinType;
+
+			if(CurProperty->HasAllPropertyFlags(CPF_Edit | CPF_BlueprintVisible) && CurProperty->HasAllFlags(RF_Public) && Schema->ConvertPropertyToPinType(CurProperty, PinType))
+			{
+				OutExposableProperties.Add(CurProperty);
+			}
+		}
+	}
+}
+
+#undef LOCTEXT_NAMESPACE
diff --git a/Engine/Source/Editor/AnimGraph/Private/AnimPreviewInstance.cpp b/Engine/Source/Editor/AnimGraph/Private/AnimPreviewInstance.cpp
--- a/Engine/Source/Editor/AnimGraph/Private/AnimPreviewInstance.cpp
+++ b/Engine/Source/Editor/AnimGraph/Private/AnimPreviewInstance.cpp
@@ -1082,4 +1082,14 @@ FAnimInstanceProxy* UAnimPreviewInstance::CreateAnimInstanceProxy()
 	return new FAnimPreviewInstanceProxy(this);
 }
 
+void UAnimPreviewInstance::PrePerformAnimationEvaluation()
+{
+	GetProxyOnGameThread<FAnimPreviewInstanceProxy>().InitializeObjects(this);
+}
+
+void UAnimPreviewInstance::PostPerformAnimationEvaluation()
+{
+	GetProxyOnGameThread<FAnimPreviewInstanceProxy>().ClearObjects();
+}
+
 #undef LOCTEXT_NAMESPACE
\ No newline at end of file
diff --git a/Engine/Source/Editor/AnimGraph/Private/AnimationGraphSchema.cpp b/Engine/Source/Editor/AnimGraph/Private/AnimationGraphSchema.cpp
--- a/Engine/Source/Editor/AnimGraph/Private/AnimationGraphSchema.cpp
+++ b/Engine/Source/Editor/AnimGraph/Private/AnimationGraphSchema.cpp
@@ -21,6 +21,8 @@
 #include "AnimGraphNode_Root.h"
 #include "AnimGraphNode_RotationOffsetBlendSpace.h"
 #include "AnimGraphNode_SequencePlayer.h"
+#include "AnimGraphNode_PoseBlendNode.h"
+#include "AnimGraphNode_PoseByName.h"
 
 #define LOCTEXT_NAMESPACE "AnimationGraphSchema"
 
@@ -297,6 +299,22 @@ void UAnimationGraphSchema::UpdateNodeWithAsset(UK2Node* K2Node, UAnimationAsset
 				}
 			}
 		}
+		else if (UPoseAsset* PoseAsset = Cast<UPoseAsset>(Asset))
+		{
+			if (UAnimGraphNode_PoseBlendNode* PoseBlendNode = Cast<UAnimGraphNode_PoseBlendNode>(K2Node))
+			{
+				// Skeleton matches, and it's a blendspace player; replace the existing blendspace with the dragged one
+				PoseBlendNode->Node.PoseAsset = PoseAsset;
+			}
+			else
+			{
+				if (UAnimGraphNode_PoseByName* PoseByNameNode = Cast<UAnimGraphNode_PoseByName>(K2Node))
+				{
+					// Skeleton matches, and it's a blendspace player; replace the existing blendspace with the dragged one
+					PoseByNameNode->Node.PoseAsset = PoseAsset;
+				}
+			}
+		}
 	}
 }
 
@@ -358,6 +376,11 @@ void UAnimationGraphSchema::GetAssetsNodeHoverMessage(const TArray<FAssetData>&
 			PlayerNodeUnderCursor = Cast<const UAnimGraphNode_BlendSpacePlayer>(HoverNode);
 		}
 	}
+	else if (Asset->IsA(UPoseAsset::StaticClass()))
+	{
+		UPoseAsset* PoseAsset = CastChecked<UPoseAsset>(Asset);
+		PlayerNodeUnderCursor = Cast<const UAnimGraphNode_PoseBlendNode>(HoverNode);
+	}
 
 	// this one only should happen when there is an Anim Blueprint
 	UAnimBlueprint* AnimBlueprint = Cast<UAnimBlueprint>(FBlueprintEditorUtils::FindBlueprintForNode(HoverNode));
diff --git a/Engine/Source/Editor/ComponentVisualizers/Private/ComponentVisualizers.cpp b/Engine/Source/Editor/ComponentVisualizers/Private/ComponentVisualizers.cpp
--- a/Engine/Source/Editor/ComponentVisualizers/Private/ComponentVisualizers.cpp
+++ b/Engine/Source/Editor/ComponentVisualizers/Private/ComponentVisualizers.cpp
@@ -12,6 +12,7 @@
 #include "AudioComponentVisualizer.h"
 #include "RadialForceComponentVisualizer.h"
 #include "ConstraintComponentVisualizer.h"
+#include "PhysicalAnimationComponentVisualizer.h"
 #include "SpringArmComponentVisualizer.h"
 #include "SplineComponentVisualizer.h"
 #include "SplineMeshComponentVisualizer.h"
@@ -25,6 +26,7 @@
 #include "Components/AudioComponent.h"
 #include "PhysicsEngine/RadialForceComponent.h"
 #include "PhysicsEngine/PhysicsConstraintComponent.h"
+#include "PhysicsEngine/PhysicalAnimationComponent.h"
 #include "GameFramework/SpringArmComponent.h"
 #include "Components/SplineComponent.h"
 #include "Components/PrimitiveComponent.h"
@@ -39,6 +41,7 @@ void FComponentVisualizersModule::StartupModule()
 	RegisterComponentVisualizer(UAudioComponent::StaticClass()->GetFName(), MakeShareable(new FAudioComponentVisualizer));
 	RegisterComponentVisualizer(URadialForceComponent::StaticClass()->GetFName(), MakeShareable(new FRadialForceComponentVisualizer));
 	RegisterComponentVisualizer(UPhysicsConstraintComponent::StaticClass()->GetFName(), MakeShareable(new FConstraintComponentVisualizer));
+	RegisterComponentVisualizer(UPhysicalAnimationComponent::StaticClass()->GetFName(), MakeShareable(new FPhysicsAnimationComponentVisualizer));
 	RegisterComponentVisualizer(USpringArmComponent::StaticClass()->GetFName(), MakeShareable(new FSpringArmComponentVisualizer));
 	RegisterComponentVisualizer(USplineComponent::StaticClass()->GetFName(), MakeShareable(new FSplineComponentVisualizer));
 	RegisterComponentVisualizer(USplineMeshComponent::StaticClass()->GetFName(), MakeShareable(new FSplineMeshComponentVisualizer));
diff --git a/Engine/Source/Editor/ComponentVisualizers/Public/PhysicalAnimationComponentVisualizer.h b/Engine/Source/Editor/ComponentVisualizers/Public/PhysicalAnimationComponentVisualizer.h
new file mode 100644
--- /dev/null
+++ b/Engine/Source/Editor/ComponentVisualizers/Public/PhysicalAnimationComponentVisualizer.h
@@ -0,0 +1,13 @@
+// Copyright 1998-2016 Epic Games, Inc. All Rights Reserved.
+
+#pragma once
+
+#include "ComponentVisualizer.h"
+
+class COMPONENTVISUALIZERS_API FPhysicsAnimationComponentVisualizer : public FComponentVisualizer
+{
+public:
+	//~ Begin FComponentVisualizer Interface
+	virtual void DrawVisualization(const UActorComponent* Component, const FSceneView* View, FPrimitiveDrawInterface* PDI) override;
+	//~ End FComponentVisualizer Interface
+};
diff --git a/Engine/Source/Editor/DetailCustomizations/Private/BodyInstanceCustomization.cpp b/Engine/Source/Editor/DetailCustomizations/Private/BodyInstanceCustomization.cpp
--- a/Engine/Source/Editor/DetailCustomizations/Private/BodyInstanceCustomization.cpp
+++ b/Engine/Source/Editor/DetailCustomizations/Private/BodyInstanceCustomization.cpp
@@ -1428,14 +1428,30 @@ bool FBodyInstanceCustomizationHelper::IsBodyMassReadOnly() const
 TOptional<float> FBodyInstanceCustomizationHelper::OnGetBodyMaxAngularVelocity() const
 {
 	UPrimitiveComponent* Comp = nullptr;
-	const float MaxAngularVelocity = UPhysicsSettings::Get()->MaxAngularVelocity;
+
+	const float DefaultMaxAngularVelocity = UPhysicsSettings::Get()->MaxAngularVelocity;
+	float MaxAngularVelocity = DefaultMaxAngularVelocity;
+	bool bFoundComponent = false;
 
 	for (auto ObjectIt = ObjectsCustomized.CreateConstIterator(); ObjectIt; ++ObjectIt)
 	{
 		if (ObjectIt->IsValid() && (*ObjectIt)->IsA(UPrimitiveComponent::StaticClass()))
 		{
 			Comp = Cast<UPrimitiveComponent>(ObjectIt->Get());
-			Comp->BodyInstance.MaxAngularVelocity = MaxAngularVelocity;	//update max angular velocity so that overriding gives the same value initially
+
+			const float CompMaxAngularVelocity = Comp->BodyInstance.bOverrideMaxAngularVelocity ?
+													Comp->BodyInstance.MaxAngularVelocity :
+													DefaultMaxAngularVelocity;
+
+			if (!bFoundComponent)
+			{
+				bFoundComponent = true;
+				MaxAngularVelocity = CompMaxAngularVelocity;
+			}
+			else if (MaxAngularVelocity != CompMaxAngularVelocity)
+			{
+				return TOptional<float>();
+			}
 		}
 	}
 
diff --git a/Engine/Source/Editor/DetailCustomizations/Private/BodySetupDetails.cpp b/Engine/Source/Editor/DetailCustomizations/Private/BodySetupDetails.cpp
--- a/Engine/Source/Editor/DetailCustomizations/Private/BodySetupDetails.cpp
+++ b/Engine/Source/Editor/DetailCustomizations/Private/BodySetupDetails.cpp
@@ -6,6 +6,7 @@
 #include "ObjectEditorUtils.h"
 #include "IDocumentation.h"
 #include "PhysicsEngine/BodySetup.h"
+#include "PhysicsEngine/PhysicsAsset.h"
 #include "SNumericEntryBox.h"
 
 #define LOCTEXT_NAMESPACE "BodySetupDetails"
@@ -49,9 +50,186 @@ void FBodySetupDetails::CustomizeDetails( IDetailLayoutBuilder& DetailBuilder )
 					CollisionCategory.AddProperty(ChildProperty);
 				}
 			}
+
+			const bool bInPhat = ObjectsCustomized.Num() && (Cast<USkeletalBodySetup>(ObjectsCustomized[0].Get()) != nullptr);
+			if(bInPhat)
+			{
+				//AddPhysicalAnimation(DetailBuilder);
+			}
 		}
 	}
 }
 
+TSharedRef<IDetailCustomization> FSkeletalBodySetupDetails::MakeInstance()
+{
+	return MakeShareable(new FSkeletalBodySetupDetails);
+}
+
+
+void FSkeletalBodySetupDetails::CustomizeDetails(IDetailLayoutBuilder& DetailBuilder)
+{
+	DetailBuilder.GetObjectsBeingCustomized(ObjectsCustomized);
+
+	IDetailCategoryBuilder& Cat = DetailBuilder.EditCategory(TEXT("PhysicalAnimation"));
+
+	const TArray<TWeakObjectPtr<UObject>>& ObjectsCustomizedLocal = ObjectsCustomized;
+	auto PhysAnimEditable = [ObjectsCustomizedLocal]() -> bool
+	{
+		bool bVisible = ObjectsCustomizedLocal.Num() > 0;
+		for (TWeakObjectPtr<UObject> WeakObj : ObjectsCustomizedLocal)
+		{
+			if (USkeletalBodySetup* BS = Cast<USkeletalBodySetup>(WeakObj.Get()))
+			{
+				if (!BS->FindPhysicalAnimationProfile(BS->GetCurrentPhysicalAnimationProfileName()))
+				{
+					bVisible = false;
+					break;
+				}
+			}
+			else
+			{
+				bVisible = false;
+				break;
+			}
+		}
+
+		return bVisible;
+	};
+
+	auto AddProfileLambda = [ObjectsCustomizedLocal]() 
+	{
+		const FScopedTransaction Transaction(LOCTEXT("AddProfile", "Add Physical Animation Profile"));
+		bool bVisible = ObjectsCustomizedLocal.Num() > 0;
+		for (TWeakObjectPtr<UObject> WeakObj : ObjectsCustomizedLocal)
+		{
+			if (USkeletalBodySetup* BS = Cast<USkeletalBodySetup>(WeakObj.Get()))
+			{
+				FName ProfileName = BS->GetCurrentPhysicalAnimationProfileName();
+				if (!BS->FindPhysicalAnimationProfile(ProfileName))
+				{
+					BS->AddPhysicalAnimationProfile(ProfileName);
+				}
+			}
+		}
+
+		return FReply::Handled();;
+	};
+
+	auto DeleteProfileLambda = [ObjectsCustomizedLocal]()
+	{
+		const FScopedTransaction Transaction(LOCTEXT("RemoveProfile", "Remove Physical Animation Profile"));
+		bool bVisible = ObjectsCustomizedLocal.Num() > 0;
+		for (TWeakObjectPtr<UObject> WeakObj : ObjectsCustomizedLocal)
+		{
+			if (USkeletalBodySetup* BS = Cast<USkeletalBodySetup>(WeakObj.Get()))
+			{
+				FName ProfileName = BS->GetCurrentPhysicalAnimationProfileName();
+				BS->RemovePhysicalAnimationProfile(ProfileName);
+			}
+		}
+
+		return FReply::Handled();;
+	};
+
+	TAttribute<EVisibility> PhysAnimVisible = TAttribute<EVisibility>::Create(TAttribute<EVisibility>::FGetter::CreateLambda([PhysAnimEditable]()
+	{
+		return PhysAnimEditable() == true ? EVisibility::Visible : EVisibility::Collapsed;
+	}));
+
+	TAttribute<EVisibility> NewPhysAnimButtonVisible = TAttribute<EVisibility>::Create(TAttribute<EVisibility>::FGetter::CreateLambda([PhysAnimEditable]()
+	{
+		return PhysAnimEditable() == true ? EVisibility::Collapsed : EVisibility::Visible;
+	}));
+
+	TAttribute<bool> PhysAnimButtonEnabled = TAttribute<bool>::Create(TAttribute<bool>::FGetter::CreateLambda([ObjectsCustomizedLocal]()
+	{
+		for (TWeakObjectPtr<UObject> WeakObj : ObjectsCustomizedLocal)
+		{
+			if (USkeletalBodySetup* BS = Cast<USkeletalBodySetup>(WeakObj.Get()))
+			{
+				if (BS->GetCurrentPhysicalAnimationProfileName() == NAME_None)
+				{
+					return false;
+				}
+			}
+		}
+
+		return true;
+	}));
+
+	Cat.AddCustomRow(LOCTEXT("NewPhysAnim", "NewPhysicalAnimationProfile"))
+	.Visibility(NewPhysAnimButtonVisible)
+	[
+		SNew(SVerticalBox)
+		+ SVerticalBox::Slot()
+		.HAlign(HAlign_Center)
+		.Padding(0, 10, 0, 10)
+		[
+			SNew(SBox)
+			.WidthOverride(128)
+			.HeightOverride(32)
+			[
+				SNew(SButton)
+				.HAlign(HAlign_Center)
+				.IsEnabled(PhysAnimButtonEnabled)
+				.OnClicked_Lambda(AddProfileLambda)
+				[
+					SNew(STextBlock)
+					.Text(LOCTEXT("NewPhysAnimButton", "Add To Current Profile"))
+					.ToolTipText(LOCTEXT("NewPhysAnimButtonToolTip", "Add to current physical animation profile."))
+				]
+			]	
+		]
+		+ SVerticalBox::Slot()
+		.HAlign(HAlign_Center)
+		[
+			SNew(STextBlock)
+			.Text(LOCTEXT("NewPhysAnimLabel", "At least one body is not in the current physical animation profile."))
+			.Font(IDetailLayoutBuilder::GetDetailFont())
+			.ToolTipText(LOCTEXT("NewPhysAnimLabelTooltip", "Could not found in the current physical animation profile."))
+		]
+	];
+
+	TSharedPtr<IPropertyHandle> PhysicalAnimationProfile = DetailBuilder.GetProperty(GET_MEMBER_NAME_CHECKED(USkeletalBodySetup, CurrentPhysicalAnimationProfile));
+	PhysicalAnimationProfile->MarkHiddenByCustomization();
+
+	uint32 NumChildren = 0;
+	TSharedPtr<IPropertyHandle> ProfileData = PhysicalAnimationProfile->GetChildHandle(GET_MEMBER_NAME_CHECKED(FPhysicalAnimationProfile, PhysicalAnimationData));
+	ProfileData->GetNumChildren(NumChildren);
+	for(uint32 ChildIdx = 0; ChildIdx < NumChildren; ++ChildIdx)
+	{
+		TSharedPtr<IPropertyHandle> Child = ProfileData->GetChildHandle(ChildIdx);
+		if(!Child->IsCustomized())
+		{
+			Cat.AddProperty(Child)
+			.Visibility(PhysAnimVisible);
+		}
+	}
+
+	Cat.AddCustomRow(LOCTEXT("DeletePhysAnim", "DeletePhysicalAnimationProfile"))
+	.Visibility(PhysAnimVisible)
+	[
+		SNew(SVerticalBox)
+		+ SVerticalBox::Slot()
+		.Padding(0, 10, 0, 10)
+		.HAlign(HAlign_Center)
+		[
+			SNew(SBox)
+			.WidthOverride(128)
+			[
+				SNew(SButton)
+				.HAlign(HAlign_Center)
+				.OnClicked_Lambda(DeleteProfileLambda)
+				[
+					SNew(STextBlock)
+					.Text(LOCTEXT("DeletePhysAnimButton", "Remove From Current Profile"))
+					.ToolTipText(LOCTEXT("DeletePhysAnimButtonToolTip", "Removes the selected body from the current physical animation profile."))
+				]
+			]
+		]
+		
+	];
+}
+
 #undef LOCTEXT_NAMESPACE
 
diff --git a/Engine/Source/Editor/DetailCustomizations/Private/BodySetupDetails.h b/Engine/Source/Editor/DetailCustomizations/Private/BodySetupDetails.h
--- a/Engine/Source/Editor/DetailCustomizations/Private/BodySetupDetails.h
+++ b/Engine/Source/Editor/DetailCustomizations/Private/BodySetupDetails.h
@@ -13,8 +13,26 @@ public:
 	/** IDetailCustomization interface */
 	virtual void CustomizeDetails( IDetailLayoutBuilder& DetailBuilder ) override;
 
+private:
+	void AddPhysicalAnimation(IDetailLayoutBuilder& DetailBuilder);
+
 private:
 	TSharedPtr<class FBodyInstanceCustomizationHelper> BodyInstanceCustomizationHelper;
 	TArray< TWeakObjectPtr<UObject> > ObjectsCustomized;
 };
 
+
+
+class FSkeletalBodySetupDetails : public IDetailCustomization
+{
+public:
+	/** Makes a new instance of this detail layout class for a specific detail view requesting it */
+	static TSharedRef<IDetailCustomization> MakeInstance();
+
+	/** IDetailCustomization interface */
+	virtual void CustomizeDetails(IDetailLayoutBuilder& DetailBuilder) override;
+
+private:
+	TArray< TWeakObjectPtr<UObject> > ObjectsCustomized;
+};
+
diff --git a/Engine/Source/Editor/DetailCustomizations/Private/DetailCustomizations.cpp b/Engine/Source/Editor/DetailCustomizations/Private/DetailCustomizations.cpp
--- a/Engine/Source/Editor/DetailCustomizations/Private/DetailCustomizations.cpp
+++ b/Engine/Source/Editor/DetailCustomizations/Private/DetailCustomizations.cpp
@@ -283,6 +283,7 @@ void FDetailCustomizationsModule::RegisterObjectCustomizations()
 	RegisterCustomClassLayout("SoundWave", FOnGetDetailCustomizationInstance::CreateStatic(&FSoundWaveDetails::MakeInstance));
 	RegisterCustomClassLayout("DialogueWave", FOnGetDetailCustomizationInstance::CreateStatic(&FDialogueWaveDetails::MakeInstance));
 	RegisterCustomClassLayout("BodySetup", FOnGetDetailCustomizationInstance::CreateStatic(&FBodySetupDetails::MakeInstance));
+	RegisterCustomClassLayout("SkeletalBodySetup", FOnGetDetailCustomizationInstance::CreateStatic(&FSkeletalBodySetupDetails::MakeInstance));
 	RegisterCustomClassLayout("DestructibleMesh", FOnGetDetailCustomizationInstance::CreateStatic(&FDestructibleMeshDetails::MakeInstance));
 	RegisterCustomClassLayout("PhysicsConstraintTemplate", FOnGetDetailCustomizationInstance::CreateStatic(&FPhysicsConstraintComponentDetails::MakeInstance));
 	RegisterCustomClassLayout("PhysicsConstraintComponent", FOnGetDetailCustomizationInstance::CreateStatic(&FPhysicsConstraintComponentDetails::MakeInstance));
diff --git a/Engine/Source/Editor/DetailCustomizations/Private/PoseAssetDetails.cpp b/Engine/Source/Editor/DetailCustomizations/Private/PoseAssetDetails.cpp
--- a/Engine/Source/Editor/DetailCustomizations/Private/PoseAssetDetails.cpp
+++ b/Engine/Source/Editor/DetailCustomizations/Private/PoseAssetDetails.cpp
@@ -227,6 +227,43 @@ void FPoseAssetDetails::CustomizeDetails(IDetailLayoutBuilder& DetailBuilder)
 			.IsEnabled(this, &FPoseAssetDetails::CanApplySettings)
 		]
 	];
+
+	IDetailCategoryBuilder& SourceCategory = DetailBuilder.EditCategory("Source");
+	SourceAnimationPropertyHandle = DetailBuilder.GetProperty(GET_MEMBER_NAME_CHECKED(UPoseAsset, SourceAnimation));
+
+	DetailBuilder.HideProperty(SourceAnimationPropertyHandle);
+
+	SourceCategory.AddCustomRow(SourceAnimationPropertyHandle->GetPropertyDisplayName())
+	.NameContent()
+	[
+		SourceAnimationPropertyHandle->CreatePropertyNameWidget()
+	]
+	.ValueContent()
+	.MinDesiredWidth(200)
+	[
+		SNew(SVerticalBox)
+		+ SVerticalBox::Slot()
+		.AutoHeight()
+		[
+			SourceAnimationPropertyHandle->CreatePropertyValueWidget()
+		]
+		+ SVerticalBox::Slot()
+		.AutoHeight()
+		[
+			SNew(SBox)
+			.Padding(5)
+			.HAlign(HAlign_Center)
+			.VAlign(VAlign_Center)
+			.WidthOverride(100)
+			[
+				SNew(SButton)
+				.Text(LOCTEXT("UpdateSource_Lable", "Update Source"))
+				.ToolTipText(LOCTEXT("UpdateSource_Tooltip", "Update Pose From Source Animation"))
+				.OnClicked(this, &FPoseAssetDetails::OnUpdatePoseSourceAnimation)
+				.HAlign(HAlign_Center)
+			]
+		]
+	];
 }
 
 FText FPoseAssetDetails::GetButtonText() const
@@ -515,4 +552,26 @@ void FPoseAssetDetails::CachePoseAssetData()
 	 return (bCachedAdditive ? ECheckBoxState::Checked : ECheckBoxState::Unchecked);
  }
 
+
+ FReply FPoseAssetDetails::OnUpdatePoseSourceAnimation()
+ {
+	 if (PoseAsset.IsValid())
+	 {
+		 UObject* ObjectSet;
+		 SourceAnimationPropertyHandle->GetValue(ObjectSet);
+
+		 UAnimSequence* AnimSequenceSelected = Cast<UAnimSequence>(ObjectSet);
+		 if (AnimSequenceSelected && AnimSequenceSelected->GetSkeleton() == PoseAsset->GetSkeleton())
+		 {
+			 PoseAsset->UpdatePoseFromAnimation(AnimSequenceSelected);
+		 }
+		 else
+		 {
+			 // invalid source asset message
+			 FMessageDialog::Open(EAppMsgType::Ok, LOCTEXT("UpdatePoseWithInvalidSkeleton", "The source animation contains invalid skeleton. Make sure to select source with the skeleton that matches current pose asset."));
+		 }
+	 }
+	 return FReply::Handled();
+ }
+
 #undef LOCTEXT_NAMESPACE
diff --git a/Engine/Source/Editor/DetailCustomizations/Private/PoseAssetDetails.h b/Engine/Source/Editor/DetailCustomizations/Private/PoseAssetDetails.h
--- a/Engine/Source/Editor/DetailCustomizations/Private/PoseAssetDetails.h
+++ b/Engine/Source/Editor/DetailCustomizations/Private/PoseAssetDetails.h
@@ -64,6 +64,9 @@ private:
 	bool CanApplySettings() const;
 	FReply OnApplyAdditiveSettings();
 
+	TSharedPtr<IPropertyHandle> SourceAnimationPropertyHandle;
+	FReply OnUpdatePoseSourceAnimation();
+
 	FText GetButtonText() const;
 };
 
diff --git a/Engine/Source/Editor/GameProjectGeneration/Private/GameProjectUtils.cpp b/Engine/Source/Editor/GameProjectGeneration/Private/GameProjectUtils.cpp
--- a/Engine/Source/Editor/GameProjectGeneration/Private/GameProjectUtils.cpp
+++ b/Engine/Source/Editor/GameProjectGeneration/Private/GameProjectUtils.cpp
@@ -3130,7 +3130,7 @@ bool GameProjectUtils::HasDefaultBuildSettings(const FName InPlatformInfoName)
 	// first check default build settings for all platforms
 	TArray<FString> BoolKeys, IntKeys, StringKeys, BuildKeys;
 	BuildKeys.Add(TEXT("bCompileApex")); BuildKeys.Add(TEXT("bCompileBox2D")); BuildKeys.Add(TEXT("bCompileICU"));
-	BuildKeys.Add(TEXT("bCompileSimplygon")); BuildKeys.Add(TEXT("bCompileSimplygonSSF")); BuildKeys.Add(TEXT("bCompi8leLeanAndMeanUE"));
+	BuildKeys.Add(TEXT("bCompileSimplygon")); BuildKeys.Add(TEXT("bCompileSimplygonSSF")); BuildKeys.Add(TEXT("bCompileLeanAndMeanUE"));
 	BuildKeys.Add(TEXT("bIncludeADO"));	BuildKeys.Add(TEXT("bCompileRecast")); BuildKeys.Add(TEXT("bCompileSpeedTree"));
 	BuildKeys.Add(TEXT("bCompileWithPluginSupport")); BuildKeys.Add(TEXT("bCompilePhysXVehicle")); BuildKeys.Add(TEXT("bCompileFreeType"));
 	BuildKeys.Add(TEXT("bCompileForSize"));	BuildKeys.Add(TEXT("bCompileCEF3"));
diff --git a/Engine/Source/Editor/GraphEditor/Private/SoundCueGraphConnectionDrawingPolicy.cpp b/Engine/Source/Editor/GraphEditor/Private/SoundCueGraphConnectionDrawingPolicy.cpp
--- a/Engine/Source/Editor/GraphEditor/Private/SoundCueGraphConnectionDrawingPolicy.cpp
+++ b/Engine/Source/Editor/GraphEditor/Private/SoundCueGraphConnectionDrawingPolicy.cpp
@@ -58,7 +58,7 @@ void FSoundCueGraphConnectionDrawingPolicy::BuildAudioFlowRoadmap()
 
 			for (int32 WaveIndex = WaveInstances.Num() - 1; WaveIndex >= 0 ; --WaveIndex)
 			{
-				UAudioComponent* WaveInstanceAudioComponent = WaveInstances[WaveIndex]->ActiveSound->GetAudioComponent();
+				UAudioComponent* WaveInstanceAudioComponent = UAudioComponent::GetAudioComponentFromID(WaveInstances[WaveIndex]->ActiveSound->GetAudioComponentID());
 				if (WaveInstanceAudioComponent != PreviewAudioComponent)
 				{
 					WaveInstances.RemoveAtSwap(WaveIndex);
