Fix post-merge Linux CIS errors

diff --git a/Engine/Source/Editor/AnimGraph/Private/AnimationGraphSchema.cpp b/Engine/Source/Editor/AnimGraph/Private/AnimationGraphSchema.cpp
--- a/Engine/Source/Editor/AnimGraph/Private/AnimationGraphSchema.cpp
+++ b/Engine/Source/Editor/AnimGraph/Private/AnimationGraphSchema.cpp
@@ -168,7 +168,7 @@ bool UAnimationGraphSchema::TryCreateConnection(UEdGraphPin* A, UEdGraphPin* B)
 		bool bHavePosePin = bOutputIsPose || bInputIsPose;
 		bool bHaveWildPin = InputPin->PinType.PinCategory == PC_Wildcard || OutputPin->PinType.PinCategory == PC_Wildcard;
 
-		if(bOutputIsPose && bInputIsPose || (bHavePosePin && bHaveWildPin))
+		if((bOutputIsPose && bInputIsPose) || (bHavePosePin && bHaveWildPin))
 		{
 			// Ok this is a valid exec-like line, we need to kill any connections already on the output pin
 			OutputPin->BreakAllPinLinks();
diff --git a/Engine/Source/Editor/Persona/Public/IAssetFamily.h b/Engine/Source/Editor/Persona/Public/IAssetFamily.h
--- a/Engine/Source/Editor/Persona/Public/IAssetFamily.h
+++ b/Engine/Source/Editor/Persona/Public/IAssetFamily.h
@@ -6,6 +6,9 @@
 class IAssetFamily
 {
 public:
+	/** Virtual destructor */
+	virtual ~IAssetFamily() {}
+
 	/** Get all the asset classes this family supports (doesnt need to include derived classes) */
 	virtual void GetAssetTypes(TArray<UClass*>& OutAssetTypes) const = 0;
 
diff --git a/Engine/Source/Editor/Persona/Public/IPersonaToolkit.h b/Engine/Source/Editor/Persona/Public/IPersonaToolkit.h
--- a/Engine/Source/Editor/Persona/Public/IPersonaToolkit.h
+++ b/Engine/Source/Editor/Persona/Public/IPersonaToolkit.h
@@ -8,6 +8,9 @@
 class IPersonaToolkit
 {
 public:
+	/** Virtual destructor */
+	virtual ~IPersonaToolkit() {}
+
 	/** Get the skeleton that we are editing */
 	virtual class USkeleton* GetSkeleton() const = 0;
 
