Fixing fallout from Dev-Framework integration

diff --git a/Engine/Source/Editor/AnimGraph/Private/AnimGraphNode_AssetPlayerBase.cpp b/Engine/Source/Editor/AnimGraph/Private/AnimGraphNode_AssetPlayerBase.cpp
--- a/Engine/Source/Editor/AnimGraph/Private/AnimGraphNode_AssetPlayerBase.cpp
+++ b/Engine/Source/Editor/AnimGraph/Private/AnimGraphNode_AssetPlayerBase.cpp
@@ -15,6 +15,7 @@
 #include "AnimGraphNode_PoseBlendNode.h"
 #include "AnimGraphNode_PoseByName.h"
 #include "AnimGraphNode_PoseDriver.h"
+#include "UObjectIterator.h"
 
 void UAnimGraphNode_AssetPlayerBase::PinConnectionListChanged(UEdGraphPin* Pin)
 {
diff --git a/Engine/Source/Editor/AnimGraph/Private/EditModes/SplineIKEditMode.cpp b/Engine/Source/Editor/AnimGraph/Private/EditModes/SplineIKEditMode.cpp
--- a/Engine/Source/Editor/AnimGraph/Private/EditModes/SplineIKEditMode.cpp
+++ b/Engine/Source/Editor/AnimGraph/Private/EditModes/SplineIKEditMode.cpp
@@ -5,6 +5,8 @@
 #include "AnimGraphNode_SplineIK.h"
 #include "Engine/StaticMeshActor.h"
 #include "Animation/DebugSkelMeshComponent.h"
+#include "SceneManagement.h"
+#include "Materials/Material.h"
 
 FSplineIKEditMode::FSplineIKEditMode()
 	: SplineIKRuntimeNode(nullptr)
@@ -66,7 +68,7 @@ void FSplineIKEditMode::Render(const FSceneView* View, FViewport* Viewport, FPri
 	{
 		PDI->SetHitProxy(new HSplineHandleHitProxy(SplineHandleIndex));
 		FTransform StartTransform = SplineIKRuntimeNode->GetTransformedSplinePoint(SplineHandleIndex);
-		float Scale = View->WorldToScreen(StartTransform.GetLocation()).W * (4.0f / View->UnscaledViewRect.Width() / View->ViewMatrices.GetProjectionMatrix().M[0][0]);
+		const float Scale = View->WorldToScreen(StartTransform.GetLocation()).W * (4.0f / View->UnscaledViewRect.Width() / View->ViewMatrices.GetProjectionMatrix().M[0][0]);
 		DrawSphere(PDI, StartTransform.GetLocation(), FVector(4.0f) * Scale, 64, 64, GEngine->ArrowMaterial->GetRenderProxy(SelectedSplinePoint == SplineHandleIndex), SDPG_Foreground);
 		DrawCoordinateSystem(PDI, StartTransform.GetLocation(), StartTransform.GetRotation().Rotator(), 30.0f * Scale, SDPG_Foreground);
 	}
diff --git a/Engine/Source/Editor/EnvironmentQueryEditor/Private/EnvironmentQueryEditor.cpp b/Engine/Source/Editor/EnvironmentQueryEditor/Private/EnvironmentQueryEditor.cpp
--- a/Engine/Source/Editor/EnvironmentQueryEditor/Private/EnvironmentQueryEditor.cpp
+++ b/Engine/Source/Editor/EnvironmentQueryEditor/Private/EnvironmentQueryEditor.cpp
@@ -19,6 +19,7 @@
 #include "Framework/Commands/Commands.h"
 #include "Framework/Application/SlateApplication.h"
 #include "DesktopPlatformModule.h"
+#include "MultiBox/MultiBoxBuilder.h"
 
 #include "EnvironmentQuery/EnvQuery.h"
 #include "EnvironmentQuery/EnvQueryManager.h"
diff --git a/Engine/Source/Editor/EnvironmentQueryEditor/Private/SEnvQueryProfiler.cpp b/Engine/Source/Editor/EnvironmentQueryEditor/Private/SEnvQueryProfiler.cpp
--- a/Engine/Source/Editor/EnvironmentQueryEditor/Private/SEnvQueryProfiler.cpp
+++ b/Engine/Source/Editor/EnvironmentQueryEditor/Private/SEnvQueryProfiler.cpp
@@ -2,6 +2,7 @@
 
 
 #include "SEnvQueryProfiler.h"
+#include "Editor.h"
 #include "Types/SlateStructs.h"
 #include "Widgets/SBoxPanel.h"
 #include "Widgets/Images/SImage.h"
diff --git a/Engine/Source/Editor/EnvironmentQueryEditor/Private/SEnvQueryProfiler.h b/Engine/Source/Editor/EnvironmentQueryEditor/Private/SEnvQueryProfiler.h
--- a/Engine/Source/Editor/EnvironmentQueryEditor/Private/SEnvQueryProfiler.h
+++ b/Engine/Source/Editor/EnvironmentQueryEditor/Private/SEnvQueryProfiler.h
@@ -6,6 +6,8 @@
 #include "Layout/Visibility.h"
 #include "Styling/SlateColor.h"
 #include "SEnvQueryLoadGraph.h"
+#include "SCompoundWidget.h"
+#include "SListView.h"
 #include "Widgets/DeclarativeSyntaxSupport.h"
 
 class SEnvQueryProfiler : public SCompoundWidget
diff --git a/Engine/Source/Editor/Persona/Private/SPoseEditor.cpp b/Engine/Source/Editor/Persona/Private/SPoseEditor.cpp
--- a/Engine/Source/Editor/Persona/Private/SPoseEditor.cpp
+++ b/Engine/Source/Editor/Persona/Private/SPoseEditor.cpp
@@ -7,9 +7,11 @@
 #include "Widgets/Input/SSpinBox.h"
 #include "Widgets/Layout/SSplitter.h"
 #include "Animation/DebugSkelMeshComponent.h"
+#include "Components/SkeletalMeshComponent.h"
 #include "ScopedTransaction.h"
 #include "Widgets/Input/SSearchBox.h"
 #include "Animation/AnimSingleNodeInstance.h"
+#include "UObjectIterator.h"
 
 #include "Widgets/Text/SInlineEditableTextBlock.h"
 
diff --git a/Engine/Source/Editor/SkeletonEditor/Private/SkeletonTreeVirtualBoneItem.h b/Engine/Source/Editor/SkeletonEditor/Private/SkeletonTreeVirtualBoneItem.h
--- a/Engine/Source/Editor/SkeletonEditor/Private/SkeletonTreeVirtualBoneItem.h
+++ b/Engine/Source/Editor/SkeletonEditor/Private/SkeletonTreeVirtualBoneItem.h
@@ -13,6 +13,8 @@
 #include "SkeletonTreeItem.h"
 #include "IEditableSkeleton.h"
 
+class SInlineEditableTextBlock;
+
 class FSkeletonTreeVirtualBoneItem : public FSkeletonTreeItem
 {
 public:
diff --git a/Engine/Source/Editor/UnrealEd/Private/PhysicsAssetUtils.cpp b/Engine/Source/Editor/UnrealEd/Private/PhysicsAssetUtils.cpp
--- a/Engine/Source/Editor/UnrealEd/Private/PhysicsAssetUtils.cpp
+++ b/Engine/Source/Editor/UnrealEd/Private/PhysicsAssetUtils.cpp
@@ -10,9 +10,11 @@
 #include "PhysicsEngine/BoxElem.h"
 #include "PhysicsEngine/SphereElem.h"
 #include "PhysicsEngine/SphylElem.h"
+#include "Animation/SkeletalMeshActor.h"
 #include "Components/SkeletalMeshComponent.h"
 #include "PhysicsEngine/PhysicsAsset.h"
 #include "PhysicsEngine/PhysicsConstraintTemplate.h"
+#include "PreviewScene.h"
 
 void FPhysAssetCreateParams::Initialize()
 {
diff --git a/Engine/Source/Runtime/Engine/Private/PhysicsEngine/PhysicsAsset.cpp b/Engine/Source/Runtime/Engine/Private/PhysicsEngine/PhysicsAsset.cpp
--- a/Engine/Source/Runtime/Engine/Private/PhysicsEngine/PhysicsAsset.cpp
+++ b/Engine/Source/Runtime/Engine/Private/PhysicsEngine/PhysicsAsset.cpp
@@ -13,6 +13,7 @@
 #include "PhysicsEngine/PhysicsConstraintTemplate.h"
 #include "UObject/ReleaseObjectVersion.h"
 #include "Logging/MessageLog.h"
+#include "UObjectIterator.h"
 
 #define LOCTEXT_NAMESPACE "PhysicsAsset"
 
