fix compile error

diff --git a/Engine/Source/Editor/AnimGraph/Private/EditModes/AnimNodeEditModes.cpp b/Engine/Source/Editor/AnimGraph/Private/EditModes/AnimNodeEditModes.cpp
--- a/Engine/Source/Editor/AnimGraph/Private/EditModes/AnimNodeEditModes.cpp
+++ b/Engine/Source/Editor/AnimGraph/Private/EditModes/AnimNodeEditModes.cpp
@@ -8,5 +8,6 @@ const FEditorModeID AnimNodeEditModes::ObserveBone("AnimGraph.SkeletalControl.Ob
 const FEditorModeID AnimNodeEditModes::ModifyBone("AnimGraph.SkeletalControl.ModifyBone");
 const FEditorModeID AnimNodeEditModes::Fabrik("AnimGraph.SkeletalControl.Fabrik");
 const FEditorModeID AnimNodeEditModes::PoseDriver("AnimGraph.AnimNode.PoseDriver");
+const FEditorModeID AnimNodeEditModes::CCDIK("AnimGraph.SkeletalControl.CCDIK");
 const FEditorModeID AnimNodeEditModes::SplineIK("AnimGraph.AnimNode.SplineIK");
 const FEditorModeID AnimNodeEditModes::LookAt("AnimGraph.SkeletalControl.LookAt");
diff --git a/Engine/Source/Editor/AnimGraph/Public/AnimNodeEditModes.h b/Engine/Source/Editor/AnimGraph/Public/AnimNodeEditModes.h
--- a/Engine/Source/Editor/AnimGraph/Public/AnimNodeEditModes.h
+++ b/Engine/Source/Editor/AnimGraph/Public/AnimNodeEditModes.h
@@ -12,6 +12,7 @@ struct ANIMGRAPH_API AnimNodeEditModes
 	const static FEditorModeID ObserveBone;
 	const static FEditorModeID ModifyBone;
 	const static FEditorModeID Fabrik;
+	const static FEditorModeID CCDIK;
 	const static FEditorModeID PoseDriver;
 	const static FEditorModeID SplineIK;
 	const static FEditorModeID LookAt;
