Fix issues introduced by copy and pasting new animgetters (incorrect state machine references causing crashes)

diff --git a/Engine/Source/Editor/AnimGraph/Classes/K2Node_AnimGetter.h b/Engine/Source/Editor/AnimGraph/Classes/K2Node_AnimGetter.h
--- a/Engine/Source/Editor/AnimGraph/Classes/K2Node_AnimGetter.h
+++ b/Engine/Source/Editor/AnimGraph/Classes/K2Node_AnimGetter.h
@@ -55,6 +55,11 @@ class UK2Node_AnimGetter : public UK2Node_CallFunction
 	GENERATED_BODY()
 public:
 
+	// UObject interface
+	virtual void Serialize(FArchive& Ar) override;
+	virtual void PostPasteNode() override;
+	// End of UObject interface
+
 	// UEdGraphNode interface
 	virtual void AllocateDefaultPins() override;
 	virtual FText GetNodeTitle(ENodeTitleType::Type TitleType) const override;
@@ -95,6 +100,9 @@ protected:
 	
 	//UFunction* GetSourceBlueprintFunction() const;
 
+	// Fixes the SourceNode to be the state machine owner of SourceStateNode (if it is not null)
+	void RestoreStateMachineNode();
+
 	/** Returns whether or not the provided UFunction requires the named parameter */
 	bool GetterRequiresParameter(const UFunction* Getter, FString ParamName) const;
 
diff --git a/Engine/Source/Editor/AnimGraph/Private/AnimGraphNode_StateMachineBase.cpp b/Engine/Source/Editor/AnimGraph/Private/AnimGraphNode_StateMachineBase.cpp
--- a/Engine/Source/Editor/AnimGraph/Private/AnimGraphNode_StateMachineBase.cpp
+++ b/Engine/Source/Editor/AnimGraph/Private/AnimGraphNode_StateMachineBase.cpp
@@ -164,6 +164,7 @@ void UAnimGraphNode_StateMachineBase::PostPasteNode()
 	{
 		GraphNode->CreateNewGuid();
 		GraphNode->PostPasteNode();
+		GraphNode->ReconstructNode();
 	}
 
 	// Find an interesting name
diff --git a/Engine/Source/Editor/AnimGraph/Private/AnimStateNode.cpp b/Engine/Source/Editor/AnimGraph/Private/AnimStateNode.cpp
--- a/Engine/Source/Editor/AnimGraph/Private/AnimStateNode.cpp
+++ b/Engine/Source/Editor/AnimGraph/Private/AnimStateNode.cpp
@@ -144,6 +144,7 @@ void UAnimStateNode::PostPasteNode()
 	{
 		GraphNode->CreateNewGuid();
 		GraphNode->PostPasteNode();
+		GraphNode->ReconstructNode();
 	}
 
 	Super::PostPasteNode();
diff --git a/Engine/Source/Editor/AnimGraph/Private/AnimStateTransitionNode.cpp b/Engine/Source/Editor/AnimGraph/Private/AnimStateTransitionNode.cpp
--- a/Engine/Source/Editor/AnimGraph/Private/AnimStateTransitionNode.cpp
+++ b/Engine/Source/Editor/AnimGraph/Private/AnimStateTransitionNode.cpp
@@ -169,6 +169,13 @@ void UAnimStateTransitionNode::PostPasteNode()
 		CreateBoundGraph();
 	}
 
+	for (UEdGraphNode* GraphNode : BoundGraph->Nodes)
+	{
+		GraphNode->CreateNewGuid();
+		GraphNode->PostPasteNode();
+		GraphNode->ReconstructNode();
+	}
+
 	if(CustomTransitionGraph)
 	{
 		// Needs to be added to the parent graph
diff --git a/Engine/Source/Editor/AnimGraph/Private/K2Node_AnimGetter.cpp b/Engine/Source/Editor/AnimGraph/Private/K2Node_AnimGetter.cpp
--- a/Engine/Source/Editor/AnimGraph/Private/K2Node_AnimGetter.cpp
+++ b/Engine/Source/Editor/AnimGraph/Private/K2Node_AnimGetter.cpp
@@ -17,6 +17,25 @@
 
 #define LOCTEXT_NAMESPACE "AnimGetter"
 
+
+void UK2Node_AnimGetter::Serialize(FArchive& Ar)
+{
+	Super::Serialize(Ar);
+	Ar.UsingCustomVersion(FReleaseObjectVersion::GUID);
+
+	if (Ar.CustomVer(FReleaseObjectVersion::GUID) < FReleaseObjectVersion::FixBrokenStateMachineReferencesInTransitionGetters)
+	{
+		RestoreStateMachineNode();
+	}
+}
+
+void UK2Node_AnimGetter::PostPasteNode()
+{
+	Super::PostPasteNode();
+
+	RestoreStateMachineNode();
+}
+
 void UK2Node_AnimGetter::AllocateDefaultPins()
 {
 	Super::AllocateDefaultPins();
@@ -305,6 +324,21 @@ bool UK2Node_AnimGetter::IsActionFilteredOut(FBlueprintActionFilter const& Filte
 	return true;
 }
 
+void UK2Node_AnimGetter::RestoreStateMachineNode()
+{
+	if(SourceStateNode)
+	{
+		UAnimationStateMachineGraph* Graph = Cast<UAnimationStateMachineGraph>(SourceStateNode->GetOuter());
+		if (Graph)
+		{
+			if (UAnimGraphNode_StateMachine* MachineNode = Cast<UAnimGraphNode_StateMachine>(Graph->GetOuter()))
+			{
+				SourceNode = MachineNode;
+			}
+		}
+	}
+}
+
 bool UK2Node_AnimGetter::GetterRequiresParameter(const UFunction* Getter, FString ParamName) const
 {
 	bool bRequiresParameter = false;
diff --git a/Engine/Source/Runtime/Core/Public/UObject/ReleaseObjectVersion.h b/Engine/Source/Runtime/Core/Public/UObject/ReleaseObjectVersion.h
--- a/Engine/Source/Runtime/Core/Public/UObject/ReleaseObjectVersion.h
+++ b/Engine/Source/Runtime/Core/Public/UObject/ReleaseObjectVersion.h
@@ -58,6 +58,9 @@ struct CORE_API FReleaseObjectVersion
 		// Default values on pins in blueprints could be saved incoherently
 		PinDefaultValuesVerified,
 
+		// During copy and paste transition getters could end up with broken state machine references
+		FixBrokenStateMachineReferencesInTransitionGetters,
+
 		// -----<new versions can be added above this line>-------------------------------------------------
 		VersionPlusOne,
 		LatestVersion = VersionPlusOne - 1
