Added the ability to access sub instance nodes by tag

Sub instances now have the tag as their object FName. We prevent duplicates in the compiler.

diff --git a/Engine/Source/Editor/AnimGraph/Private/AnimGraphNode_SubInstance.cpp b/Engine/Source/Editor/AnimGraph/Private/AnimGraphNode_SubInstance.cpp
--- a/Engine/Source/Editor/AnimGraph/Private/AnimGraphNode_SubInstance.cpp
+++ b/Engine/Source/Editor/AnimGraph/Private/AnimGraphNode_SubInstance.cpp
@@ -44,11 +44,27 @@ FText UAnimGraphNode_SubInstance::GetNodeTitle(ENodeTitleType::Type TitleType) c
 	}
 	if(TitleType == ENodeTitleType::ListView)
 	{
-		return FText::Format(LOCTEXT("TitleListFormat", "{NodeTitle} - Target Class: {TargetClass}"), Args);
+		if(Node.Tag != NAME_None)
+		{
+			Args.Add(TEXT("Tag"), FText::FromName(Node.Tag));
+			return FText::Format(LOCTEXT("TitleListFormatTagged", "{NodeTitle} - Target Class: {TargetClass} - Tag: {Tag}"), Args);
+		}
+		else
+		{
+			return FText::Format(LOCTEXT("TitleListFormat", "{NodeTitle} - Target Class: {TargetClass}"), Args);
+		}
 	}
 	else
 	{
-		return FText::Format(LOCTEXT("TitleFormat", "{NodeTitle}\nTarget Class: {TargetClass}"), Args);
+		if(Node.Tag != NAME_None)
+		{
+			Args.Add(TEXT("Tag"), FText::FromName(Node.Tag));
+			return FText::Format(LOCTEXT("TitleFormatTagged", "{NodeTitle}\nTarget Class: {TargetClass}\nTag: {Tag}"), Args);
+		}
+		else
+		{
+			return FText::Format(LOCTEXT("TitleFormat", "{NodeTitle}\nTarget Class: {TargetClass}"), Args);
+		}
 	}
 }
 
@@ -96,6 +112,28 @@ void UAnimGraphNode_SubInstance::ValidateAnimNodeDuringCompilation(USkeleton* Fo
 		MessageLog.Error(TEXT("Detected loop in sub instance chain starting at @@ inside class @@"), this, AnimBP->GetAnimBlueprintGeneratedClass());
 	}
 
+	// Check for duplicate tags in this anim blueprint
+	for(UEdGraph* Graph : Graphs)
+	{
+		TArray<UAnimGraphNode_SubInstance*> SubInstanceNodes;
+		Graph->GetNodesOfClass(SubInstanceNodes);
+
+		for(UAnimGraphNode_SubInstance* SubInstanceNode : SubInstanceNodes)
+		{
+			if(SubInstanceNode == OriginalNode)
+			{
+				continue;
+			}
+
+			FAnimNode_SubInstance& InnerNode = SubInstanceNode->Node;
+
+			if(InnerNode.Tag != NAME_None && InnerNode.Tag == Node.Tag)
+			{
+				MessageLog.Error(*FText::Format(LOCTEXT("DuplicateTagErrorFormat", "Node @@ and node @@ both have the same tag '{0}'."), FText::FromName(Node.Tag)).ToString(), this, SubInstanceNode);
+			}
+		}
+	}
+
 	// Check we don't try to spawn our own blueprint
 	if(*Node.InstanceClass == AnimBP->GetAnimBlueprintGeneratedClass())
 	{
diff --git a/Engine/Source/Runtime/Engine/Classes/Animation/AnimNode_SubInstance.h b/Engine/Source/Runtime/Engine/Classes/Animation/AnimNode_SubInstance.h
--- a/Engine/Source/Runtime/Engine/Classes/Animation/AnimNode_SubInstance.h
+++ b/Engine/Source/Runtime/Engine/Classes/Animation/AnimNode_SubInstance.h
@@ -28,9 +28,14 @@ public:
 	UPROPERTY()
 	FPoseLink InPose;
 
+	/** The class spawned for this sub-instance */
 	UPROPERTY(EditAnywhere, BlueprintReadWrite, Category = Settings)
 	TSubclassOf<UAnimInstance> InstanceClass;
 
+	/** Optional tag used to identify this sub-instance */
+	UPROPERTY(EditAnywhere, BlueprintReadWrite, Category = Settings)
+	FName Tag;
+
 	/** This is the actual instance allocated at runtime that will run */
 	UPROPERTY(transient)
 	UAnimInstance* InstanceToRun;
diff --git a/Engine/Source/Runtime/Engine/Classes/Components/SkeletalMeshComponent.h b/Engine/Source/Runtime/Engine/Classes/Components/SkeletalMeshComponent.h
--- a/Engine/Source/Runtime/Engine/Classes/Components/SkeletalMeshComponent.h
+++ b/Engine/Source/Runtime/Engine/Classes/Components/SkeletalMeshComponent.h
@@ -663,6 +663,13 @@ public:
 	UFUNCTION(BlueprintCallable, Category = "Components|SkeletalMesh", meta = (Keywords = "AnimBlueprint"))
 	UAnimInstance* GetPostProcessInstance() const;
 
+	/**
+	 * Returns the a tagged sub-instance node. If non sub instances are found or none are tagged with the
+	 * supplied name, this will return NULL.
+	 */
+	UFUNCTION(BlueprintCallable, Category = "Components|SkeletalMesh", meta = (Keywords = "AnimBlueprint"))
+	UAnimInstance* GetSubInstanceByName(FName InName) const;
+
 	/** 
 	 * Returns whether there are any valid instances to run, currently this means whether we have
 	 * have an animation instance or a post process instance available to process.
diff --git a/Engine/Source/Runtime/Engine/Private/Animation/AnimNode_SubInstance.cpp b/Engine/Source/Runtime/Engine/Private/Animation/AnimNode_SubInstance.cpp
--- a/Engine/Source/Runtime/Engine/Private/Animation/AnimNode_SubInstance.cpp
+++ b/Engine/Source/Runtime/Engine/Private/Animation/AnimNode_SubInstance.cpp
@@ -7,6 +7,7 @@
 
 FAnimNode_SubInstance::FAnimNode_SubInstance()
 	: InstanceClass(nullptr)
+	, Tag(NAME_None)
 	, InstanceToRun(nullptr)
 {
 
@@ -142,8 +143,10 @@ void FAnimNode_SubInstance::OnInitializeAnimInstance(const FAnimInstanceProxy* I
 			InstanceToRun = nullptr;
 		}
 
-		// Need an instance to run
-		InstanceToRun = NewObject<UAnimInstance>(MeshComp, InstanceClass);
+		// Need an instance to run, so create it now
+		// We use the tag to name the object, but as we verify there are no duplicates in the compiler we
+		// dont need to verify it is unique here.
+		InstanceToRun = NewObject<UAnimInstance>(MeshComp, InstanceClass, Tag);
 
 		// Set up bone transform array
 		AllocateBoneTransforms(InstanceToRun);
diff --git a/Engine/Source/Runtime/Engine/Private/Components/SkeletalMeshComponent.cpp b/Engine/Source/Runtime/Engine/Private/Components/SkeletalMeshComponent.cpp
--- a/Engine/Source/Runtime/Engine/Private/Components/SkeletalMeshComponent.cpp
+++ b/Engine/Source/Runtime/Engine/Private/Components/SkeletalMeshComponent.cpp
@@ -2590,6 +2590,19 @@ UAnimInstance* USkeletalMeshComponent::GetPostProcessInstance() const
 	return PostProcessAnimInstance;
 }
 
+UAnimInstance* USkeletalMeshComponent::GetSubInstanceByName(FName InName) const
+{
+	for(UAnimInstance* SubInstance : SubInstances)
+	{
+		if(SubInstance->GetFName() == InName)
+		{
+			return SubInstance;
+		}
+	}
+
+	return nullptr;
+}
+
 bool USkeletalMeshComponent::HasValidAnimationInstance() const
 {
 	return AnimScriptInstance || PostProcessAnimInstance;
