Fixed incorrect sub instance error reporting when multiple sub-instances of the same class are used

diff --git a/Engine/Source/Editor/AnimGraph/Private/AnimGraphNode_SubInstance.cpp b/Engine/Source/Editor/AnimGraph/Private/AnimGraphNode_SubInstance.cpp
--- a/Engine/Source/Editor/AnimGraph/Private/AnimGraphNode_SubInstance.cpp
+++ b/Engine/Source/Editor/AnimGraph/Private/AnimGraphNode_SubInstance.cpp
@@ -82,36 +82,15 @@ void UAnimGraphNode_SubInstance::ValidateAnimNodeDuringCompilation(USkeleton* Fo
 		MessageLog.Error(TEXT("Sub instance node @@ has no valid instance class to spawn."), this);
 	}
 
-	// Check for cycles from other sub instance nodes
-	TArray<UEdGraph*> Graphs;
-	AnimBP->GetAllGraphs(Graphs);
-
-	for(UEdGraph* Graph : Graphs)
-	{
-		TArray<UAnimGraphNode_SubInstance*> SubInstanceNodes;
-		Graph->GetNodesOfClass(SubInstanceNodes);
-
-		for(UAnimGraphNode_SubInstance* SubInstanceNode : SubInstanceNodes)
-		{
-			if(SubInstanceNode == OriginalNode)
-			{
-				continue;
-			}
-
-			FAnimNode_SubInstance& InnerNode = SubInstanceNode->Node;
-
-			if(*InnerNode.InstanceClass && *InnerNode.InstanceClass == *Node.InstanceClass)
-			{
-				MessageLog.Error(TEXT("Node @@ and node @@ both target the same class @@, causing a sub instance loop."), this, SubInstanceNode, *Node.InstanceClass);
-			}
-		}
-	}
-
 	if(HasInstanceLoop())
 	{
 		MessageLog.Error(TEXT("Detected loop in sub instance chain starting at @@ inside class @@"), this, AnimBP->GetAnimBlueprintGeneratedClass());
 	}
 
+	// Check for cycles from other sub instance nodes
+	TArray<UEdGraph*> Graphs;
+	AnimBP->GetAllGraphs(Graphs);
+
 	// Check for duplicate tags in this anim blueprint
 	for(UEdGraph* Graph : Graphs)
 	{
