Added a warning message to the AnimGraphNode_Constraint to flag nodes which don't have normalized weights. A screenshot can be seen in the jira

diff --git a/Engine/Source/Editor/AnimGraph/Private/AnimGraphNode_Constraint.cpp b/Engine/Source/Editor/AnimGraph/Private/AnimGraphNode_Constraint.cpp
--- a/Engine/Source/Editor/AnimGraph/Private/AnimGraphNode_Constraint.cpp
+++ b/Engine/Source/Editor/AnimGraph/Private/AnimGraphNode_Constraint.cpp
@@ -55,6 +55,16 @@ void UAnimGraphNode_Constraint::ValidateAnimNodeDuringCompilation(USkeleton* For
 		}
 	}
 
+	float OverallWeight = 0.f;
+	for (const float& Weight : Node.ConstraintWeights)
+	{
+		OverallWeight += Weight;
+	}
+	if (Node.ConstraintWeights.Num() > 0 && !FMath::IsNearlyEqual(OverallWeight, 1.f, ZERO_ANIMWEIGHT_THRESH * float(Node.ConstraintWeights.Num())))
+	{
+		MessageLog.Warning(*LOCTEXT("WeightsDontSumToOne", "@@ - The weights don't add up to 1.0").ToString(), this);
+	}
+
 	Super::ValidateAnimNodeDuringCompilation(ForSkeleton, MessageLog);
 }
 
