Back out changelist 5229480. Cook warnings. This warning is good but we have to fix content before enabling it.

diff --git a/Engine/Source/Editor/AnimGraph/Private/AnimGraphNode_Constraint.cpp b/Engine/Source/Editor/AnimGraph/Private/AnimGraphNode_Constraint.cpp
--- a/Engine/Source/Editor/AnimGraph/Private/AnimGraphNode_Constraint.cpp
+++ b/Engine/Source/Editor/AnimGraph/Private/AnimGraphNode_Constraint.cpp
@@ -55,16 +55,6 @@ void UAnimGraphNode_Constraint::ValidateAnimNodeDuringCompilation(USkeleton* For
 		}
 	}
 
-	float OverallWeight = 0.f;
-	for (const float& Weight : Node.ConstraintWeights)
-	{
-		OverallWeight += Weight;
-	}
-	if (Node.ConstraintWeights.Num() > 0 && !FMath::IsNearlyEqual(OverallWeight, 1.f, ZERO_ANIMWEIGHT_THRESH * float(Node.ConstraintWeights.Num())))
-	{
-		MessageLog.Warning(*LOCTEXT("WeightsDontSumToOne", "@@ - The weights don't add up to 1.0").ToString(), this);
-	}
-
 	Super::ValidateAnimNodeDuringCompilation(ForSkeleton, MessageLog);
 }
 
