Re-fix old bug that never got merged from 4.21

diff --git a/Engine/Source/Editor/AnimGraph/Private/AnimationTransitionSchema.cpp b/Engine/Source/Editor/AnimGraph/Private/AnimationTransitionSchema.cpp
--- a/Engine/Source/Editor/AnimGraph/Private/AnimationTransitionSchema.cpp
+++ b/Engine/Source/Editor/AnimGraph/Private/AnimationTransitionSchema.cpp
@@ -46,10 +46,12 @@ void UAnimationTransitionSchema::GetGraphDisplayInformation(const UEdGraph& Grap
 	if (TransNode == NULL && !Graph.IsPendingKill())
 	{
 		//@TODO: Transition graphs should be created with the transition node as their outer as well!
-		UAnimBlueprint* Blueprint = CastChecked<UAnimBlueprint>(FBlueprintEditorUtils::FindBlueprintForGraph(&Graph));
-		if (UAnimBlueprintGeneratedClass* AnimBlueprintClass = Blueprint->GetAnimBlueprintSkeletonClass())
+		if(UAnimBlueprint* Blueprint = Cast<UAnimBlueprint>(FBlueprintEditorUtils::FindBlueprintForGraph(&Graph)))
 		{
-			TransNode = GetTransitionNodeFromGraph(AnimBlueprintClass->GetAnimBlueprintDebugData(), &Graph);
+			if (UAnimBlueprintGeneratedClass* AnimBlueprintClass = Blueprint->GetAnimBlueprintSkeletonClass())
+			{
+				TransNode = GetTransitionNodeFromGraph(AnimBlueprintClass->GetAnimBlueprintDebugData(), &Graph);
+			}
 		}
 	}
 
