Unable to use cached poses in animation states
#fix Make sure we keeping traversing the outer chain rather than early out if we come across something other than a EdGraph

diff --git a/Engine/Source/Editor/AnimGraph/Private/AnimGraphNode_UseCachedPose.cpp b/Engine/Source/Editor/AnimGraph/Private/AnimGraphNode_UseCachedPose.cpp
--- a/Engine/Source/Editor/AnimGraph/Private/AnimGraphNode_UseCachedPose.cpp
+++ b/Engine/Source/Editor/AnimGraph/Private/AnimGraphNode_UseCachedPose.cpp
@@ -138,9 +138,15 @@ bool UAnimGraphNode_UseCachedPose::IsActionFilteredOut(class FBlueprintActionFil
 		{
 			// Get layer we are in (i.e. the last outer that is a graph)
 			UEdGraph* OuterLayer = Graph;
-			while(UEdGraph* TestGraphOuter = Cast<UEdGraph>(OuterLayer->GetOuter()))
+			UObject* TestGraphOuter = OuterLayer->GetOuter();
+			while (TestGraphOuter)
 			{
-				OuterLayer = TestGraphOuter;
+				if (UEdGraph* EdGraph = Cast<UEdGraph>(TestGraphOuter))
+				{
+					OuterLayer = EdGraph;
+				}
+
+				TestGraphOuter = TestGraphOuter->GetOuter();
 			}
 
 			// Check if the save cached pose node has the layer in its outer chain
