Warning occurred in Static Analysis Win64 (PVS-Studio): 'animgraphnode_layer.cpp'

diff --git a/Engine/Source/Editor/AnimGraph/Private/AnimGraphNode_Layer.cpp b/Engine/Source/Editor/AnimGraph/Private/AnimGraphNode_Layer.cpp
--- a/Engine/Source/Editor/AnimGraph/Private/AnimGraphNode_Layer.cpp
+++ b/Engine/Source/Editor/AnimGraph/Private/AnimGraphNode_Layer.cpp
@@ -75,12 +75,16 @@ void UAnimGraphNode_Layer::ValidateAnimNodeDuringCompilation(USkeleton* ForSkele
 		{
 			// check we implement this interface
 			bool bImplementsInterface = false;
-			for(FBPInterfaceDescription& InterfaceDesc : CurrentBlueprint->ImplementedInterfaces)
+
+			if (CurrentBlueprint)
 			{
-				if(InterfaceDesc.Interface.Get() == TargetClass)
+				for (FBPInterfaceDescription& InterfaceDesc : CurrentBlueprint->ImplementedInterfaces)
 				{
-					bImplementsInterface = true;
-					break;
+					if (InterfaceDesc.Interface.Get() == TargetClass)
+					{
+						bImplementsInterface = true;
+						break;
+					}
 				}
 			}
 
