- Change ensure to Log since this could happen by editing control rig BP externally
- Also cleared non existing IO mapping in editor if found
- I still leave the log because this can happen in game if somebody changes the ControlRig without modifying the mapping
- New ticket for https://jira.it.epicgames.net/browse/UE-77080 to fix for longer term

diff --git a/Engine/Plugins/Experimental/ControlRig/Source/ControlRig/Private/AnimNode_ControlRig.cpp b/Engine/Plugins/Experimental/ControlRig/Source/ControlRig/Private/AnimNode_ControlRig.cpp
--- a/Engine/Plugins/Experimental/ControlRig/Source/ControlRig/Private/AnimNode_ControlRig.cpp
+++ b/Engine/Plugins/Experimental/ControlRig/Source/ControlRig/Private/AnimNode_ControlRig.cpp
@@ -138,7 +138,7 @@ void FAnimNode_ControlRig::UpdateInput(UControlRig* InControlRig, const FPoseCon
 {
 	FAnimNode_ControlRigBase::UpdateInput(InControlRig, InOutput);
 	// now go through variable mapping table and see if anything is mapping through input
-	if (InputMapping.Num() > 0)
+	if (InputMapping.Num() > 0 && InControlRig)
 	{
 		for (auto Iter = InputMapping.CreateConstIterator(); Iter; ++Iter)
 		{
@@ -155,7 +155,10 @@ void FAnimNode_ControlRig::UpdateInput(UControlRig* InControlRig, const FPoseCon
 	
 					// helper function to set input value for ControlRig
 					// This converts to the proper destination type, and sets the float type Value
-					ensure(FControlRigIOHelper::SetInputValue(ControlRig, SourcePath, FControlRigIOTypes::GetTypeString<float>(), Value));
+					if (!FControlRigIOHelper::SetInputValue(InControlRig, SourcePath, FControlRigIOTypes::GetTypeString<float>(), Value))
+					{
+						UE_LOG(LogAnimation, Warning, TEXT("[%s] Missing Input Property [%s]"), *GetNameSafe(InControlRig->GetClass()), *SourcePath.ToString());
+					}
 				}
 			}
 		} 
@@ -167,7 +170,7 @@ void FAnimNode_ControlRig::UpdateOutput(UControlRig* InControlRig, FPoseContext&
 	FAnimNode_ControlRigBase::UpdateOutput(InControlRig, InOutput);
 
 	// update output curves
-	if (OutputMapping.Num() > 0)
+	if (OutputMapping.Num() > 0 && InControlRig)
 	{
 		for (auto Iter = OutputMapping.CreateConstIterator(); Iter; ++Iter)
 		{
@@ -180,7 +183,7 @@ void FAnimNode_ControlRig::UpdateOutput(UControlRig* InControlRig, FPoseContext&
 				// find Segment is right value
 				float Value;
 				// helper function to get output value and convert to float 
-				if (ensure(FControlRigIOHelper::GetOutputValue(ControlRig, SourcePath, FControlRigIOTypes::GetTypeString<float>(), Value)))
+				if (FControlRigIOHelper::GetOutputValue(InControlRig, SourcePath, FControlRigIOTypes::GetTypeString<float>(), Value))
 				{
 					SmartName::UID_Type* UID = CurveMappingUIDs.Find(Iter.Value());
 					if (UID)
@@ -188,6 +191,10 @@ void FAnimNode_ControlRig::UpdateOutput(UControlRig* InControlRig, FPoseContext&
 						InOutput.Curve.Set(*UID, Value);
 					}
 				}
+				else
+				{
+					UE_LOG(LogAnimation, Warning, TEXT("[%s] Missing Output Property [%s]"), *GetNameSafe(ControlRig->GetClass()), *SourcePath.ToString());
+				}
 			}
 		}
 	}
diff --git a/Engine/Plugins/Experimental/ControlRig/Source/ControlRigDeveloper/Private/AnimGraphNode_ControlRig.cpp b/Engine/Plugins/Experimental/ControlRig/Source/ControlRigDeveloper/Private/AnimGraphNode_ControlRig.cpp
--- a/Engine/Plugins/Experimental/ControlRig/Source/ControlRigDeveloper/Private/AnimGraphNode_ControlRig.cpp
+++ b/Engine/Plugins/Experimental/ControlRig/Source/ControlRigDeveloper/Private/AnimGraphNode_ControlRig.cpp
@@ -9,6 +9,8 @@
 #include "DetailWidgetRow.h"
 #include "ControlRigVariables.h"
 #include "SVariableMappingWidget.h"
+#include "ScopedTransaction.h"
+#include "Kismet2/BlueprintEditorUtils.h"
 
 #define LOCTEXT_NAMESPACE "AnimGraphNode_ControlRig"
 
@@ -73,6 +75,26 @@ void UAnimGraphNode_ControlRig::RebuildExposedProperties()
 	// we still update OUtputvariables
 	// we don't want output to be exposed
 	GetIOProperties(false, OutputVariables);
+
+	// clear IO variables that don't exist anymore
+	auto ClearInvalidMapping = [](const TMap<FName, FControlRigIOVariable>& InVariables, TMap<FName, FName>& InOutMapping)
+	{
+		TArray<FName> KeyArray;
+		InOutMapping.GenerateKeyArray(KeyArray);
+
+		for (int32 Index=0; Index<KeyArray.Num(); ++Index)
+		{
+			// if this input doesn't exist anymore
+			if (!InVariables.Contains(KeyArray[Index]))
+			{
+				InOutMapping.Remove(KeyArray[Index]);
+			}
+		}
+	};
+
+	ClearInvalidMapping(InputVariables, Node.InputMapping);
+	ClearInvalidMapping(OutputVariables, Node.OutputMapping);
+
 	for (auto Iter = InputVariables.CreateConstIterator(); Iter; ++Iter)
 	{
 		KnownExposableProperties.Add(Iter.Key());
@@ -130,6 +152,11 @@ void UAnimGraphNode_ControlRig::OnPropertyExposeCheckboxChanged(ECheckBoxState N
 	// see if any of my child has the mapping, and clear them
 	if (NewState == ECheckBoxState::Checked)
 	{
+		FScopedTransaction Transaction(LOCTEXT("PropertyExposedChanged", "Expose Property to Pin"));
+		Modify();
+
+		FBlueprintEditorUtils::MarkBlueprintAsStructurallyModified(GetBlueprint());
+
 		bool bInput = IsInputProperty(PropertyName);
 		// if checked, we clear mapping
 		// and unclear all children
@@ -217,6 +244,11 @@ void UAnimGraphNode_ControlRig::GetIOProperties(bool bInput, TMap<FName, FContro
 
 void UAnimGraphNode_ControlRig::OnVariableMappingChanged(const FName& PathName, const FName& Curve, bool bInput)
 {
+	FScopedTransaction Transaction(LOCTEXT("VariableMappingChanged", "Change Variable Mapping"));
+	Modify();
+
+	FBlueprintEditorUtils::MarkBlueprintAsStructurallyModified(GetBlueprint());
+
 	// @todo: this is not enough when we start breaking down struct
 	Node.SetIOMapping(bInput, PathName, Curve);
 }
diff --git a/Engine/Source/Editor/AnimGraph/Classes/AnimGraphNode_CustomProperty.h b/Engine/Source/Editor/AnimGraph/Classes/AnimGraphNode_CustomProperty.h
--- a/Engine/Source/Editor/AnimGraph/Classes/AnimGraphNode_CustomProperty.h
+++ b/Engine/Source/Editor/AnimGraph/Classes/AnimGraphNode_CustomProperty.h
@@ -29,8 +29,6 @@ public:
 	//~ End UEdGraphNode Interface.
 
 	// UAnimGraphNode_Base interface
-	virtual void CustomizeDetails(IDetailLayoutBuilder& DetailBuilder) override;
-
 	// Gets the property on InOwnerInstanceClass that corresponds to InInputPin
 	void GetInstancePinProperty(const UClass* InOwnerInstanceClass, UEdGraphPin* InInputPin, UProperty*& OutProperty);
 	// Gets the unique name for the property linked to a given pin
@@ -42,6 +40,19 @@ public:
 	// Helper used to get the skeleton class we are targeting
 	virtual UClass* GetTargetSkeletonClass() const;
 
+	// ----- UI CALLBACKS ----- //
+	// User changed the instance class etc.
+	void OnStructuralPropertyChanged(IDetailLayoutBuilder* DetailBuilder);
+	// If given property exposed on this node
+	virtual ECheckBoxState IsPropertyExposed(FName PropertyName) const;
+	// User chose to expose, or unexpose a property
+	virtual void OnPropertyExposeCheckboxChanged(ECheckBoxState NewState, FName PropertyName);
+	// If all possible properties are exposed on this node
+	virtual ECheckBoxState AreAllPropertiesExposed() const;
+	// User chose to expose, or unexpose all properties
+	virtual void OnPropertyExposeAllCheckboxChanged(ECheckBoxState NewState);
+	// User changed the instance class
+	void OnInstanceClassChanged(IDetailLayoutBuilder* DetailBuilder);
 protected:
 
 	/** List of property names we know to exist on the target class, so we can detect when
@@ -61,20 +72,6 @@ protected:
 	// Given a new class, rebuild the known property list (for tracking class changes and moving pins)
 	virtual void RebuildExposedProperties();
 
-	// ----- UI CALLBACKS ----- //
-	// User changed the instance class etc.
-	void OnStructuralPropertyChanged(IDetailLayoutBuilder* DetailBuilder);
-	// If given property exposed on this node
-	virtual ECheckBoxState IsPropertyExposed(FName PropertyName) const;
-	// User chose to expose, or unexpose a property
-	virtual void OnPropertyExposeCheckboxChanged(ECheckBoxState NewState, FName PropertyName);
-	// If all possible properties are exposed on this node
-	virtual ECheckBoxState AreAllPropertiesExposed() const;
-	// User chose to expose, or unexpose all properties
-	virtual void OnPropertyExposeAllCheckboxChanged(ECheckBoxState NewState);
-	// User changed the instance class
-	void OnInstanceClassChanged(IDetailLayoutBuilder* DetailBuilder);
-
 	// internal node accessor
 	virtual FAnimNode_CustomProperty* GetCustomPropertyNode() PURE_VIRTUAL(UAnimGraphNode_CustomProperty::GetCustomPropertyNode, return nullptr;);
 	virtual const FAnimNode_CustomProperty* GetCustomPropertyNode() const PURE_VIRTUAL(UAnimGraphNode_CustomProperty::GetCustomPropertyNode, return nullptr;);
diff --git a/Engine/Source/Editor/AnimGraph/Private/AnimGraphNode_CustomProperty.cpp b/Engine/Source/Editor/AnimGraph/Private/AnimGraphNode_CustomProperty.cpp
--- a/Engine/Source/Editor/AnimGraph/Private/AnimGraphNode_CustomProperty.cpp
+++ b/Engine/Source/Editor/AnimGraph/Private/AnimGraphNode_CustomProperty.cpp
@@ -99,112 +99,6 @@ void UAnimGraphNode_CustomProperty::ReallocatePinsDuringReconstruction(TArray<UE
 	}
 }
 
-void UAnimGraphNode_CustomProperty::CustomizeDetails(IDetailLayoutBuilder& DetailBuilder)
-{
-	Super::CustomizeDetails(DetailBuilder);
-
-	// We dont allow multi-select here
-	if(DetailBuilder.GetSelectedObjects().Num() > 1)
-	{
-		DetailBuilder.HideCategory(TEXT("Settings"));
-		return;
-	}
-
-	TArray<UProperty*> ExposableProperties;
-	GetExposableProperties(ExposableProperties);
-
-	if(ExposableProperties.Num() > 0)
-	{
-		IDetailCategoryBuilder& CategoryBuilder = DetailBuilder.EditCategory(FName(TEXT("Exposable Properties")));
-
-		FDetailWidgetRow& HeaderWidgetRow = CategoryBuilder.AddCustomRow(LOCTEXT("ExposeAll", "Expose All"));
-		
-		HeaderWidgetRow.NameContent()
-		[
-			SNew(STextBlock)
-			.Text(LOCTEXT("PropertyName", "Name"))
-			.Font(IDetailLayoutBuilder::GetDetailFontBold())
-		];
-
-		HeaderWidgetRow.ValueContent()
-		[
-			SNew(SHorizontalBox)
-			+SHorizontalBox::Slot()
-			.AutoWidth()
-			.VAlign(VAlign_Center)
-			[
-				SNew(STextBlock)
-				.Text(LOCTEXT("ExposeAllPropertyValue", "Expose All"))
-				.Font(IDetailLayoutBuilder::GetDetailFontBold())
-			]
-			+SHorizontalBox::Slot()
-			.FillWidth(1.0f)
-			.HAlign(HAlign_Right)
-			.VAlign(VAlign_Center)
-			[
-				SNew(SCheckBox)
-				.IsChecked_UObject(this, &UAnimGraphNode_CustomProperty::AreAllPropertiesExposed)
-				.OnCheckStateChanged_UObject(this, &UAnimGraphNode_CustomProperty::OnPropertyExposeAllCheckboxChanged)
-			]
-		];
-
-		for(UProperty* Property : ExposableProperties)
-		{
-			FDetailWidgetRow& PropertyWidgetRow = CategoryBuilder.AddCustomRow(FText::FromString(Property->GetName()));
-
-			FName PropertyName = Property->GetFName();
-			FText PropertyTypeText = GetPropertyTypeText(Property);
-
-			FFormatNamedArguments Args;
-			Args.Add(TEXT("PropertyName"), FText::FromName(PropertyName));
-			Args.Add(TEXT("PropertyType"), PropertyTypeText);
-
-			FText TooltipText = FText::Format(LOCTEXT("PropertyTooltipText", "{PropertyName}\nType: {PropertyType}"), Args);
-
-			PropertyWidgetRow.NameContent()
-			[
-				SNew(STextBlock)
-				.Text(FText::FromString(Property->GetName()))
-				.ToolTipText(TooltipText)
-				.Font(IDetailLayoutBuilder::GetDetailFont())
-			];
-
-			PropertyWidgetRow.ValueContent()
-			[
-				SNew(SHorizontalBox)
-				+SHorizontalBox::Slot()
-				.VAlign(VAlign_Center)
-				.AutoWidth()
-				[
-					SNew(STextBlock)
-					.Text(LOCTEXT("ExposePropertyValue", "Expose:"))
-					.Font(IDetailLayoutBuilder::GetDetailFont())
-				]
-				+SHorizontalBox::Slot()
-				.FillWidth(1.0f)
-				.HAlign(HAlign_Right)
-				.VAlign(VAlign_Center)
-				[
-					SNew(SCheckBox)
-					.IsChecked_UObject(this, &UAnimGraphNode_CustomProperty::IsPropertyExposed, PropertyName)
-					.OnCheckStateChanged_UObject(this, &UAnimGraphNode_CustomProperty::OnPropertyExposeCheckboxChanged, PropertyName)
-				]
-			];
-		}
-	}
-
-	IDetailCategoryBuilder& CategoryBuilder = DetailBuilder.EditCategory(FName(TEXT("Settings")));
-
-	// Customize InstanceClass
-	{
-		TSharedRef<IPropertyHandle> ClassHandle = DetailBuilder.GetProperty(TEXT("Node.InstanceClass"), GetClass());
-		if(ClassHandle->IsValidHandle())
-		{
-			ClassHandle->SetOnPropertyValueChanged(FSimpleDelegate::CreateUObject(this, &UAnimGraphNode_CustomProperty::OnStructuralPropertyChanged, &DetailBuilder));
-		}
-	}
-}
-
 void UAnimGraphNode_CustomProperty::GetInstancePinProperty(const UClass* InOwnerInstanceClass, UEdGraphPin* InInputPin, UProperty*& OutProperty)
 {
 	// The actual name of the instance property
diff --git a/Engine/Source/Editor/AnimGraph/Private/AnimGraphNode_SubInstanceBase.cpp b/Engine/Source/Editor/AnimGraph/Private/AnimGraphNode_SubInstanceBase.cpp
--- a/Engine/Source/Editor/AnimGraph/Private/AnimGraphNode_SubInstanceBase.cpp
+++ b/Engine/Source/Editor/AnimGraph/Private/AnimGraphNode_SubInstanceBase.cpp
@@ -234,6 +234,96 @@ void UAnimGraphNode_SubInstanceBase::CustomizeDetails(IDetailLayoutBuilder& Deta
 		return;
 	}
 
+	// We dont allow multi-select here
+	if(DetailBuilder.GetSelectedObjects().Num() > 1)
+	{
+		DetailBuilder.HideCategory(TEXT("Settings"));
+		return;
+	}
+
+	TArray<UProperty*> ExposableProperties;
+	GetExposableProperties(ExposableProperties);
+
+	if(ExposableProperties.Num() > 0)
+	{
+		IDetailCategoryBuilder& CategoryBuilder = DetailBuilder.EditCategory(FName(TEXT("Exposable Properties")));
+
+		FDetailWidgetRow& HeaderWidgetRow = CategoryBuilder.AddCustomRow(LOCTEXT("ExposeAll", "Expose All"));
+		
+		HeaderWidgetRow.NameContent()
+		[
+			SNew(STextBlock)
+			.Text(LOCTEXT("PropertyName", "Name"))
+			.Font(IDetailLayoutBuilder::GetDetailFontBold())
+		];
+
+		HeaderWidgetRow.ValueContent()
+		[
+			SNew(SHorizontalBox)
+			+SHorizontalBox::Slot()
+			.AutoWidth()
+			.VAlign(VAlign_Center)
+			[
+				SNew(STextBlock)
+				.Text(LOCTEXT("ExposeAllPropertyValue", "Expose All"))
+				.Font(IDetailLayoutBuilder::GetDetailFontBold())
+			]
+			+SHorizontalBox::Slot()
+			.FillWidth(1.0f)
+			.HAlign(HAlign_Right)
+			.VAlign(VAlign_Center)
+			[
+				SNew(SCheckBox)
+				.IsChecked_UObject(this, &UAnimGraphNode_CustomProperty::AreAllPropertiesExposed)
+				.OnCheckStateChanged_UObject(this, &UAnimGraphNode_CustomProperty::OnPropertyExposeAllCheckboxChanged)
+			]
+		];
+
+		for(UProperty* Property : ExposableProperties)
+		{
+			FDetailWidgetRow& PropertyWidgetRow = CategoryBuilder.AddCustomRow(FText::FromString(Property->GetName()));
+
+			FName PropertyName = Property->GetFName();
+			FText PropertyTypeText = GetPropertyTypeText(Property);
+
+			FFormatNamedArguments Args;
+			Args.Add(TEXT("PropertyName"), FText::FromName(PropertyName));
+			Args.Add(TEXT("PropertyType"), PropertyTypeText);
+
+			FText TooltipText = FText::Format(LOCTEXT("PropertyTooltipText", "{PropertyName}\nType: {PropertyType}"), Args);
+
+			PropertyWidgetRow.NameContent()
+			[
+				SNew(STextBlock)
+				.Text(FText::FromString(Property->GetName()))
+				.ToolTipText(TooltipText)
+				.Font(IDetailLayoutBuilder::GetDetailFont())
+			];
+
+			PropertyWidgetRow.ValueContent()
+			[
+				SNew(SHorizontalBox)
+				+SHorizontalBox::Slot()
+				.VAlign(VAlign_Center)
+				.AutoWidth()
+				[
+					SNew(STextBlock)
+					.Text(LOCTEXT("ExposePropertyValue", "Expose:"))
+					.Font(IDetailLayoutBuilder::GetDetailFont())
+				]
+				+SHorizontalBox::Slot()
+				.FillWidth(1.0f)
+				.HAlign(HAlign_Right)
+				.VAlign(VAlign_Center)
+				[
+					SNew(SCheckBox)
+					.IsChecked_UObject(this, &UAnimGraphNode_CustomProperty::IsPropertyExposed, PropertyName)
+					.OnCheckStateChanged_UObject(this, &UAnimGraphNode_CustomProperty::OnPropertyExposeCheckboxChanged, PropertyName)
+				]
+			];
+		}
+	}
+
 	IDetailCategoryBuilder& CategoryBuilder = DetailBuilder.EditCategory(FName(TEXT("Settings")));
 
 	// Customize InstanceClass
