Fixed anim blueprint editor viewport-linked skeletal meshes not animating

diff --git a/Engine/Source/Editor/AnimGraph/Private/AnimPreviewInstance.cpp b/Engine/Source/Editor/AnimGraph/Private/AnimPreviewInstance.cpp
--- a/Engine/Source/Editor/AnimGraph/Private/AnimPreviewInstance.cpp
+++ b/Engine/Source/Editor/AnimGraph/Private/AnimPreviewInstance.cpp
@@ -144,6 +144,11 @@ void FAnimPreviewInstanceProxy::PreUpdate(UAnimInstance* InAnimInstance, float D
 {
 	FAnimSingleNodeInstanceProxy::PreUpdate(InAnimInstance, DeltaSeconds);
 
+	if (CopyPoseNode.SourceMeshComponent.IsValid())
+	{
+		CopyPoseNode.PreUpdate(InAnimInstance);
+	}
+
 	if (!bForceRetargetBasePose)
 	{
 		CurveSource.PreUpdate(InAnimInstance);
