Fixed incorrect spawn location for dragged assets in anim blueprints

diff --git a/Engine/Source/Editor/AnimGraph/Private/AnimationGraphSchema.cpp b/Engine/Source/Editor/AnimGraph/Private/AnimationGraphSchema.cpp
--- a/Engine/Source/Editor/AnimGraph/Private/AnimationGraphSchema.cpp
+++ b/Engine/Source/Editor/AnimGraph/Private/AnimationGraphSchema.cpp
@@ -367,7 +367,9 @@ void UAnimationGraphSchema::DroppedAssetsOnPin(const TArray<FAssetData>& Assets,
 	UAnimationAsset* Asset = FAssetData::GetFirstAsset<UAnimationAsset>(Assets);
 	if ((Asset != NULL) && (Pin != NULL))
 	{
-		SpawnNodeFromAsset(Asset, GraphPosition, Pin->GetOwningNode()->GetGraph(), Pin);
+		// Don't have access to bounding information for node, using fixed offset that should work for most cases.
+		const FVector2D FixedOffset(-250.0f, 50.0f);
+		SpawnNodeFromAsset(Asset, GraphPosition + FixedOffset, Pin->GetOwningNode()->GetGraph(), Pin);
 	}
 }
 
diff --git a/Engine/Source/Editor/GraphEditor/Private/SGraphPin.cpp b/Engine/Source/Editor/GraphEditor/Private/SGraphPin.cpp
--- a/Engine/Source/Editor/GraphEditor/Private/SGraphPin.cpp
+++ b/Engine/Source/Editor/GraphEditor/Private/SGraphPin.cpp
@@ -781,7 +781,7 @@ FReply SGraphPin::OnDrop( const FGeometry& MyGeometry, const FDragDropEvent& Dra
 			TSharedPtr<FAssetDragDropOp> AssetOp = StaticCastSharedPtr<FAssetDragDropOp>(Operation);
 			if (AssetOp->HasAssets())
 			{
-				Node->GetSchema()->DroppedAssetsOnPin(AssetOp->GetAssets(), DragDropEvent.GetScreenSpacePosition(), GraphPinObj);
+				Node->GetSchema()->DroppedAssetsOnPin(AssetOp->GetAssets(), NodeWidget->GetPosition() + MyGeometry.Position, GraphPinObj);
 			}
 		}
 		return FReply::Handled();
