Improving the PoseDriver by

* adding debug drawing for the cones used for each target
* adding a new normalization mode using a median cone
* adding overrides for distance method + function type per target
* improving the reloading of the pose list to maintain settings for known poses

diff --git a/Engine/Source/Editor/AnimGraph/Classes/AnimGraphNode_PoseDriver.h b/Engine/Source/Editor/AnimGraph/Classes/AnimGraphNode_PoseDriver.h
--- a/Engine/Source/Editor/AnimGraph/Classes/AnimGraphNode_PoseDriver.h
+++ b/Engine/Source/Editor/AnimGraph/Classes/AnimGraphNode_PoseDriver.h
@@ -19,6 +19,14 @@ class ANIMGRAPH_API UAnimGraphNode_PoseDriver : public UAnimGraphNode_PoseHandle
 	UPROPERTY(EditAnywhere, Category = Settings)
 	FAnimNode_PoseDriver Node;
 
+	/** Length of axis in world units used for debug drawing */
+	UPROPERTY(EditAnywhere, Category = Debugging, meta = (DefaultValue = "20.0", UIMin = "1.0", UIMax = "100.0"))
+	float AxisLength;
+
+	/** Number of subdivisions / lines used when debug drawing a cone */
+	UPROPERTY(EditAnywhere, Category = Debugging, meta = (DefaultValue = "32", UIMin = "6", UIMax = "128"))
+	int32 ConeSubdivision;
+
 	/** Used to indicate selected target to edit mode drawing */
 	int32 SelectedTargetIndex;
 	/** Delegate to call when selection changes */
@@ -41,6 +49,9 @@ public:
 	/** Reallocates transforms arrays as necessary to accommodate source bones */
 	void ReserveTargetTransforms();
 
+	/** Return the color for a given weight. Used for Details and EditMode */
+	FLinearColor GetColorFromWeight(float InWeight);
+
 	/** Used to refer back to preview instance in anim tools */
 	UPROPERTY(Transient)
 	USkeletalMeshComponent* LastPreviewComponent;
diff --git a/Engine/Source/Editor/AnimGraph/Private/AnimGraphNode_PoseDriver.cpp b/Engine/Source/Editor/AnimGraph/Private/AnimGraphNode_PoseDriver.cpp
--- a/Engine/Source/Editor/AnimGraph/Private/AnimGraphNode_PoseDriver.cpp
+++ b/Engine/Source/Editor/AnimGraph/Private/AnimGraphNode_PoseDriver.cpp
@@ -42,6 +42,8 @@ UAnimGraphNode_PoseDriver::UAnimGraphNode_PoseDriver(const FObjectInitializer& O
 	:Super(ObjectInitializer)
 {
 	SelectedTargetIndex = INDEX_NONE;
+	AxisLength = 20.0f;
+	ConeSubdivision = 32;
 }
 
 FText UAnimGraphNode_PoseDriver::GetTooltipText() const
@@ -210,6 +212,10 @@ void UAnimGraphNode_PoseDriver::CopyNodeDataToPreviewNode(FAnimNode_Base* InPrev
 	PreviewPoseDriver->RBFParams.DistanceMethod = Node.RBFParams.DistanceMethod;
 	PreviewPoseDriver->RBFParams.TwistAxis = Node.RBFParams.TwistAxis;
 	PreviewPoseDriver->RBFParams.WeightThreshold = Node.RBFParams.WeightThreshold;
+	PreviewPoseDriver->RBFParams.NormalizeMethod = Node.RBFParams.NormalizeMethod;
+	PreviewPoseDriver->RBFParams.MedianReference = Node.RBFParams.MedianReference;
+	PreviewPoseDriver->RBFParams.MedianMin = Node.RBFParams.MedianMin;
+	PreviewPoseDriver->RBFParams.MedianMax = Node.RBFParams.MedianMax;
 	PreviewPoseDriver->DriveOutput = Node.DriveOutput;
 	PreviewPoseDriver->DriveSource = Node.DriveSource;
 	PreviewPoseDriver->PoseTargets = Node.PoseTargets;
@@ -261,6 +267,13 @@ void UAnimGraphNode_PoseDriver::CopyTargetsFromPoseAsset()
 {
 	UPoseAsset* PoseAsset = Node.PoseAsset;
 
+	// store the previous targets for re-use
+	TMap< FName, FPoseDriverTarget > PreviousTargets;
+	for (int32 i = 0; i < Node.PoseTargets.Num(); i++)
+	{
+		PreviousTargets.Add(Node.PoseTargets[i].DrivenName, Node.PoseTargets[i]);
+	}
+
 	if (PoseAsset && PoseAsset->GetSkeleton()) // Use PoseAsset here not CurrentPoseAsset, because we want to be able to run this on on nodes that have not been init'd yet
 	{
 		Node.PoseTargets.Empty();
@@ -313,6 +326,18 @@ void UAnimGraphNode_PoseDriver::CopyTargetsFromPoseAsset()
 				PoseTarget.BoneTransforms.Add(PoseTransform);
 			}
 
+			// re-apply the same setting in case we have seen this target before
+			const FPoseDriverTarget * PreviousTarget = PreviousTargets.Find(PoseTarget.DrivenName);
+			if (PreviousTarget)
+			{
+				PoseTarget.TargetScale = PreviousTarget->TargetScale;
+				PoseTarget.bApplyCustomCurve = PreviousTarget->bApplyCustomCurve;
+				PoseTarget.CustomCurve = PreviousTarget->CustomCurve;
+				PoseTarget.DistanceMethod = PreviousTarget->DistanceMethod;
+				PoseTarget.FunctionType = PreviousTarget->FunctionType;
+				PoseTarget.bIsHidden = PreviousTarget->bIsHidden;
+			}
+
 			Node.PoseTargets.Add(PoseTarget);
 		}
 
@@ -338,6 +363,11 @@ void UAnimGraphNode_PoseDriver::ReserveTargetTransforms()
 	}
 }
 
+FLinearColor UAnimGraphNode_PoseDriver::GetColorFromWeight(float InWeight)
+{
+	return FMath::Lerp(FLinearColor::Blue.Desaturate(0.5), FLinearColor::Red, InWeight);
+}
+
 void UAnimGraphNode_PoseDriver::AutoSetTargetScales(float &OutMaxDistance)
 {
 	TArray<FRBFTarget> RBFTargets;
diff --git a/Engine/Source/Editor/AnimGraph/Private/EditModes/PoseDriverEditMode.cpp b/Engine/Source/Editor/AnimGraph/Private/EditModes/PoseDriverEditMode.cpp
--- a/Engine/Source/Editor/AnimGraph/Private/EditModes/PoseDriverEditMode.cpp
+++ b/Engine/Source/Editor/AnimGraph/Private/EditModes/PoseDriverEditMode.cpp
@@ -27,11 +27,6 @@ void FPoseDriverEditMode::ExitMode()
 	FAnimNodeEditMode::ExitMode();
 }
 
-static FLinearColor GetColorFromWeight(float InWeight)
-{
-	return FMath::Lerp(FLinearColor::White, FLinearColor::Red, InWeight);
-}
-
 /** Hit proxy for selecting targets */
 struct HPDTargetHitProxy : public HHitProxy
 {
@@ -61,9 +56,12 @@ void FPoseDriverEditMode::Render(const FSceneView* View, FViewport* Viewport, FP
 	GraphNode->LastPreviewComponent = SkelComp;
 
 	static const float DrawLineWidth = 0.1f;
-	static const float DrawAxisLength = 20.f;
 	static const float DrawPosSize = 2.f;
+	float DrawAxisLength = GraphNode->AxisLength;
+	int32 DrawConeSubdivision = GraphNode->ConeSubdivision;
 
+	TArray<FRBFTarget> RBFTargets;
+	RuntimeNode->GetRBFTargets(RBFTargets);
 
 	// Iterate over each bone in the 'source bones' array
 	for (int32 SourceIdx=0; SourceIdx< RuntimeNode->SourceBones.Num(); SourceIdx++)
@@ -103,9 +101,49 @@ void FPoseDriverEditMode::Render(const FSceneView* View, FViewport* Viewport, FP
 				// Rotation drawing
 				if (RuntimeNode->DriveSource == EPoseDriverSource::Rotation)
 				{
-					FVector LocalVec = SourceBoneTM.TransformVectorNoScale(RuntimeNode->RBFParams.GetTwistAxisVector());
-					FVector WorldVec = EvalSpaceTM.TransformVectorNoScale(LocalVec);
-					PDI->DrawLine(BonePos, BonePos + (WorldVec*DrawAxisLength), FLinearColor::Green, SDPG_Foreground, DrawLineWidth);
+					FVector LocalTwistVec = SourceBoneTM.TransformVectorNoScale(RuntimeNode->RBFParams.GetTwistAxisVector());
+					FVector WorldTwistVec = EvalSpaceTM.TransformVectorNoScale(LocalTwistVec);
+					PDI->DrawLine(BonePos, BonePos + (WorldTwistVec*DrawAxisLength), FLinearColor::Green, SDPG_Foreground, DrawLineWidth);
+
+					// draw the median cones
+					if (RuntimeNode->RBFParams.NormalizeMethod == ERBFNormalizeMethod::NormalizeWithinMedian)
+					{
+						const FVector& MedianRot = RuntimeNode->RBFParams.MedianReference;
+						LocalTwistVec = FRotator(MedianRot.X, MedianRot.Y, MedianRot.Z).RotateVector(RuntimeNode->RBFParams.GetTwistAxisVector());
+						WorldTwistVec = EvalSpaceTM.TransformVectorNoScale(LocalTwistVec);
+
+						FVector LocalSwingVec = FVector::CrossProduct(LocalTwistVec, FVector(1, 1, 1));
+						FVector WorldSwingVec = EvalSpaceTM.TransformVectorNoScale(LocalSwingVec);
+						WorldSwingVec.Normalize();
+
+						FQuat WorldTwistQ(WorldTwistVec, PI * 2.0f / float(DrawConeSubdivision));
+						FQuat WorldSwingMinQ(WorldSwingVec, FMath::DegreesToRadians(RuntimeNode->RBFParams.MedianMin));
+						FQuat WorldSwingMaxQ(WorldSwingVec, FMath::DegreesToRadians(RuntimeNode->RBFParams.MedianMax));
+
+						FVector FirstMinPositionOnCircle = WorldSwingMinQ.RotateVector(WorldTwistVec * DrawAxisLength);
+						FVector FirstMaxPositionOnCircle = WorldSwingMaxQ.RotateVector(WorldTwistVec * DrawAxisLength);
+						FVector LastMinPositionOnCircle = FirstMinPositionOnCircle;
+						FVector LastMaxPositionOnCircle = FirstMaxPositionOnCircle;
+
+						FLinearColor MinColor = FLinearColor::Yellow;
+						FLinearColor MaxColor = MinColor.Desaturate(0.5);
+
+						for (int32 i = 0; i < DrawConeSubdivision; i++)
+						{
+							FVector NextMinPositionOnCircle = WorldTwistQ.RotateVector(LastMinPositionOnCircle);
+							FVector NextMaxPositionOnCircle = WorldTwistQ.RotateVector(LastMaxPositionOnCircle);
+							PDI->DrawLine(BonePos, BonePos + NextMinPositionOnCircle, MinColor, SDPG_Foreground, DrawLineWidth);
+							PDI->DrawLine(BonePos, BonePos + NextMaxPositionOnCircle, MaxColor , SDPG_Foreground, DrawLineWidth);
+							PDI->DrawLine(BonePos + LastMinPositionOnCircle, BonePos + NextMinPositionOnCircle, MinColor, SDPG_Foreground, DrawLineWidth);
+							PDI->DrawLine(BonePos + LastMaxPositionOnCircle, BonePos + NextMaxPositionOnCircle, MaxColor, SDPG_Foreground, DrawLineWidth);
+							PDI->DrawLine(BonePos + NextMinPositionOnCircle, BonePos + NextMaxPositionOnCircle, MaxColor, SDPG_Foreground, DrawLineWidth);
+
+							LastMinPositionOnCircle = NextMinPositionOnCircle;
+							LastMaxPositionOnCircle = NextMaxPositionOnCircle;
+						}
+						PDI->DrawLine(BonePos + LastMinPositionOnCircle, BonePos + FirstMinPositionOnCircle, MinColor, SDPG_Foreground, DrawLineWidth);
+						PDI->DrawLine(BonePos + LastMaxPositionOnCircle, BonePos + FirstMaxPositionOnCircle, MaxColor, SDPG_Foreground, DrawLineWidth);
+					}
 				}
 				// Translation drawing
 				else if (RuntimeNode->DriveSource == EPoseDriverSource::Translation)
@@ -113,13 +151,27 @@ void FPoseDriverEditMode::Render(const FSceneView* View, FViewport* Viewport, FP
 					FVector LocalPos = SourceBoneTM.GetTranslation();
 					FVector WorldPos = EvalSpaceTM.TransformPosition(LocalPos);
 					DrawWireDiamond(PDI, FTranslationMatrix(WorldPos), DrawPosSize, FLinearColor::Green, SDPG_Foreground);
+
+					// draw the median diamonds
+					if (RuntimeNode->RBFParams.NormalizeMethod == ERBFNormalizeMethod::NormalizeWithinMedian)
+					{
+						WorldPos = EvalSpaceTM.TransformPosition(RuntimeNode->RBFParams.MedianReference);
+
+						FLinearColor MinColor = FLinearColor::Yellow;
+						FLinearColor MaxColor = MinColor.Desaturate(0.5);
+
+						DrawWireDiamond(PDI, FTranslationMatrix(WorldPos), RuntimeNode->RBFParams.MedianMin, MinColor, SDPG_Foreground);
+						DrawWireDiamond(PDI, FTranslationMatrix(WorldPos), RuntimeNode->RBFParams.MedianMax, MaxColor, SDPG_Foreground);
+					}
 				}
 
 				// Build array of weight for every target
+				float TotalWeight = 0.f;
 				TArray<float> PerTargetWeights;
 				PerTargetWeights.AddZeroed(RuntimeNode->PoseTargets.Num());
 				for (const FRBFOutputWeight& Weight : RuntimeNode->OutputWeights)
 				{
+					TotalWeight += Weight.TargetWeight;
 					PerTargetWeights[Weight.TargetIndex] = Weight.TargetWeight;
 				}
 
@@ -128,6 +180,14 @@ void FPoseDriverEditMode::Render(const FSceneView* View, FViewport* Viewport, FP
 				{
 					// Check we have a target transform for this bone
 					const FPoseDriverTarget& PoseTarget = RuntimeNode->PoseTargets[TargetIdx];
+					const FRBFTarget& RBFTarget = RBFTargets[TargetIdx];
+
+					// skip hidden entries
+					if (PoseTarget.bIsHidden)
+					{
+						continue;
+					}
+
 					if (PoseTarget.BoneTransforms.IsValidIndex(SourceIdx))
 					{
 						const FPoseDriverTransform& TargetTM = PoseTarget.BoneTransforms[SourceIdx];
@@ -135,23 +195,44 @@ void FPoseDriverEditMode::Render(const FSceneView* View, FViewport* Viewport, FP
 						bool bSelected = (GraphNode->SelectedTargetIndex == TargetIdx);
 						float AxisLength = bSelected ? (DrawAxisLength * 1.5f) : DrawAxisLength;
 						float LineWidth = bSelected ? (DrawLineWidth * 3.f) : DrawLineWidth;
-						float PosSize = bSelected ? (DrawPosSize * 1.5f) : DrawPosSize;
+						float Radius = RuntimeNode->GetRadiusForTarget(RBFTarget);
+						
+						FLinearColor Color = TotalWeight <= 0.f ? FLinearColor::Black : GraphNode->GetColorFromWeight(PerTargetWeights[TargetIdx]);
 
 						PDI->SetHitProxy(new HPDTargetHitProxy(TargetIdx));
 
 						// Rotation drawing
 						if (RuntimeNode->DriveSource == EPoseDriverSource::Rotation)
 						{
-							FVector LocalVec = TargetTM.TargetRotation.RotateVector(RuntimeNode->RBFParams.GetTwistAxisVector());
-							FVector WorldVec = EvalSpaceTM.TransformVectorNoScale(LocalVec);
-							PDI->DrawLine(BonePos, BonePos + (WorldVec*AxisLength), GetColorFromWeight(PerTargetWeights[TargetIdx]), SDPG_Foreground, LineWidth);
+							FVector LocalTwistVec = TargetTM.TargetRotation.RotateVector(RuntimeNode->RBFParams.GetTwistAxisVector());
+							FVector WorldTwistVec = EvalSpaceTM.TransformVectorNoScale(LocalTwistVec);
+
+							FVector LocalSwingVec = FVector::CrossProduct(LocalTwistVec, FVector(1, 1, 1));
+							FVector WorldSwingVec = EvalSpaceTM.TransformVectorNoScale(LocalSwingVec);
+							WorldSwingVec.Normalize();
+
+							FQuat WorldTwistQ(WorldTwistVec, PI * 2.0f / float(DrawConeSubdivision));
+							FQuat WorldSwingQ(WorldSwingVec, FMath::DegreesToRadians(Radius));
+
+							FVector FirstPositionOnCircle = WorldSwingQ.RotateVector(WorldTwistVec * DrawAxisLength);
+							FVector LastPositionOnCircle = FirstPositionOnCircle;
+
+							for (int32 i = 0; i < DrawConeSubdivision; i++)
+							{
+								FVector NextPositionOnCircle = WorldTwistQ.RotateVector(LastPositionOnCircle);
+								PDI->DrawLine(BonePos, BonePos + NextPositionOnCircle, Color, SDPG_Foreground, DrawLineWidth);
+								PDI->DrawLine(BonePos + LastPositionOnCircle, BonePos + NextPositionOnCircle, Color, SDPG_Foreground, DrawLineWidth);
+
+								LastPositionOnCircle = NextPositionOnCircle;
+							}
+							PDI->DrawLine(BonePos + LastPositionOnCircle, BonePos + FirstPositionOnCircle, Color, SDPG_Foreground, DrawLineWidth);
 						}
 						// Translation drawing
 						else if (RuntimeNode->DriveSource == EPoseDriverSource::Translation)
 						{
 							FVector LocalPos = TargetTM.TargetTranslation;
 							FVector WorldPos = EvalSpaceTM.TransformPosition(LocalPos);
-							DrawWireDiamond(PDI, FTranslationMatrix(WorldPos), PosSize, GetColorFromWeight(PerTargetWeights[TargetIdx]), SDPG_Foreground, LineWidth);
+							DrawWireDiamond(PDI, FTranslationMatrix(WorldPos), Radius, Color, SDPG_Foreground, LineWidth);
 						}
 
 						PDI->SetHitProxy(nullptr);
diff --git a/Engine/Source/Editor/AnimGraph/Private/PoseDriverDetails.cpp b/Engine/Source/Editor/AnimGraph/Private/PoseDriverDetails.cpp
--- a/Engine/Source/Editor/AnimGraph/Private/PoseDriverDetails.cpp
+++ b/Engine/Source/Editor/AnimGraph/Private/PoseDriverDetails.cpp
@@ -23,12 +23,32 @@
 #define LOCTEXT_NAMESPACE "PoseDriverDetails"
 
 static const FName ColumnId_Target("Target");
+TArray< TSharedPtr<FName> > SPDD_TargetRow::DistanceMethodOptions;
+TArray< TSharedPtr<FName> > SPDD_TargetRow::FunctionTypeOptions;
 
 void SPDD_TargetRow::Construct(const FArguments& InArgs, const TSharedRef<STableViewBase>& InOwnerTableView)
 {
 	TargetInfoPtr = InArgs._TargetInfo;
 	PoseDriverDetailsPtr = InArgs._PoseDriverDetails;
 
+	if (DistanceMethodOptions.Num() == 0)
+	{
+		UEnum * Enum = FindObjectChecked<UEnum>(ANY_PACKAGE, TEXT("ERBFDistanceMethod"));
+		for (int32 i = 0; i < Enum->NumEnums() - 1; i++)
+		{
+			DistanceMethodOptions.Add(MakeShareable(new FName(*Enum->GetDisplayNameTextByIndex(i).ToString())));
+		}
+	}
+
+	if (FunctionTypeOptions.Num() == 0)
+	{
+		UEnum * Enum = FindObjectChecked<UEnum>(ANY_PACKAGE, TEXT("ERBFFunctionType"));
+		for (int32 i = 0; i < Enum->NumEnums() - 1; i++)
+		{
+			FunctionTypeOptions.Add(MakeShareable(new FName(*Enum->GetDisplayNameTextByIndex(i).ToString())));
+		}
+	}
+
 	// Register delegate so TargetInfo can trigger UI expansion
 	TSharedPtr<FPDD_TargetInfo> TargetInfo = TargetInfoPtr.Pin();
 	TargetInfo->ExpandTargetDelegate.AddSP(this, &SPDD_TargetRow::ExpandTargetInfo);
@@ -183,6 +203,84 @@ TSharedRef< SWidget > SPDD_TargetRow::GenerateWidgetForColumn(const FName& Colum
 							]
 						]
 
+						+ SHorizontalBox::Slot()
+						.VAlign(VAlign_Center)
+						.AutoWidth()
+						.Padding(FMargin(6, 0, 3, 0))
+						[
+							SNew(STextBlock)
+							.Text(LOCTEXT("IsHidden", "Hidden:"))
+						]
+
+						+ SHorizontalBox::Slot()
+						.VAlign(VAlign_Center)
+						.AutoWidth()
+						.Padding(FMargin(6, 0, 3, 0))
+						[
+							SNew(SCheckBox)
+							.IsChecked_Lambda([=]() { return IsHidden() ? ECheckBoxState::Checked : ECheckBoxState::Unchecked; })
+							.OnCheckStateChanged(this, &SPDD_TargetRow::OnIsHiddenChanged)
+							.Padding(FMargin(4.0f, 0.0f))
+							.ToolTipText(LOCTEXT("IsHiddenToolTip", "Define if this target should be hidden from debug drawing."))
+						]
+
+						+ SHorizontalBox::Slot()
+						.FillWidth(1)
+						[
+							SNew(SSpacer)
+						]
+
+					]
+
+					+ SVerticalBox::Slot()
+					.AutoHeight()
+					.Padding(2.0f)
+					.VAlign(VAlign_Fill)
+					[
+						SNew(SHorizontalBox)
+
+						+ SHorizontalBox::Slot()
+						.VAlign(VAlign_Center)
+						.AutoWidth()
+						.Padding(FMargin(6, 0, 3, 0))
+						[
+							SNew(STextBlock)
+							.Text(LOCTEXT("Override", "Override:"))
+						]
+					
+						+ SHorizontalBox::Slot()
+						.VAlign(VAlign_Center)
+						.AutoWidth()
+						.Padding(FMargin(6, 0, 3, 0))
+						[
+							SNew(SComboBox<TSharedPtr<FName>>)
+							.OptionsSource(&DistanceMethodOptions)
+							.OnGenerateWidget(this, &SPDD_TargetRow::MakeDrivenNameWidget)
+							.OnSelectionChanged(this, &SPDD_TargetRow::OnDistanceMethodChanged)
+							.Content()
+							[
+								SNew(STextBlock)
+								.Text(this, &SPDD_TargetRow::GetDistanceMethodAsText)
+							]
+						]
+
+						+ SHorizontalBox::Slot()
+						.VAlign(VAlign_Center)
+						.AutoWidth()
+						.Padding(FMargin(6, 0, 3, 0))
+						[
+							SNew(SComboBox<TSharedPtr<FName>>)
+							.Visibility_Lambda([=]() { return IsCustomCurveEnabled() ? EVisibility::Collapsed : EVisibility::Visible;  })
+							.OptionsSource(&FunctionTypeOptions)
+							.OnGenerateWidget(this, &SPDD_TargetRow::MakeDrivenNameWidget)
+							.OnSelectionChanged(this, &SPDD_TargetRow::OnFunctionTypeChanged)
+							.Content()
+							[
+								SNew(STextBlock)
+								.Text(this, &SPDD_TargetRow::GetFunctionTypeAsText)
+							]
+						]
+
 						+ SHorizontalBox::Slot()
 						.VAlign(VAlign_Center)
 						.AutoWidth()
@@ -209,8 +307,6 @@ TSharedRef< SWidget > SPDD_TargetRow::GenerateWidgetForColumn(const FName& Colum
 						[
 							SNew(SSpacer)
 						]
-
-
 					]
 
 					+ SVerticalBox::Slot()
@@ -317,7 +413,12 @@ UAnimGraphNode_PoseDriver* SPDD_TargetRow::GetPoseDriverGraphNode() const
 
 bool SPDD_TargetRow::IsEditingRotation() const
 {
-	return (GetPoseDriverGraphNode()->Node.DriveSource == EPoseDriverSource::Rotation);
+	UAnimGraphNode_PoseDriver* Driver = GetPoseDriverGraphNode();
+	if (Driver)
+	{
+		return (Driver->Node.DriveSource == EPoseDriverSource::Rotation);
+	}
+	return false;
 }
 
 void SPDD_TargetRow::NotifyTargetChanged()
@@ -423,6 +524,80 @@ void SPDD_TargetRow::OnApplyCustomCurveChanged(const ECheckBoxState NewCheckStat
 	}
 }
 
+bool SPDD_TargetRow::IsHidden() const
+{
+	FPoseDriverTarget* Target = GetTarget();
+	if (Target)
+	{
+		return Target->bIsHidden;
+	}
+	return false;
+}
+
+
+void SPDD_TargetRow::OnIsHiddenChanged(const ECheckBoxState NewCheckState)
+{
+	FPoseDriverTarget* Target = GetTarget();
+	if (Target)
+	{
+		Target->bIsHidden= (NewCheckState == ECheckBoxState::Checked);
+		NotifyTargetChanged();
+	}
+}
+
+FText SPDD_TargetRow::GetDistanceMethodAsText() const
+{
+	const FPoseDriverTarget* Target = GetTarget();
+	if (Target)
+	{
+		return FText::FromName(*DistanceMethodOptions[(int32)Target->DistanceMethod]);
+	}
+	return FText();
+}
+
+void SPDD_TargetRow::OnDistanceMethodChanged(TSharedPtr<FName> InItem, ESelectInfo::Type SelectionType)
+{
+	FPoseDriverTarget* Target = GetTarget();
+	if (Target)
+	{
+		for (int32 i = 0; i < DistanceMethodOptions.Num(); i++)
+		{
+			if (InItem->Compare(*DistanceMethodOptions[i]) == 0)
+			{
+				Target->DistanceMethod = (ERBFDistanceMethod)i;
+				NotifyTargetChanged();
+				return;
+			}
+		}
+	}
+}
+
+FText SPDD_TargetRow::GetFunctionTypeAsText() const
+{
+	const FPoseDriverTarget* Target = GetTarget();
+	if (Target)
+	{
+		return FText::FromName(*FunctionTypeOptions[(int32)Target->FunctionType]);
+	}
+	return FText();
+}
+
+void SPDD_TargetRow::OnFunctionTypeChanged(TSharedPtr<FName> InItem, ESelectInfo::Type SelectionType)
+{
+	FPoseDriverTarget* Target = GetTarget();
+	if (Target)
+	{
+		for (int32 i = 0; i < FunctionTypeOptions.Num(); i++)
+		{
+			if (InItem->Compare(*FunctionTypeOptions[i]) == 0)
+			{
+				Target->FunctionType = (ERBFFunctionType)i;
+				NotifyTargetChanged();
+				return;
+			}
+		}
+	}
+}
 
 FText SPDD_TargetRow::GetDrivenNameText() const
 {
@@ -463,10 +638,14 @@ FText SPDD_TargetRow::GetTargetWeightText() const
 
 FSlateColor SPDD_TargetRow::GetWeightBarColor() const
 {
-	return FMath::Lerp(FLinearColor::White, FLinearColor::Red, GetTargetWeight());
+	UAnimGraphNode_PoseDriver* PoseDriver = GetPoseDriverGraphNode();
+	if (PoseDriver)
+	{
+		return PoseDriver->GetColorFromWeight(GetTargetWeight());
+	}
+	return FLinearColor::Black;
 }
 
-
 void SPDD_TargetRow::SetTranslation(float NewTrans, int32 BoneIndex, EAxis::Type Axis)
 {
 	FPoseDriverTarget* Target = GetTarget();
diff --git a/Engine/Source/Editor/AnimGraph/Private/PoseDriverDetails.h b/Engine/Source/Editor/AnimGraph/Private/PoseDriverDetails.h
--- a/Engine/Source/Editor/AnimGraph/Private/PoseDriverDetails.h
+++ b/Engine/Source/Editor/AnimGraph/Private/PoseDriverDetails.h
@@ -93,6 +93,14 @@ public:
 	bool IsCustomCurveEnabled() const;
 	void OnApplyCustomCurveChanged(const ECheckBoxState NewCheckState);
 
+	bool IsHidden() const;
+	void OnIsHiddenChanged(const ECheckBoxState NewCheckState);
+
+	FText GetDistanceMethodAsText() const;
+	void OnDistanceMethodChanged(TSharedPtr<FName> InItem, ESelectInfo::Type SelectionType);
+	FText GetFunctionTypeAsText() const;
+	void OnFunctionTypeChanged(TSharedPtr<FName> InItem, ESelectInfo::Type SelectionType);
+
 	/** Remove this target from  */
 	void RemoveTarget();
 
@@ -116,6 +124,9 @@ public:
 
 	/** Info that this widget represents */
 	TWeakPtr<FPDD_TargetInfo> TargetInfoPtr;
+
+	static TArray< TSharedPtr<FName> > DistanceMethodOptions;
+	static TArray< TSharedPtr<FName> > FunctionTypeOptions;
 };
 
 
diff --git a/Engine/Source/Runtime/AnimGraphRuntime/Private/AnimNodes/AnimNode_PoseDriver.cpp b/Engine/Source/Runtime/AnimGraphRuntime/Private/AnimNodes/AnimNode_PoseDriver.cpp
--- a/Engine/Source/Runtime/AnimGraphRuntime/Private/AnimNodes/AnimNode_PoseDriver.cpp
+++ b/Engine/Source/Runtime/AnimGraphRuntime/Private/AnimNodes/AnimNode_PoseDriver.cpp
@@ -138,6 +138,11 @@ void FAnimNode_PoseDriver::GatherDebugData(FNodeDebugData& DebugData)
 	SourcePose.GatherDebugData(DebugData.BranchFlow(1.f));
 }
 
+float FAnimNode_PoseDriver::GetRadiusForTarget(const FRBFTarget& Target) const
+{
+	return FRBFSolver::GetRadiusForTarget(Target, RBFParams);
+}
+
 bool FAnimNode_PoseDriver::IsBoneDriven(FName BoneName) const
 {
 	// If not filtering, drive all the bones
@@ -196,6 +201,8 @@ void FAnimNode_PoseDriver::GetRBFTargets(TArray<FRBFTarget>& OutTargets) const
 		RBFTarget.ScaleFactor = PoseTarget.TargetScale;
 		RBFTarget.bApplyCustomCurve = PoseTarget.bApplyCustomCurve;
 		RBFTarget.CustomCurve = PoseTarget.CustomCurve;
+		RBFTarget.DistanceMethod = PoseTarget.DistanceMethod;
+		RBFTarget.FunctionType = PoseTarget.FunctionType;
 	}
 }
 
diff --git a/Engine/Source/Runtime/AnimGraphRuntime/Private/RBF/RBFSolver.cpp b/Engine/Source/Runtime/AnimGraphRuntime/Private/RBF/RBFSolver.cpp
--- a/Engine/Source/Runtime/AnimGraphRuntime/Private/RBF/RBFSolver.cpp
+++ b/Engine/Source/Runtime/AnimGraphRuntime/Private/RBF/RBFSolver.cpp
@@ -48,6 +48,10 @@ FRBFParams::FRBFParams()
 	, DistanceMethod(ERBFDistanceMethod::Euclidean)
 	, TwistAxis(EBoneAxis::BA_X)
 	, WeightThreshold(KINDA_SMALL_NUMBER)
+	, NormalizeMethod(ERBFNormalizeMethod::OnlyNormalizeAboveOne)
+	, MedianReference(FVector(0, 0, 0))
+	, MedianMin(45.0f)
+	, MedianMax(60.0f)
 {
 
 }
@@ -68,12 +72,14 @@ FVector FRBFParams::GetTwistAxisVector() const
 
 //////////////////////////////////////////////////////////////////////////
 
-float FRBFSolver::FindDistanceBetweenEntries(const FRBFEntry& A, const FRBFEntry& B, const FRBFParams& Params)
+float FRBFSolver::FindDistanceBetweenEntries(const FRBFEntry& A, const FRBFEntry& B, const FRBFParams& Params, ERBFDistanceMethod OverrideMethod)
 {
 	check(A.GetDimensions() == B.GetDimensions());
 
+	ERBFDistanceMethod DistanceMethod = OverrideMethod == ERBFDistanceMethod::DefaultMethod ? Params.DistanceMethod : OverrideMethod;
+
 	// Simple n-dimensional distance
-	if (Params.DistanceMethod == ERBFDistanceMethod::Euclidean)
+	if (DistanceMethod == ERBFDistanceMethod::Euclidean)
 	{
 		float DistSqr = 0.f;
 
@@ -85,7 +91,7 @@ float FRBFSolver::FindDistanceBetweenEntries(const FRBFEntry& A, const FRBFEntry
 		return FMath::Sqrt(DistSqr);
 	}
 	// Treat values as sequence of eulers - find quat distance between each pair, then sqrt-sum-of-squares of those
-	else if (Params.DistanceMethod == ERBFDistanceMethod::Quaternion)
+	else if (DistanceMethod == ERBFDistanceMethod::Quaternion)
 	{
 		float DistSqr = 0.f;
 
@@ -99,7 +105,7 @@ float FRBFSolver::FindDistanceBetweenEntries(const FRBFEntry& A, const FRBFEntry
 		return FMath::Sqrt(DistSqr);
 	}
 	// Treat values as sequence of eulers - find 'swing' distance between each pair, then sqrt-sum-of-squares of those
-	else if(Params.DistanceMethod == ERBFDistanceMethod::SwingAngle)
+	else if(DistanceMethod == ERBFDistanceMethod::SwingAngle || DistanceMethod == ERBFDistanceMethod::DefaultMethod)
 	{
 		float DistSqr = 0.f;
 
@@ -141,31 +147,32 @@ void FRBFSolver::Solve(const FRBFParams& Params, const TArray<FRBFTarget>& Targe
 	for (int32 TargetIdx = 0; TargetIdx < Targets.Num(); TargetIdx++)
 	{
 		const FRBFTarget& Target = Targets[TargetIdx];
+		ERBFFunctionType FunctionType = Target.FunctionType == ERBFFunctionType::DefaultFunction ? Params.Function : Target.FunctionType;
 
 		// Find distance
-		const float Distance = FindDistanceBetweenEntries(Target, Input, Params);
-		const float Scaling = FMath::Max(Params.Radius * Target.ScaleFactor, KINDA_SMALL_NUMBER);
+		const float Distance = FindDistanceBetweenEntries(Target, Input, Params, Target.DistanceMethod);
+		const float Scaling = GetRadiusForTarget(Target, Params);
 		const float X = Distance / Scaling;
 
 		// Evaluate radial basis function to find weight
 		float Weight = 0.f;
-		if (Params.Function == ERBFFunctionType::Gaussian)
+		if (FunctionType == ERBFFunctionType::Gaussian)
 		{
 			Weight = FMath::Exp(-(X * X));
 		}
-		else if (Params.Function == ERBFFunctionType::Exponential)
+		else if (FunctionType == ERBFFunctionType::Exponential)
 		{
 			Weight = 1.f / FMath::Exp(X);
 		}
-		else if (Params.Function == ERBFFunctionType::Linear)
+		else if (FunctionType == ERBFFunctionType::Linear || FunctionType == ERBFFunctionType::DefaultFunction)
 		{
 			Weight = FMath::Max(1.f - X, 0.f);
 		}
-		else if (Params.Function == ERBFFunctionType::Cubic)
+		else if (FunctionType == ERBFFunctionType::Cubic)
 		{
 			Weight = FMath::Max(1.f - (X * X * X), 0.f);
 		}
-		else if (Params.Function == ERBFFunctionType::Quintic)
+		else if (FunctionType == ERBFFunctionType::Quintic)
 		{
 			Weight = FMath::Max(1.f - (X * X * X * X * X), 0.f);
 		}
@@ -186,8 +193,56 @@ void FRBFSolver::Solve(const FRBFParams& Params, const TArray<FRBFTarget>& Targe
 	// Only normalize and apply if we got some kind of weight
 	if (TotalWeight > KINDA_SMALL_NUMBER)
 	{
-		// If total is >1, we renormalize
-		const float WeightScale = (TotalWeight > 1.f) ? 1.f / TotalWeight : 1.f;
+		float WeightScale = 1.f;
+		if (TotalWeight > 1.f)
+		{
+			WeightScale = 1.f / TotalWeight;
+		}
+		else
+		{
+			switch (Params.NormalizeMethod)
+			{
+				case ERBFNormalizeMethod::OnlyNormalizeAboveOne:
+				{
+					break;
+				}
+				case ERBFNormalizeMethod::AlwaysNormalize:
+				{
+					WeightScale = 1.f / TotalWeight;
+					break;
+				}
+				case ERBFNormalizeMethod::NormalizeWithinMedian:
+				{
+					if (Params.MedianMax < Params.MedianMin)
+					{
+						break;
+					}
+
+					FRBFEntry MedianEntry;
+					while (Input.GetDimensions() > MedianEntry.GetDimensions())
+					{
+						MedianEntry.AddFromVector(Params.MedianReference);
+					}
+					
+					float MedianDistance = FindDistanceBetweenEntries(Input, MedianEntry, Params);
+					if (MedianDistance > Params.MedianMax)
+					{
+						break;
+					}
+					if (MedianDistance <= Params.MedianMin)
+					{
+						WeightScale = 1.f / TotalWeight;
+						break;
+					}
+
+					float Bias = FMath::Clamp<float>((MedianDistance - Params.MedianMin) / (Params.MedianMax - Params.MedianMin), 0.f, 1.f);
+					WeightScale = FMath::Lerp<float>(1.f / TotalWeight, 1.f, Bias);
+					break;
+				}
+			}
+		}
+		
+		/// TotalWeight : (Params.bNormalizeWeightsBelowSumOfOne ? 1.f / TotalWeight : 1.f);
 		for (int32 TargetIdx = 0; TargetIdx < Targets.Num(); TargetIdx++)
 		{
 			float NormalizedWeight = AllWeights[TargetIdx] * WeightScale;
@@ -222,7 +277,7 @@ bool FRBFSolver::FindTargetNeighbourDistances(const FRBFParams& Params, const TA
 				if (OtherTargetIdx != TargetIdx) // If not ourself..
 				{
 					// Get distance between poses
-					float Dist = FindDistanceBetweenEntries(Targets[TargetIdx], Targets[OtherTargetIdx], Params);
+					float Dist = FindDistanceBetweenEntries(Targets[TargetIdx], Targets[OtherTargetIdx], Params, Targets[TargetIdx].DistanceMethod);
 					NearestDist = FMath::Min(Dist, NearestDist);
 				}
 			}
@@ -238,3 +293,8 @@ bool FRBFSolver::FindTargetNeighbourDistances(const FRBFParams& Params, const TA
 		return false;
 	}
 }
+
+float FRBFSolver::GetRadiusForTarget(const FRBFTarget& Target, const FRBFParams& Params)
+{
+	return FMath::Max(Params.Radius * Target.ScaleFactor, KINDA_SMALL_NUMBER);
+}
diff --git a/Engine/Source/Runtime/AnimGraphRuntime/Public/AnimNodes/AnimNode_PoseDriver.h b/Engine/Source/Runtime/AnimGraphRuntime/Public/AnimNodes/AnimNode_PoseDriver.h
--- a/Engine/Source/Runtime/AnimGraphRuntime/Public/AnimNodes/AnimNode_PoseDriver.h
+++ b/Engine/Source/Runtime/AnimGraphRuntime/Public/AnimNodes/AnimNode_PoseDriver.h
@@ -79,6 +79,14 @@ struct ANIMGRAPHRUNTIME_API FPoseDriverTarget
 	UPROPERTY(EditAnywhere, Category = PoseDriver)
 	float TargetScale;
 
+	/** Override for the distance method to use for each target */
+	UPROPERTY(EditAnywhere, Category = RBFData)
+	ERBFDistanceMethod DistanceMethod;
+
+	/** Override for the function method to use for each target */
+	UPROPERTY(EditAnywhere, Category = RBFData)
+	ERBFFunctionType FunctionType;
+
 	/** If we should apply a custom curve mapping to how this target activates */
 	UPROPERTY(EditAnywhere, Category = PoseDriver)
 	bool bApplyCustomCurve;
@@ -101,13 +109,20 @@ struct ANIMGRAPHRUNTIME_API FPoseDriverTarget
 	/** Pose buffer index */
 	int32 PoseCurveIndex;
 
+	/** If we should hide this pose from the UI */
+	UPROPERTY(EditAnywhere, Category = PoseDriver)
+	bool bIsHidden;
+
 	FPoseDriverTarget()
 		: TargetRotation(ForceInitToZero)
 		, TargetScale(1.f)
+		, DistanceMethod(ERBFDistanceMethod::DefaultMethod)
+		, FunctionType(ERBFFunctionType::DefaultFunction)
 		, bApplyCustomCurve(false)
 		, DrivenName(NAME_None)
 		, DrivenUID(INDEX_NONE)
 		, PoseCurveIndex(INDEX_NONE)
+		, bIsHidden(false)
 	{}
 };
 
@@ -191,6 +206,9 @@ struct ANIMGRAPHRUNTIME_API FAnimNode_PoseDriver : public FAnimNode_PoseHandler
 
 	FAnimNode_PoseDriver();
 
+	/** Returns the radius for a given target */
+	float GetRadiusForTarget(const FRBFTarget& Target) const;
+
 	/** Util for seeing if BoneName is in the list of driven bones (and bFilterDrivenBones is true) */
 	bool IsBoneDriven(FName BoneName) const;
 
diff --git a/Engine/Source/Runtime/AnimGraphRuntime/Public/RBF/RBFSolver.h b/Engine/Source/Runtime/AnimGraphRuntime/Public/RBF/RBFSolver.h
--- a/Engine/Source/Runtime/AnimGraphRuntime/Public/RBF/RBFSolver.h
+++ b/Engine/Source/Runtime/AnimGraphRuntime/Public/RBF/RBFSolver.h
@@ -20,7 +20,10 @@ enum class ERBFFunctionType : uint8
 
 	Cubic,
 
-	Quintic
+	Quintic,
+
+	/** Uses the setting of the parent container */
+	DefaultFunction
 };
 
 /** Method for determining distance from input to targets */
@@ -34,7 +37,33 @@ enum class ERBFDistanceMethod : uint8
 	Quaternion,
 
 	/** Treat inputs as quaternion, and find distance between rotated TwistAxis direction */
-	SwingAngle
+	SwingAngle,
+
+	/** Uses the setting of the parent container */
+	DefaultMethod
+};
+
+/** Method to normalize weights */
+UENUM()
+enum class ERBFNormalizeMethod : uint8
+{
+	/** Only normalize above one */
+	OnlyNormalizeAboveOne,
+
+	/** 
+		Always normalize. 
+		Zero distribution weights stay zero.
+	*/
+	AlwaysNormalize,
+
+	/** 
+		Normalize only within reference median. The median
+		is a cone with a minimum and maximum angle within
+		which the value will be interpolated between 
+		non-normalized and normalized. This helps to define
+		the volume in which normalization is always required.
+	*/
+	NormalizeWithinMedian
 };
 
 /** Struct storing a particular entry within the RBF */
@@ -79,9 +108,17 @@ struct ANIMGRAPHRUNTIME_API FRBFTarget : public FRBFEntry
 	UPROPERTY(EditAnywhere, Category = RBFData)
 	FRichCurve CustomCurve;
 
+	UPROPERTY(EditAnywhere, Category = RBFData)
+	ERBFDistanceMethod DistanceMethod;
+
+	UPROPERTY(EditAnywhere, Category = RBFData)
+	ERBFFunctionType FunctionType;
+
 	FRBFTarget()
 		: ScaleFactor(1.f)
 		, bApplyCustomCurve(false)
+		, DistanceMethod(ERBFDistanceMethod::DefaultMethod)
+		, FunctionType(ERBFFunctionType::DefaultFunction)
 	{}
 };
 
@@ -132,6 +169,22 @@ struct ANIMGRAPHRUNTIME_API FRBFParams
 	UPROPERTY(EditAnywhere, Category = RBFData)
 	float WeightThreshold;
 
+	/** Method to use for normalizing the weight */
+	UPROPERTY(EditAnywhere, Category = RBFData)
+	ERBFNormalizeMethod NormalizeMethod;
+
+	/** Rotation or position of median (used for normalization) */
+	UPROPERTY(EditAnywhere, Category = RBFData)
+	FVector MedianReference;
+
+	/** Minimum distance used for median */
+	UPROPERTY(EditAnywhere, Category = RBFData, meta = (UIMin = "0", UIMax = "90"))
+	float MedianMin;
+
+	/** Maximum distance used for median */
+	UPROPERTY(EditAnywhere, Category = RBFData, meta = (UIMin = "0", UIMax = "90"))
+	float MedianMax;
+
 	FRBFParams();
 
 	/** Util for returning unit direction vector for swing axis */
@@ -148,5 +201,8 @@ struct ANIMGRAPHRUNTIME_API FRBFSolver
 	static bool FindTargetNeighbourDistances(const FRBFParams& Params, const TArray<FRBFTarget>& Targets, TArray<float>& NeighbourDists);
 
 	/** Util to find distance between two entries, using provided params */
-	static float FindDistanceBetweenEntries(const FRBFEntry& A, const FRBFEntry& B, const FRBFParams& Params);
+	static float FindDistanceBetweenEntries(const FRBFEntry& A, const FRBFEntry& B, const FRBFParams& Params, ERBFDistanceMethod OverrideMethod = ERBFDistanceMethod::DefaultMethod);
+
+	/** Returns the radius for a given target */
+	static float GetRadiusForTarget(const FRBFTarget& Target, const FRBFParams& Params);
 };
