Ensuring that the PoseDrivers Details use the right AnimNode to drive it. When the Anim Blueprint Editor is set to an instance in the world it needs to look that up first.

diff --git a/Engine/Source/Editor/AnimGraph/Private/AnimGraphNode_PoseDriver.cpp b/Engine/Source/Editor/AnimGraph/Private/AnimGraphNode_PoseDriver.cpp
--- a/Engine/Source/Editor/AnimGraph/Private/AnimGraphNode_PoseDriver.cpp
+++ b/Engine/Source/Editor/AnimGraph/Private/AnimGraphNode_PoseDriver.cpp
@@ -6,6 +6,7 @@
 #include "RBF/RBFSolver.h"
 #include "Kismet2/BlueprintEditorUtils.h"
 #include "Components/SkeletalMeshComponent.h"
+#include "Animation/AnimInstance.h"
 
 #define LOCTEXT_NAMESPACE "PoseDriver"
 
@@ -225,10 +226,28 @@ void UAnimGraphNode_PoseDriver::CopyNodeDataToPreviewNode(FAnimNode_Base* InPrev
 FAnimNode_PoseDriver* UAnimGraphNode_PoseDriver::GetPreviewPoseDriverNode() const
 {
 	FAnimNode_PoseDriver* PreviewNode = nullptr;
+	USkeletalMeshComponent * Component = nullptr;
 
-	if (LastPreviewComponent != nullptr && LastPreviewComponent->GetAnimInstance() != nullptr)
+	// look for a valid component in the object being debugged,
+	// we might be set to something other than the preview.
+	UObject * ObjectBeingDebugged = GetAnimBlueprint()->GetObjectBeingDebugged();
+	if (ObjectBeingDebugged)
 	{
-		PreviewNode = static_cast<FAnimNode_PoseDriver*>(FindDebugAnimNode(LastPreviewComponent));
+		UAnimInstance * InstanceBeingDebugged = Cast<UAnimInstance>(ObjectBeingDebugged);
+		if (InstanceBeingDebugged)
+		{
+			Component = InstanceBeingDebugged->GetSkelMeshComponent();
+		}
+	}
+
+	// Fall back to the LastPreviewComponent
+	if (Component == nullptr)
+	{
+		Component = LastPreviewComponent;
+	}
+	if (Component != nullptr && Component->GetAnimInstance() != nullptr)
+	{
+		PreviewNode = static_cast<FAnimNode_PoseDriver*>(FindDebugAnimNode(Component));
 	}
 
 	return PreviewNode;
