Constraint node NOTE remains when the weight values are set to 1

diff --git a/Engine/Source/Editor/AnimGraph/Private/AnimGraphNode_Constraint.cpp b/Engine/Source/Editor/AnimGraph/Private/AnimGraphNode_Constraint.cpp
--- a/Engine/Source/Editor/AnimGraph/Private/AnimGraphNode_Constraint.cpp
+++ b/Engine/Source/Editor/AnimGraph/Private/AnimGraphNode_Constraint.cpp
@@ -56,9 +56,12 @@ void UAnimGraphNode_Constraint::ValidateAnimNodeDuringCompilation(USkeleton* For
 	}
 
 	float OverallWeight = 0.f;
-	for (const float& Weight : Node.ConstraintWeights)
+	for (UEdGraphPin* Pin : Pins)
 	{
-		OverallWeight += Weight;
+		if (Pin->GetName().StartsWith(TEXT("ConstraintWeights")))
+		{
+			OverallWeight += FCString::Atof(*(Pin->DefaultValue));
+		}
 	}
 	if (Node.ConstraintWeights.Num() > 0 && !FMath::IsNearlyEqual(OverallWeight, 1.f, ZERO_ANIMWEIGHT_THRESH * float(Node.ConstraintWeights.Num())))
 	{
