PoseDriver: Add checkbox for enabling / disabling cone drawing

diff --git a/Engine/Source/Editor/AnimGraph/Classes/AnimGraphNode_PoseDriver.h b/Engine/Source/Editor/AnimGraph/Classes/AnimGraphNode_PoseDriver.h
--- a/Engine/Source/Editor/AnimGraph/Classes/AnimGraphNode_PoseDriver.h
+++ b/Engine/Source/Editor/AnimGraph/Classes/AnimGraphNode_PoseDriver.h
@@ -27,6 +27,10 @@ class ANIMGRAPH_API UAnimGraphNode_PoseDriver : public UAnimGraphNode_PoseHandle
 	UPROPERTY(EditAnywhere, Category = Debugging, meta = (DefaultValue = "32", UIMin = "6", UIMax = "128"))
 	int32 ConeSubdivision;
 
+	/** If checked the cones will be drawn in 3d for debugging */
+	UPROPERTY(EditAnywhere, Category = Debugging, meta = (DefaultValue = "True"))
+	bool bDrawDebugCones;
+
 	/** Used to indicate selected target to edit mode drawing */
 	int32 SelectedTargetIndex;
 	/** Delegate to call when selection changes */
diff --git a/Engine/Source/Editor/AnimGraph/Private/AnimGraphNode_PoseDriver.cpp b/Engine/Source/Editor/AnimGraph/Private/AnimGraphNode_PoseDriver.cpp
--- a/Engine/Source/Editor/AnimGraph/Private/AnimGraphNode_PoseDriver.cpp
+++ b/Engine/Source/Editor/AnimGraph/Private/AnimGraphNode_PoseDriver.cpp
@@ -44,6 +44,7 @@ UAnimGraphNode_PoseDriver::UAnimGraphNode_PoseDriver(const FObjectInitializer& O
 {
 	SelectedTargetIndex = INDEX_NONE;
 	AxisLength = 20.0f;
+	bDrawDebugCones = true;
 	ConeSubdivision = 32;
 }
 
diff --git a/Engine/Source/Editor/AnimGraph/Private/EditModes/PoseDriverEditMode.cpp b/Engine/Source/Editor/AnimGraph/Private/EditModes/PoseDriverEditMode.cpp
--- a/Engine/Source/Editor/AnimGraph/Private/EditModes/PoseDriverEditMode.cpp
+++ b/Engine/Source/Editor/AnimGraph/Private/EditModes/PoseDriverEditMode.cpp
@@ -59,6 +59,7 @@ void FPoseDriverEditMode::Render(const FSceneView* View, FViewport* Viewport, FP
 	static const float DrawPosSize = 2.f;
 	float DrawAxisLength = GraphNode->AxisLength;
 	int32 DrawConeSubdivision = GraphNode->ConeSubdivision;
+	bool bDrawDebugCones = GraphNode->bDrawDebugCones;
 
 	TArray<FRBFTarget> RBFTargets;
 	RuntimeNode->GetRBFTargets(RBFTargets);
@@ -106,43 +107,46 @@ void FPoseDriverEditMode::Render(const FSceneView* View, FViewport* Viewport, FP
 					PDI->DrawLine(BonePos, BonePos + (WorldTwistVec*DrawAxisLength), FLinearColor::Green, SDPG_Foreground, DrawLineWidth);
 
 					// draw the median cones
-					if (RuntimeNode->RBFParams.NormalizeMethod == ERBFNormalizeMethod::NormalizeWithinMedian)
+					if (bDrawDebugCones)
 					{
-						const FVector& MedianRot = RuntimeNode->RBFParams.MedianReference;
-						LocalTwistVec = FRotator(MedianRot.X, MedianRot.Y, MedianRot.Z).RotateVector(RuntimeNode->RBFParams.GetTwistAxisVector());
-						WorldTwistVec = EvalSpaceTM.TransformVectorNoScale(LocalTwistVec);
+						if (RuntimeNode->RBFParams.NormalizeMethod == ERBFNormalizeMethod::NormalizeWithinMedian)
+						{
+							const FVector& MedianRot = RuntimeNode->RBFParams.MedianReference;
+							LocalTwistVec = FRotator(MedianRot.X, MedianRot.Y, MedianRot.Z).RotateVector(RuntimeNode->RBFParams.GetTwistAxisVector());
+							WorldTwistVec = EvalSpaceTM.TransformVectorNoScale(LocalTwistVec);
 
-						FVector LocalSwingVec = FVector::CrossProduct(LocalTwistVec, FVector(1, 1, 1));
-						FVector WorldSwingVec = EvalSpaceTM.TransformVectorNoScale(LocalSwingVec);
-						WorldSwingVec.Normalize();
+							FVector LocalSwingVec = FVector::CrossProduct(LocalTwistVec, FVector(1, 1, 1));
+							FVector WorldSwingVec = EvalSpaceTM.TransformVectorNoScale(LocalSwingVec);
+							WorldSwingVec.Normalize();
 
-						FQuat WorldTwistQ(WorldTwistVec, PI * 2.0f / float(DrawConeSubdivision));
-						FQuat WorldSwingMinQ(WorldSwingVec, FMath::DegreesToRadians(RuntimeNode->RBFParams.MedianMin));
-						FQuat WorldSwingMaxQ(WorldSwingVec, FMath::DegreesToRadians(RuntimeNode->RBFParams.MedianMax));
+							FQuat WorldTwistQ(WorldTwistVec, PI * 2.0f / float(DrawConeSubdivision));
+							FQuat WorldSwingMinQ(WorldSwingVec, FMath::DegreesToRadians(RuntimeNode->RBFParams.MedianMin));
+							FQuat WorldSwingMaxQ(WorldSwingVec, FMath::DegreesToRadians(RuntimeNode->RBFParams.MedianMax));
 
-						FVector FirstMinPositionOnCircle = WorldSwingMinQ.RotateVector(WorldTwistVec * DrawAxisLength);
-						FVector FirstMaxPositionOnCircle = WorldSwingMaxQ.RotateVector(WorldTwistVec * DrawAxisLength);
-						FVector LastMinPositionOnCircle = FirstMinPositionOnCircle;
-						FVector LastMaxPositionOnCircle = FirstMaxPositionOnCircle;
+							FVector FirstMinPositionOnCircle = WorldSwingMinQ.RotateVector(WorldTwistVec * DrawAxisLength);
+							FVector FirstMaxPositionOnCircle = WorldSwingMaxQ.RotateVector(WorldTwistVec * DrawAxisLength);
+							FVector LastMinPositionOnCircle = FirstMinPositionOnCircle;
+							FVector LastMaxPositionOnCircle = FirstMaxPositionOnCircle;
 
-						FLinearColor MinColor = FLinearColor::Yellow;
-						FLinearColor MaxColor = MinColor.Desaturate(0.5);
+							FLinearColor MinColor = FLinearColor::Yellow;
+							FLinearColor MaxColor = MinColor.Desaturate(0.5);
 
-						for (int32 i = 0; i < DrawConeSubdivision; i++)
-						{
-							FVector NextMinPositionOnCircle = WorldTwistQ.RotateVector(LastMinPositionOnCircle);
-							FVector NextMaxPositionOnCircle = WorldTwistQ.RotateVector(LastMaxPositionOnCircle);
-							PDI->DrawLine(BonePos, BonePos + NextMinPositionOnCircle, MinColor, SDPG_Foreground, DrawLineWidth);
-							PDI->DrawLine(BonePos, BonePos + NextMaxPositionOnCircle, MaxColor , SDPG_Foreground, DrawLineWidth);
-							PDI->DrawLine(BonePos + LastMinPositionOnCircle, BonePos + NextMinPositionOnCircle, MinColor, SDPG_Foreground, DrawLineWidth);
-							PDI->DrawLine(BonePos + LastMaxPositionOnCircle, BonePos + NextMaxPositionOnCircle, MaxColor, SDPG_Foreground, DrawLineWidth);
-							PDI->DrawLine(BonePos + NextMinPositionOnCircle, BonePos + NextMaxPositionOnCircle, MaxColor, SDPG_Foreground, DrawLineWidth);
-
-							LastMinPositionOnCircle = NextMinPositionOnCircle;
-							LastMaxPositionOnCircle = NextMaxPositionOnCircle;
+							for (int32 i = 0; i < DrawConeSubdivision; i++)
+							{
+								FVector NextMinPositionOnCircle = WorldTwistQ.RotateVector(LastMinPositionOnCircle);
+								FVector NextMaxPositionOnCircle = WorldTwistQ.RotateVector(LastMaxPositionOnCircle);
+								PDI->DrawLine(BonePos, BonePos + NextMinPositionOnCircle, MinColor, SDPG_Foreground, DrawLineWidth);
+								PDI->DrawLine(BonePos, BonePos + NextMaxPositionOnCircle, MaxColor, SDPG_Foreground, DrawLineWidth);
+								PDI->DrawLine(BonePos + LastMinPositionOnCircle, BonePos + NextMinPositionOnCircle, MinColor, SDPG_Foreground, DrawLineWidth);
+								PDI->DrawLine(BonePos + LastMaxPositionOnCircle, BonePos + NextMaxPositionOnCircle, MaxColor, SDPG_Foreground, DrawLineWidth);
+								PDI->DrawLine(BonePos + NextMinPositionOnCircle, BonePos + NextMaxPositionOnCircle, MaxColor, SDPG_Foreground, DrawLineWidth);
+
+								LastMinPositionOnCircle = NextMinPositionOnCircle;
+								LastMaxPositionOnCircle = NextMaxPositionOnCircle;
+							}
+							PDI->DrawLine(BonePos + LastMinPositionOnCircle, BonePos + FirstMinPositionOnCircle, MinColor, SDPG_Foreground, DrawLineWidth);
+							PDI->DrawLine(BonePos + LastMaxPositionOnCircle, BonePos + FirstMaxPositionOnCircle, MaxColor, SDPG_Foreground, DrawLineWidth);
 						}
-						PDI->DrawLine(BonePos + LastMinPositionOnCircle, BonePos + FirstMinPositionOnCircle, MinColor, SDPG_Foreground, DrawLineWidth);
-						PDI->DrawLine(BonePos + LastMaxPositionOnCircle, BonePos + FirstMaxPositionOnCircle, MaxColor, SDPG_Foreground, DrawLineWidth);
 					}
 				}
 				// Translation drawing
@@ -188,7 +192,7 @@ void FPoseDriverEditMode::Render(const FSceneView* View, FViewport* Viewport, FP
 						continue;
 					}
 
-					if (PoseTarget.BoneTransforms.IsValidIndex(SourceIdx))
+					if (bDrawDebugCones && PoseTarget.BoneTransforms.IsValidIndex(SourceIdx))
 					{
 						const FPoseDriverTransform& TargetTM = PoseTarget.BoneTransforms[SourceIdx];
 
