Layer nodes can now only select from the interfaces/layers that are in the current animation blueprint

Slight change to the runtime behaviour: when running a layer with the same class as its hosting BP, we use a 'self' layer rather than outputting a warning.

#rnx

diff --git a/Engine/Source/Editor/AnimGraph/Classes/AnimGraphNode_Layer.h b/Engine/Source/Editor/AnimGraph/Classes/AnimGraphNode_Layer.h
--- a/Engine/Source/Editor/AnimGraph/Classes/AnimGraphNode_Layer.h
+++ b/Engine/Source/Editor/AnimGraph/Classes/AnimGraphNode_Layer.h
@@ -39,15 +39,12 @@ public:
 
 protected:
 	// ----- UI CALLBACKS ----- //
-	// Gets path to the currently selected interface class' blueprint
-	FString GetCurrentInterfaceBlueprintPath() const;
-	// Filter callback for interface blueprints (only accept anim blueprint interfaces)
-	bool OnShouldFilterInterfaceBlueprint(const FAssetData& AssetData) const;
-	// Interface was changed by user
-	void OnSetInterfaceBlueprint(const FAssetData& AssetData, TSharedRef<IPropertyHandle> InterfaceClassPropHandle);
 	// Handlers for layer combo
 	void GetLayerNames(TArray<TSharedPtr<FString>>& OutStrings, TArray<TSharedPtr<SToolTip>>& OutToolTips, TArray<bool>& OutRestrictedItems);
 	FString GetLayerName() const;
+	void OnLayerChanged(IDetailLayoutBuilder* DetailBuilder);
+	bool HasAvailableLayers() const;
+	bool HasValidNonSelfLayer() const;
 	// ----- END UI CALLBACKS ----- //
 
 	// Begin UAnimGraphNode_CustomProperty
@@ -63,4 +60,7 @@ protected:
 	virtual bool OnShouldFilterInstanceBlueprint(const FAssetData& AssetData) const override;
 	virtual FString GetCurrentInstanceBlueprintPath() const override;
 	virtual bool IsStructuralProperty(UProperty* InProperty) const override;
+
+	// Helper function to get the interface currently in use by the selected layer
+	TSubclassOf<UInterface> GetInterfaceForLayer() const;
 };
diff --git a/Engine/Source/Editor/AnimGraph/Private/AnimGraphNode_Layer.cpp b/Engine/Source/Editor/AnimGraph/Private/AnimGraphNode_Layer.cpp
--- a/Engine/Source/Editor/AnimGraph/Private/AnimGraphNode_Layer.cpp
+++ b/Engine/Source/Editor/AnimGraph/Private/AnimGraphNode_Layer.cpp
@@ -71,6 +71,38 @@ void UAnimGraphNode_Layer::ValidateAnimNodeDuringCompilation(USkeleton* ForSkele
 				TargetClass = *CurrentBlueprint->SkeletonGeneratedClass;
 			}
 		}
+		else
+		{
+			// check we implement this interface
+			bool bImplementsInterface = false;
+			for(FBPInterfaceDescription& InterfaceDesc : CurrentBlueprint->ImplementedInterfaces)
+			{
+				if(InterfaceDesc.Interface.Get() == TargetClass)
+				{
+					bImplementsInterface = true;
+					break;
+				}
+			}
+
+			if(!bImplementsInterface)
+			{
+				// Its possible we have a left-over interface referenced here that needs clearing now we are a 'self' layer
+				if(GetInterfaceForLayer() == nullptr)
+				{
+					Node.Interface = nullptr;
+
+					// No interface any more, use this class
+					if (CurrentBlueprint)
+					{
+						TargetClass = *CurrentBlueprint->SkeletonGeneratedClass;
+					}
+				}
+				else
+				{
+					MessageLog.Error(*LOCTEXT("MissingInterfaceError", "Layer node @@ uses interface @@ that this blueprint does not implement.").ToString(), this, Node.Interface.Get());
+				}
+			}
+		}
 
 		if(TargetClass)
 		{
@@ -213,43 +245,17 @@ void UAnimGraphNode_Layer::CustomizeDetails(IDetailLayoutBuilder& DetailBuilder)
 	TSharedRef<IPropertyHandle> TagHandle = DetailBuilder.GetProperty(TEXT("Node.Tag"), GetClass());
 	TagHandle->MarkHiddenByCustomization();
 
-	// Customize Interface
-	{
-		TSharedRef<IPropertyHandle> InterfaceHandle = DetailBuilder.GetProperty(TEXT("Node.Interface"), GetClass());
-		if(InterfaceHandle->IsValidHandle())
-		{
-			InterfaceHandle->SetOnPropertyValueChanged(FSimpleDelegate::CreateUObject(this, &UAnimGraphNode_Layer::OnStructuralPropertyChanged, &DetailBuilder));
-		}
-
-		InterfaceHandle->MarkHiddenByCustomization();
-
-		FDetailWidgetRow& InterfaceWidgetRow = CategoryBuilder.AddCustomRow(LOCTEXT("FilterStringInterface", "Interface"));
-		InterfaceWidgetRow.NameContent()
-		[
-			InterfaceHandle->CreatePropertyNameWidget()
-		]
-		.ValueContent()
-		.MinDesiredWidth(250.0f)
-		[
-			SNew(SObjectPropertyEntryBox)
-			.ObjectPath_UObject(this, &UAnimGraphNode_Layer::GetCurrentInterfaceBlueprintPath)
-			.AllowedClass(UAnimBlueprint::StaticClass())
-			.NewAssetFactories(TArray<UFactory*>())
-			.OnShouldFilterAsset(FOnShouldFilterAsset::CreateUObject(this, &UAnimGraphNode_Layer::OnShouldFilterInterfaceBlueprint))
-			.OnObjectChanged(FOnSetObject::CreateUObject(this, &UAnimGraphNode_Layer::OnSetInterfaceBlueprint, InterfaceHandle))
-		];
-	}
-
 	// Customize Layer
 	{
 		TSharedRef<IPropertyHandle> LayerHandle = DetailBuilder.GetProperty(TEXT("Node.Layer"), GetClass());
 		if(LayerHandle->IsValidHandle())
 		{
-			LayerHandle->SetOnPropertyValueChanged(FSimpleDelegate::CreateUObject(this, &UAnimGraphNode_Layer::OnStructuralPropertyChanged, &DetailBuilder));
+			LayerHandle->SetOnPropertyValueChanged(FSimpleDelegate::CreateUObject(this, &UAnimGraphNode_Layer::OnLayerChanged, &DetailBuilder));
 		}
 
 		LayerHandle->MarkHiddenByCustomization();
 
+		// Check layers available in this BP
 		FDetailWidgetRow& LayerWidgetRow = CategoryBuilder.AddCustomRow(LOCTEXT("FilterStringLayer", "Layer"));
 		LayerWidgetRow.NameContent()
 		[
@@ -258,32 +264,68 @@ void UAnimGraphNode_Layer::CustomizeDetails(IDetailLayoutBuilder& DetailBuilder)
 		.ValueContent()
 		.MinDesiredWidth(150.0f)
 		[
-			PropertyCustomizationHelpers::MakePropertyComboBox(
-				LayerHandle, 
-				FOnGetPropertyComboBoxStrings::CreateUObject(this,  &UAnimGraphNode_Layer::GetLayerNames),
-				FOnGetPropertyComboBoxValue::CreateUObject(this,  &UAnimGraphNode_Layer::GetLayerName)
-			)
+			SNew(SHorizontalBox)
+			+SHorizontalBox::Slot()
+			.VAlign(VAlign_Center)
+			[
+				SNew(SBox)
+				.Visibility_Lambda([this](){ return HasAvailableLayers() ? EVisibility::Visible : EVisibility::Collapsed; })
+				[
+					PropertyCustomizationHelpers::MakePropertyComboBox(
+						LayerHandle, 
+						FOnGetPropertyComboBoxStrings::CreateUObject(this,  &UAnimGraphNode_Layer::GetLayerNames),
+						FOnGetPropertyComboBoxValue::CreateUObject(this,  &UAnimGraphNode_Layer::GetLayerName)
+					)
+				]
+			]
+			+SHorizontalBox::Slot()
+			.VAlign(VAlign_Center)
+			[
+				SNew(STextBlock)
+				.Visibility_Lambda([this](){ return !HasAvailableLayers() ? EVisibility::Visible : EVisibility::Collapsed; })
+				.Font(IDetailLayoutBuilder::GetDetailFont())
+				.Text(LOCTEXT("NoLayersWarning", "No available layers."))
+				.ToolTipText(LOCTEXT("NoLayersWarningTooltip", "This Animation Blueprint has no layers to choose from.\nTo add some, either implement an Animation Layer Interface via the Class Settings, or add an animation layer in the My Blueprint tab."))
+			]
 		];
 	}
 
-	Super::CustomizeDetails(DetailBuilder);
-}
-
-FString UAnimGraphNode_Layer::GetCurrentInterfaceBlueprintPath() const
-{
-	UClass* InterfaceClass = *Node.Interface;
-
-	if(InterfaceClass)
+	// Customize InstanceClass with unique visibility (identical to parent class apart from this)
 	{
-		UBlueprint* ActualBlueprint = UBlueprint::GetBlueprintFromClass(InterfaceClass);
+		TSharedRef<IPropertyHandle> ClassHandle = DetailBuilder.GetProperty(TEXT("Node.InstanceClass"), GetClass());
+		ClassHandle->MarkHiddenByCustomization();
 
-		if(ActualBlueprint)
-		{
-			return ActualBlueprint->GetPathName();
-		}
+		FDetailWidgetRow& ClassWidgetRow = CategoryBuilder.AddCustomRow(LOCTEXT("FilterStringInstanceClass", "Instance Class"));
+		ClassWidgetRow.NameContent()
+		[
+			ClassHandle->CreatePropertyNameWidget()
+		]
+		.ValueContent()
+		.MinDesiredWidth(250.0f)
+		[
+			SNew(SHorizontalBox)
+			+SHorizontalBox::Slot()
+			.VAlign(VAlign_Center)
+			[
+				SNew(SObjectPropertyEntryBox)
+				.Visibility_Lambda([this](){ return HasValidNonSelfLayer() ? EVisibility::Visible : EVisibility::Collapsed; })
+				.ObjectPath_UObject(this, &UAnimGraphNode_Layer::GetCurrentInstanceBlueprintPath)
+				.AllowedClass(UAnimBlueprint::StaticClass())
+				.NewAssetFactories(TArray<UFactory*>())
+				.OnShouldFilterAsset(FOnShouldFilterAsset::CreateUObject(this, &UAnimGraphNode_Layer::OnShouldFilterInstanceBlueprint))
+				.OnObjectChanged(FOnSetObject::CreateUObject(this, &UAnimGraphNode_Layer::OnSetInstanceBlueprint, &DetailBuilder))
+			]
+			+SHorizontalBox::Slot()
+			.VAlign(VAlign_Center)
+			[
+				SNew(STextBlock)
+				.Visibility_Lambda([this](){ return !HasValidNonSelfLayer() ? EVisibility::Visible : EVisibility::Collapsed; })
+				.Font(IDetailLayoutBuilder::GetDetailFont())
+				.Text(LOCTEXT("SelfLayersWarning", "Uses layer in this Blueprint."))
+				.ToolTipText(LOCTEXT("SelfLayersWarningTooltip", "This layer node refers to a layer only in this blueprint, so cannot be overriden by an external blueprint implementation.\nChange to use a layer from an implemented interface to allow this override."))
+			]
+		];
 	}
-
-	return FString();
 }
 
 bool UAnimGraphNode_Layer::OnShouldFilterInstanceBlueprint(const FAssetData& AssetData) const
@@ -295,8 +337,20 @@ bool UAnimGraphNode_Layer::OnShouldFilterInstanceBlueprint(const FAssetData& Ass
 
 	if (UAnimBlueprint* CurrentBlueprint = Cast<UAnimBlueprint>(GetBlueprint()))
 	{
+		TArray<TSubclassOf<UInterface>> AnimInterfaces;
+		for(const FBPInterfaceDescription& InterfaceDesc : CurrentBlueprint->ImplementedInterfaces)
+		{
+			if(InterfaceDesc.Interface && InterfaceDesc.Interface->IsChildOf<UAnimLayerInterface>())
+			{
+				if(Node.Layer == NAME_None || InterfaceDesc.Interface->FindFunctionByName(Node.Layer))
+				{
+					AnimInterfaces.Add(InterfaceDesc.Interface);
+				}
+			}
+		}
+
 		// Check interface compatibility
-		if(UClass* InterfaceClass = Node.Interface.Get())
+		if(AnimInterfaces.Num() > 0)
 		{
 			bool bMatchesInterface = false;
 
@@ -317,7 +371,12 @@ bool UAnimGraphNode_Layer::OnShouldFilterInstanceBlueprint(const FAssetData& Ass
 							InterfacePath.RemoveFromEnd(TEXT("\"'"));
 
 							FCoreRedirectObjectName ResolvedInterfaceName = FCoreRedirects::GetRedirectedName(ECoreRedirectFlags::Type_Class, FCoreRedirectObjectName(InterfacePath));
-							bMatchesInterface = ResolvedInterfaceName.ObjectName == InterfaceClass->GetFName();
+
+							// Verify against all interfaces we currently implement
+							for(TSubclassOf<UInterface> AnimInterface : AnimInterfaces)
+							{
+								bMatchesInterface |= ResolvedInterfaceName.ObjectName == AnimInterface->GetFName();
+							}
 						}
 					}
 			
@@ -330,6 +389,11 @@ bool UAnimGraphNode_Layer::OnShouldFilterInstanceBlueprint(const FAssetData& Ass
 				return true;
 			}
 		}
+		else
+		{
+			// No interfaces, so no compatible BPs
+			return true;
+		}
 	}
 
 	return false;
@@ -352,38 +416,6 @@ FString UAnimGraphNode_Layer::GetCurrentInstanceBlueprintPath() const
 	return FString();
 }
 
-bool UAnimGraphNode_Layer::OnShouldFilterInterfaceBlueprint(const FAssetData& AssetData) const
-{
-	FAssetDataTagMapSharedView::FFindTagResult Result = AssetData.TagsAndValues.FindTag("ParentClass");
-	if (Result.IsSet())
-	{
-		FString ParentClassObjectPath;
-		if (FPackageName::ParseExportTextPath(Result.GetValue(), nullptr, &ParentClassObjectPath))
-		{
-			ParentClassObjectPath.RemoveFromEnd(TEXT("_C"));
-			return ParentClassObjectPath != UAnimLayerInterface::StaticClass()->GetPathName();
-		}
-
-	}
-
-	return true;
-}
-
-void UAnimGraphNode_Layer::OnSetInterfaceBlueprint(const FAssetData& AssetData, TSharedRef<IPropertyHandle> InterfaceClassPropHandle)
-{
-	FScopedTransaction Transaction(LOCTEXT("SetInterface", "Set Instance Interface"));
-	Modify();
-
-	if(UAnimBlueprint* Blueprint = Cast<UAnimBlueprint>(AssetData.GetAsset()))
-	{
-		InterfaceClassPropHandle->SetValue(Blueprint->GetAnimBlueprintGeneratedClass());
-	}
-	else
-	{
-		InterfaceClassPropHandle->SetValue((UObject*)nullptr);
-	}
-}
-
 void UAnimGraphNode_Layer::GetExposableProperties(TArray<UProperty*>& OutExposableProperties) const
 {
 	UClass* TargetClass = GetTargetSkeletonClass();
@@ -407,26 +439,21 @@ void UAnimGraphNode_Layer::GetExposableProperties(TArray<UProperty*>& OutExposab
 
 void UAnimGraphNode_Layer::GetLayerNames(TArray<TSharedPtr<FString>>& OutStrings, TArray<TSharedPtr<SToolTip>>& OutToolTips, TArray<bool>& OutRestrictedItems)
 {
-	UClass* TargetClass = *Node.Interface;
-	if(TargetClass == nullptr)
-	{
-		// If no interface specified, use this class
-		if (UAnimBlueprint* CurrentBlueprint = Cast<UAnimBlueprint>(GetBlueprint()))
-		{
-			TargetClass = *CurrentBlueprint->SkeletonGeneratedClass;
-		}
-	}
-
-	if(TargetClass)
+	// If no interface specified, use this class
+	if (UAnimBlueprint* CurrentBlueprint = Cast<UAnimBlueprint>(GetBlueprint()))
 	{
-		IAnimClassInterface* AnimClassInterface = IAnimClassInterface::GetFromClass(TargetClass);
-		for(const FAnimBlueprintFunction& AnimBlueprintFunction : AnimClassInterface->GetAnimBlueprintFunctions())
+		UClass* TargetClass = *CurrentBlueprint->SkeletonGeneratedClass;
+		if(TargetClass)
 		{
-			if(AnimBlueprintFunction.Name != UEdGraphSchema_K2::GN_AnimGraph)
+			IAnimClassInterface* AnimClassInterface = IAnimClassInterface::GetFromClass(TargetClass);
+			for(const FAnimBlueprintFunction& AnimBlueprintFunction : AnimClassInterface->GetAnimBlueprintFunctions())
 			{
-				OutStrings.Add(MakeShared<FString>(AnimBlueprintFunction.Name.ToString()));
-				OutToolTips.Add(nullptr);
-				OutRestrictedItems.Add(false);
+				if(AnimBlueprintFunction.Name != UEdGraphSchema_K2::GN_AnimGraph)
+				{
+					OutStrings.Add(MakeShared<FString>(AnimBlueprintFunction.Name.ToString()));
+					OutToolTips.Add(nullptr);
+					OutRestrictedItems.Add(false);
+				}
 			}
 		}
 	}
@@ -440,7 +467,6 @@ FString UAnimGraphNode_Layer::GetLayerName() const
 bool UAnimGraphNode_Layer::IsStructuralProperty(UProperty* InProperty) const
 {
 	return Super::IsStructuralProperty(InProperty) ||
-		InProperty->GetFName() == GET_MEMBER_NAME_CHECKED(FAnimNode_Layer, Interface) ||
 		InProperty->GetFName() == GET_MEMBER_NAME_CHECKED(FAnimNode_Layer, Layer);
 }
 
@@ -449,7 +475,7 @@ UClass* UAnimGraphNode_Layer::GetTargetSkeletonClass() const
 	UClass* SuperTargetSkeletonClass = Super::GetTargetSkeletonClass();
 	if(SuperTargetSkeletonClass == nullptr)
 	{
-		// If no interface specified, use this class
+		// If no concrete class specified, use this class
 		if (UAnimBlueprint* CurrentBlueprint = Cast<UAnimBlueprint>(GetBlueprint()))
 		{
 			SuperTargetSkeletonClass = *CurrentBlueprint->SkeletonGeneratedClass;
@@ -458,4 +484,85 @@ UClass* UAnimGraphNode_Layer::GetTargetSkeletonClass() const
 	return SuperTargetSkeletonClass;
 }
 
+TSubclassOf<UInterface> UAnimGraphNode_Layer::GetInterfaceForLayer() const
+{
+	if (UAnimBlueprint* CurrentBlueprint = Cast<UAnimBlueprint>(GetBlueprint()))
+	{
+		UClass* TargetClass = *CurrentBlueprint->SkeletonGeneratedClass;
+		if(TargetClass)
+		{
+			// Find layer with this name in interfaces
+			for(FBPInterfaceDescription& InterfaceDesc : CurrentBlueprint->ImplementedInterfaces)
+			{
+				for(UEdGraph* InterfaceGraph : InterfaceDesc.Graphs)
+				{
+					if(InterfaceGraph->GetFName() == Node.Layer)
+					{
+						return InterfaceDesc.Interface;
+					}
+				}
+			}
+		}
+	}
+
+	return nullptr;
+}
+
+void UAnimGraphNode_Layer::OnLayerChanged(IDetailLayoutBuilder* DetailBuilder)
+{
+	OnStructuralPropertyChanged(DetailBuilder);
+
+	// Get the interface for this layer. If null, then we are using a 'self' layer.
+	Node.Interface = GetInterfaceForLayer();
+
+	if(Node.Interface.Get() == nullptr)
+	{
+		// Self layers cannot have override implementations
+		Node.InstanceClass = nullptr;
+	}
+}
+
+bool UAnimGraphNode_Layer::HasAvailableLayers() const
+{
+	if (UAnimBlueprint* CurrentBlueprint = Cast<UAnimBlueprint>(GetBlueprint()))
+	{
+		UClass* TargetClass = *CurrentBlueprint->SkeletonGeneratedClass;
+		if(TargetClass)
+		{
+			IAnimClassInterface* AnimClassInterface = IAnimClassInterface::GetFromClass(TargetClass);
+			for(const FAnimBlueprintFunction& AnimBlueprintFunction : AnimClassInterface->GetAnimBlueprintFunctions())
+			{
+				if(AnimBlueprintFunction.Name != UEdGraphSchema_K2::GN_AnimGraph)
+				{
+					return true;
+				}
+			}
+		}
+	}
+
+	return false;
+}
+
+bool UAnimGraphNode_Layer::HasValidNonSelfLayer() const
+{
+	if (UAnimBlueprint* CurrentBlueprint = Cast<UAnimBlueprint>(GetBlueprint()))
+	{
+		if(Node.Interface.Get())
+		{
+			for(const FBPInterfaceDescription& InterfaceDesc : CurrentBlueprint->ImplementedInterfaces)
+			{
+				if(InterfaceDesc.Interface && InterfaceDesc.Interface->IsChildOf<UAnimLayerInterface>())
+				{
+					if(InterfaceDesc.Interface->FindFunctionByName(Node.Layer))
+					{
+						return true;
+					}
+				}
+			}
+		}
+	}
+
+	return false;
+}
+
 #undef LOCTEXT_NAMESPACE
diff --git a/Engine/Source/Runtime/Engine/Classes/Animation/AnimNode_Layer.h b/Engine/Source/Runtime/Engine/Classes/Animation/AnimNode_Layer.h
--- a/Engine/Source/Runtime/Engine/Classes/Animation/AnimNode_Layer.h
+++ b/Engine/Source/Runtime/Engine/Classes/Animation/AnimNode_Layer.h
@@ -23,8 +23,11 @@ struct ENGINE_API FAnimNode_Layer : public FAnimNode_SubInstance
 	GENERATED_BODY()
 
 public:
-	/** Optional interface. If this is set then this node will only accept (both statically and dynamically) anim instances that implement this interface */
-	UPROPERTY(EditAnywhere, Category = Settings)
+	/** 
+	 * Optional interface. If this is set then this node will only accept (both statically and dynamically) anim instances that implement this interface.
+	 * If not set, then this is considered a 'self' layer. This value is set when Layer is changed in the details panel.
+	 */
+	UPROPERTY()
 	TSubclassOf<UAnimLayerInterface> Interface;
 
 	/** The layer in the interface to use */
diff --git a/Engine/Source/Runtime/Engine/Private/Animation/AnimInstance.cpp b/Engine/Source/Runtime/Engine/Private/Animation/AnimInstance.cpp
--- a/Engine/Source/Runtime/Engine/Private/Animation/AnimInstance.cpp
+++ b/Engine/Source/Runtime/Engine/Private/Animation/AnimInstance.cpp
@@ -2538,15 +2538,8 @@ void UAnimInstance::SetLayerOverlay(TSubclassOf<UAnimInstance> InClass)
 			{
 				// If the class is null, then reset to default (which can be null)
 				UClass* ClassToSet = NewClass != nullptr ? NewClass : LayerPair.Value[0]->InstanceClass.Get();
-				if(ClassToSet != nullptr)
+				if(ClassToSet != nullptr && ClassToSet != GetClass())
 				{
-					if(ClassToSet == GetClass())
-					{
-						// Cant instance a version of this class, otherwise we will end up in an infinite recursion
-						UE_LOG(LogAnimation, Error, TEXT("Setting layer overlay: Potential infinite recursion - cant set an overlay to the same class as its layer's host (%s)"), *ClassToSet->GetName())
-						break;
-					}
-
 					// Create and add one sub-instance for this group
 					USkeletalMeshComponent* MeshComp = GetSkelMeshComponent();
 					UAnimInstance* NewSubInstance = NewObject<UAnimInstance>(MeshComp, ClassToSet);
diff --git a/Engine/Source/Runtime/Engine/Private/Animation/AnimNode_Layer.cpp b/Engine/Source/Runtime/Engine/Private/Animation/AnimNode_Layer.cpp
--- a/Engine/Source/Runtime/Engine/Private/Animation/AnimNode_Layer.cpp
+++ b/Engine/Source/Runtime/Engine/Private/Animation/AnimNode_Layer.cpp
@@ -23,7 +23,7 @@ void FAnimNode_Layer::OnInitializeAnimInstance(const FAnimInstanceProxy* InProxy
 {
 	// We only initialize here if we are running a 'self' layer. Layers that use external instances need to be 
 	// initialized by the owning anim instance as they may share sub-instances via grouping.
-	if(Interface.Get() == nullptr)
+	if(Interface.Get() == nullptr || InstanceClass.Get() == nullptr)
 	{
 		UAnimInstance* CurrentTarget = GetTargetInstance<UAnimInstance>();
 
