Fix exposed variables not appearing for layer nodes

#rnx

diff --git a/Engine/Source/Editor/AnimGraph/Private/AnimGraphNode_Layer.cpp b/Engine/Source/Editor/AnimGraph/Private/AnimGraphNode_Layer.cpp
--- a/Engine/Source/Editor/AnimGraph/Private/AnimGraphNode_Layer.cpp
+++ b/Engine/Source/Editor/AnimGraph/Private/AnimGraphNode_Layer.cpp
@@ -290,6 +290,8 @@ void UAnimGraphNode_Layer::CustomizeDetails(IDetailLayoutBuilder& DetailBuilder)
 		];
 	}
 
+	UAnimGraphNode_CustomProperty::CustomizeDetails(DetailBuilder);
+
 	// Customize InstanceClass with unique visibility (identical to parent class apart from this)
 	{
 		TSharedRef<IPropertyHandle> ClassHandle = DetailBuilder.GetProperty(TEXT("Node.InstanceClass"), GetClass());
