PR #5688: Blend poses have pins by default (Contributed by [REDACTED_PERSON])

diff --git a/Engine/Source/Editor/AnimGraph/Classes/AnimGraphNode_BlendListBase.h b/Engine/Source/Editor/AnimGraph/Classes/AnimGraphNode_BlendListBase.h
--- a/Engine/Source/Editor/AnimGraph/Classes/AnimGraphNode_BlendListBase.h
+++ b/Engine/Source/Editor/AnimGraph/Classes/AnimGraphNode_BlendListBase.h
@@ -26,5 +26,5 @@ protected:
 	// removes removed pins and adjusts array indices of remained pins
 	void RemovePinsFromOldPins(TArray<UEdGraphPin*>& OldPins, int32 RemovedArrayIndex);
 
-	int32 RemovedPinArrayIndex;
+	TArray<int32> RemovedPinArrayIndices;
 };
diff --git a/Engine/Source/Editor/AnimGraph/Classes/AnimGraphNode_BlendListByEnum.h b/Engine/Source/Editor/AnimGraph/Classes/AnimGraphNode_BlendListByEnum.h
--- a/Engine/Source/Editor/AnimGraph/Classes/AnimGraphNode_BlendListByEnum.h
+++ b/Engine/Source/Editor/AnimGraph/Classes/AnimGraphNode_BlendListByEnum.h
@@ -32,6 +32,7 @@ public:
 	virtual FText GetTooltipText() const override;
 	virtual FText GetNodeTitle(ENodeTitleType::Type TitleType) const override;
 	virtual void Serialize(FArchive& Ar) override;
+	virtual void AllocateDefaultPins() override;
 	// End of UEdGraphNode interface
 
 	// UK2Node interface
@@ -47,8 +48,9 @@ public:
 	virtual void PreloadRequiredAssets() override;
 	// End of UAnimGraphNode_Base interface
 
+	void RemoveUnconnectedPinsFromBlendList();
 	//@TODO: Generalize this behavior (returning a list of actions/delegates maybe?)
-	virtual void RemovePinFromBlendList(UEdGraphPin* Pin);
+	void RemovePinFromBlendList(UEdGraphPin* Pin);
 
 	// INodeDependingOnEnumInterface
 	virtual class UEnum* GetEnum() const override { return BoundEnum; }
diff --git a/Engine/Source/Editor/AnimGraph/Private/AnimGraphNode_BlendListBase.cpp b/Engine/Source/Editor/AnimGraph/Private/AnimGraphNode_BlendListBase.cpp
--- a/Engine/Source/Editor/AnimGraph/Private/AnimGraphNode_BlendListBase.cpp
+++ b/Engine/Source/Editor/AnimGraph/Private/AnimGraphNode_BlendListBase.cpp
@@ -7,7 +7,6 @@
 
 UAnimGraphNode_BlendListBase::UAnimGraphNode_BlendListBase(const FObjectInitializer& ObjectInitializer)
 	: Super(ObjectInitializer)
-	, RemovedPinArrayIndex(INDEX_NONE)
 {
 }
 
@@ -96,11 +95,12 @@ void UAnimGraphNode_BlendListBase::ReallocatePinsDuringReconstruction(TArray<UEd
 	Super::ReallocatePinsDuringReconstruction(OldPins);
 
 	// Delete Pins by removed pin info 
-	if (RemovedPinArrayIndex != INDEX_NONE)
+	for (const auto& i : RemovedPinArrayIndices)
 	{
-		RemovePinsFromOldPins(OldPins, RemovedPinArrayIndex);
-		// Clears removed pin info to avoid to remove multiple times
-		// @TODO : Considering receiving RemovedPinArrayIndex as an argument of ReconstructNode()
-		RemovedPinArrayIndex = INDEX_NONE;
+		RemovePinsFromOldPins(OldPins, i);
 	}
+
+	// Clears removed pin info to avoid to remove multiple times
+	// @TODO : Considering receiving RemovedPinArrayIndex as an argument of ReconstructNode()
+	RemovedPinArrayIndices.Empty();
 }
diff --git a/Engine/Source/Editor/AnimGraph/Private/AnimGraphNode_BlendListByEnum.cpp b/Engine/Source/Editor/AnimGraph/Private/AnimGraphNode_BlendListByEnum.cpp
--- a/Engine/Source/Editor/AnimGraph/Private/AnimGraphNode_BlendListByEnum.cpp
+++ b/Engine/Source/Editor/AnimGraph/Private/AnimGraphNode_BlendListByEnum.cpp
@@ -19,8 +19,6 @@
 UAnimGraphNode_BlendListByEnum::UAnimGraphNode_BlendListByEnum(const FObjectInitializer& ObjectInitializer)
 	: Super(ObjectInitializer)
 {
-	// Make sure we start out with a pin
-	Node.AddPose();
 }
 
 FString UAnimGraphNode_BlendListByEnum::GetNodeCategory() const
@@ -98,6 +96,12 @@ void UAnimGraphNode_BlendListByEnum::GetContextMenuActions(const FGraphNodeConte
 					FUIAction Action = FUIAction( FExecuteAction::CreateUObject( this, &UAnimGraphNode_BlendListByEnum::RemovePinFromBlendList, const_cast<UEdGraphPin*>(Context.Pin)) );
 					Context.MenuBuilder->AddMenuEntry( LOCTEXT("RemovePose", "Remove Pose"), FText::GetEmpty(), FSlateIcon(), Action );
 				}
+
+				if (VisibleEnumEntries.Num() > 0)
+				{
+					FUIAction Action = FUIAction(FExecuteAction::CreateUObject(this, &UAnimGraphNode_BlendListByEnum::RemoveUnconnectedPinsFromBlendList));
+					Context.MenuBuilder->AddMenuEntry(LOCTEXT("RemoveUnconnectedPoses", "Remove unconnected Poses"), FText::GetEmpty(), FSlateIcon(), Action);
+				}
 			}
 		}
 
@@ -149,6 +153,51 @@ void UAnimGraphNode_BlendListByEnum::ExposeEnumElementAsPin(FName EnumElementNam
 	}
 }
 
+void UAnimGraphNode_BlendListByEnum::RemoveUnconnectedPinsFromBlendList()
+{
+	FScopedTransaction Transaction(LOCTEXT("RemoveUnconnectedPins", "Remove unconnected pins"));
+	Modify();
+
+	// inverse loop as otherwise VisibleEnumEntries.RemoveAt is incorrect
+	for (int32 i = Pins.Num() - 1; i >= 0; --i)
+	{
+		// skip if output pin OR if it has a connection
+		if (Pins[i]->Direction != EGPD_Input || Pins[i]->LinkedTo.Num() > 0)
+		{
+			continue;
+		}
+
+		int32 RawArrayIndex = 0;
+		bool bIsPosePin = false;
+		bool bIsTimePin = false;
+		GetPinInformation(Pins[i]->PinName.ToString(), /*out*/ RawArrayIndex, /*out*/ bIsPosePin, /*out*/ bIsTimePin);
+
+		// Only IsPosePin otherwise arrays go out of bounds
+		const int32 ExposedEnumIndex = bIsPosePin ? (RawArrayIndex - 1) : INDEX_NONE;
+
+		if (ExposedEnumIndex != INDEX_NONE)
+		{
+			// Record it as no longer exposed
+			VisibleEnumEntries.RemoveAt(ExposedEnumIndex);
+
+			// Remove the pose from the node
+			UProperty* AssociatedProperty;
+			int32 ArrayIndex;
+			GetPinAssociatedProperty(GetFNodeType(), Pins[i], /*out*/ AssociatedProperty, /*out*/ ArrayIndex);
+
+			ensure(ArrayIndex == (ExposedEnumIndex + 1));
+
+			RemovedPinArrayIndices.Add(ArrayIndex);
+			Node.RemovePose(ArrayIndex);
+			Pins[i]->SetSavePinIfOrphaned(false);
+		}
+	}
+
+	ReconstructNode();
+	//@TODO: Just want to invalidate the visual representation currently
+	FBlueprintEditorUtils::MarkBlueprintAsStructurallyModified(GetBlueprint());
+}
+
 void UAnimGraphNode_BlendListByEnum::RemovePinFromBlendList(UEdGraphPin* Pin)
 {
 	int32 RawArrayIndex = 0;
@@ -174,7 +223,7 @@ void UAnimGraphNode_BlendListByEnum::RemovePinFromBlendList(UEdGraphPin* Pin)
 		ensure(ArrayIndex == (ExposedEnumIndex + 1));
 
 		// setting up removed pins info 
-		RemovedPinArrayIndex = ArrayIndex;
+		RemovedPinArrayIndices.Add(ArrayIndex);
 		Node.RemovePose(ArrayIndex);
 		Pin->SetSavePinIfOrphaned(false);
 		ReconstructNode();
@@ -289,6 +338,26 @@ void UAnimGraphNode_BlendListByEnum::Serialize(FArchive& Ar)
 	}
 }
 
+void UAnimGraphNode_BlendListByEnum::AllocateDefaultPins()
+{
+	Super::AllocateDefaultPins();
+
+	// Make sure we start out with a pin
+	Node.AddPose();
+
+	if (BoundEnum && VisibleEnumEntries.Num() <= 0)
+	{
+		const int32 MaxIndex = BoundEnum->NumEnums() - 1; // we don't want to show _MAX enum
+		for (int32 Index = 0; Index < MaxIndex; ++Index)
+		{
+			VisibleEnumEntries.Add(BoundEnum->GetNameByIndex(Index));
+			Node.AddPose();
+		}
+	}
+
+	ReconstructNode();
+}
+
 void UAnimGraphNode_BlendListByEnum::ValidateAnimNodeDuringCompilation(class USkeleton* ForSkeleton, class FCompilerResultsLog& MessageLog)
 {
 	Super::ValidateAnimNodeDuringCompilation(ForSkeleton, MessageLog);
diff --git a/Engine/Source/Editor/AnimGraph/Private/AnimGraphNode_BlendListByInt.cpp b/Engine/Source/Editor/AnimGraph/Private/AnimGraphNode_BlendListByInt.cpp
--- a/Engine/Source/Editor/AnimGraph/Private/AnimGraphNode_BlendListByInt.cpp
+++ b/Engine/Source/Editor/AnimGraph/Private/AnimGraphNode_BlendListByInt.cpp
@@ -60,7 +60,7 @@ void UAnimGraphNode_BlendListByInt::RemovePinFromBlendList(UEdGraphPin* Pin)
 	{
 		//@TODO: ANIMREFACTOR: Need to handle moving pins below up correctly
 		// setting up removed pins info
-		RemovedPinArrayIndex = ArrayIndex;
+		RemovedPinArrayIndices.Add(ArrayIndex);
 		Node.RemovePose(ArrayIndex);
 		// removes the selected pin and related properties in reconstructNode()
 		// @TODO: Considering passing "RemovedPinArrayIndex" to ReconstructNode as the argument
diff --git a/Engine/Source/Editor/AnimGraph/Private/AnimGraphNode_LayeredBoneBlend.cpp b/Engine/Source/Editor/AnimGraph/Private/AnimGraphNode_LayeredBoneBlend.cpp
--- a/Engine/Source/Editor/AnimGraph/Private/AnimGraphNode_LayeredBoneBlend.cpp
+++ b/Engine/Source/Editor/AnimGraph/Private/AnimGraphNode_LayeredBoneBlend.cpp
@@ -61,7 +61,7 @@ void UAnimGraphNode_LayeredBoneBlend::RemovePinFromBlendByFilter(UEdGraphPin* Pi
 	{
 		//@TODO: ANIMREFACTOR: Need to handle moving pins below up correctly
 		// setting up removed pins info 
-		RemovedPinArrayIndex = ArrayIndex;
+		RemovedPinArrayIndices.Add(ArrayIndex);
 		Node.RemovePose(ArrayIndex);
 		ReconstructNode();
 		FBlueprintEditorUtils::MarkBlueprintAsStructurallyModified(GetBlueprint());
