Guarded against NULL (probably garbage-collected) assets when creating the BP action database

diff --git a/Engine/Source/Editor/AnimGraph/Private/AnimGraphNode_SequencePlayer.cpp b/Engine/Source/Editor/AnimGraph/Private/AnimGraphNode_SequencePlayer.cpp
--- a/Engine/Source/Editor/AnimGraph/Private/AnimGraphNode_SequencePlayer.cpp
+++ b/Engine/Source/Editor/AnimGraph/Private/AnimGraphNode_SequencePlayer.cpp
@@ -197,10 +197,13 @@ void UAnimGraphNode_SequencePlayer::GetMenuActions(FBlueprintActionDatabaseRegis
 			UBlueprintNodeSpawner* NodeSpawner = UBlueprintNodeSpawner::Create(GetClass());
 			if (Asset.IsAssetLoaded())
 			{
-				TWeakObjectPtr<UAnimSequence> AnimSequence = Cast<UAnimSequence>(Asset.GetAsset());
-				NodeSpawner->CustomizeNodeDelegate = UBlueprintNodeSpawner::FCustomizeNodeDelegate::CreateStatic(LoadedAssetSetup, AnimSequence);
-				NodeSpawner->DefaultMenuSignature.MenuName = GetTitleGivenAssetInfo(FText::FromName(AnimSequence->GetFName()), AnimSequence->IsValidAdditive());
-				NodeSpawner->DefaultMenuSignature.Tooltip = GetTitleGivenAssetInfo(FText::FromString(AnimSequence->GetPathName()), AnimSequence->IsValidAdditive());
+				UAnimSequence* AnimSequence = Cast<UAnimSequence>(Asset.GetAsset());
+				if(AnimSequence)
+				{
+					NodeSpawner->CustomizeNodeDelegate = UBlueprintNodeSpawner::FCustomizeNodeDelegate::CreateStatic(LoadedAssetSetup, AnimSequence);
+					NodeSpawner->DefaultMenuSignature.MenuName = GetTitleGivenAssetInfo(FText::FromName(AnimSequence->GetFName()), AnimSequence->IsValidAdditive());
+					NodeSpawner->DefaultMenuSignature.Tooltip = GetTitleGivenAssetInfo(FText::FromString(AnimSequence->GetPathName()), AnimSequence->IsValidAdditive());
+				}
 			}
 			else
 			{
