Hardened asset players against playing montages

diff --git a/Engine/Source/Editor/AnimGraph/Private/AnimGraphNode_SequencePlayer.cpp b/Engine/Source/Editor/AnimGraph/Private/AnimGraphNode_SequencePlayer.cpp
--- a/Engine/Source/Editor/AnimGraph/Private/AnimGraphNode_SequencePlayer.cpp
+++ b/Engine/Source/Editor/AnimGraph/Private/AnimGraphNode_SequencePlayer.cpp
@@ -200,7 +200,7 @@ void UAnimGraphNode_SequencePlayer::GetMenuActions(FBlueprintActionDatabaseRegis
 				UAnimSequence* AnimSequence = Cast<UAnimSequence>(Asset.GetAsset());
 				if(AnimSequence)
 				{
-					NodeSpawner->CustomizeNodeDelegate = UBlueprintNodeSpawner::FCustomizeNodeDelegate::CreateStatic(LoadedAssetSetup, AnimSequence);
+					NodeSpawner->CustomizeNodeDelegate = UBlueprintNodeSpawner::FCustomizeNodeDelegate::CreateStatic(LoadedAssetSetup, TWeakObjectPtr<UAnimSequence>(AnimSequence));
 					NodeSpawner->DefaultMenuSignature.MenuName = GetTitleGivenAssetInfo(FText::FromName(AnimSequence->GetFName()), AnimSequence->IsValidAdditive());
 					NodeSpawner->DefaultMenuSignature.Tooltip = GetTitleGivenAssetInfo(FText::FromString(AnimSequence->GetPathName()), AnimSequence->IsValidAdditive());
 				}
@@ -329,6 +329,10 @@ void UAnimGraphNode_SequencePlayer::ValidateAnimNodeDuringCompilation(class USke
 			MessageLog.Error(TEXT("@@ references an unknown sequence"), this);
 		}
 	}
+	else if(SupportsAssetClass(SequenceToCheck->GetClass()) == EAnimAssetHandlerType::NotSupported)
+	{
+		MessageLog.Error(*FText::Format(LOCTEXT("UnsupportedAssetError", "@@ is trying to play a {0} as a sequence, which is not allowed."), SequenceToCheck->GetClass()->GetDisplayNameText()).ToString(), this);
+	}
 	else
 	{
 		USkeleton* SeqSkeleton = SequenceToCheck->GetSkeleton();
diff --git a/Engine/Source/Runtime/Engine/Private/Animation/AnimNode_SequencePlayer.cpp b/Engine/Source/Runtime/Engine/Private/Animation/AnimNode_SequencePlayer.cpp
--- a/Engine/Source/Runtime/Engine/Private/Animation/AnimNode_SequencePlayer.cpp
+++ b/Engine/Source/Runtime/Engine/Private/Animation/AnimNode_SequencePlayer.cpp
@@ -32,6 +32,12 @@ void FAnimNode_SequencePlayer::Initialize_AnyThread(const FAnimationInitializeCo
 	FAnimNode_AssetPlayerBase::Initialize_AnyThread(Context);
 
 	GetEvaluateGraphExposedInputs().Execute(Context);
+
+	if(ensureMsgf(Sequence->IsA<UAnimMontage>(), TEXT("Sequence players do not support anim montages.")))
+	{
+		Sequence = nullptr;
+	}
+
 	InternalTimeAccumulator = StartPosition;
 	PlayRateScaleBiasClamp.Reinitialize();
 
@@ -55,6 +61,11 @@ void FAnimNode_SequencePlayer::UpdateAssetPlayer(const FAnimationUpdateContext&
 {
 	GetEvaluateGraphExposedInputs().Execute(Context);
 
+	if(ensureMsgf(Sequence->IsA<UAnimMontage>(), TEXT("Sequence players do not support anim montages.")))
+	{
+		Sequence = nullptr;
+	}
+
 	if ((Sequence != nullptr) && (Context.AnimInstanceProxy->IsSkeletonCompatible(Sequence->GetSkeleton())))
 	{
 		InternalTimeAccumulator = FMath::Clamp(InternalTimeAccumulator, 0.f, Sequence->SequenceLength);
