Fix build error

diff --git a/Engine/Source/Editor/AnimGraph/Private/AnimGraphNode_BlendListByEnum.cpp b/Engine/Source/Editor/AnimGraph/Private/AnimGraphNode_BlendListByEnum.cpp
--- a/Engine/Source/Editor/AnimGraph/Private/AnimGraphNode_BlendListByEnum.cpp
+++ b/Engine/Source/Editor/AnimGraph/Private/AnimGraphNode_BlendListByEnum.cpp
@@ -99,7 +99,7 @@ void UAnimGraphNode_BlendListByEnum::GetContextMenuActions(const FGraphNodeConte
 
 				if (VisibleEnumEntries.Num() > 0)
 				{
-					FUIAction Action = FUIAction(FExecuteAction::CreateUObject(this, &UAnimGraphNode_BlendListByEnum::RemoveUnconnectedPinsFromBlendList));
+					FUIAction Action = FUIAction(FExecuteAction::CreateUObject(const_cast<UAnimGraphNode_BlendListByEnum*>(this), &UAnimGraphNode_BlendListByEnum::RemoveUnconnectedPinsFromBlendList));
 					Context.MenuBuilder->AddMenuEntry(LOCTEXT("RemoveUnconnectedPoses", "Remove unconnected Poses"), FText::GetEmpty(), FSlateIcon(), Action);
 				}
 			}
