Back out changelist 6426679

fail jira - UE-72441

- this changes behavior severly and it is not desirable at this point without proper investigation.
- reverting change 6426679

diff --git a/Engine/Source/Editor/AnimGraph/Classes/AnimGraphNode_BlendListBase.h b/Engine/Source/Editor/AnimGraph/Classes/AnimGraphNode_BlendListBase.h
--- a/Engine/Source/Editor/AnimGraph/Classes/AnimGraphNode_BlendListBase.h
+++ b/Engine/Source/Editor/AnimGraph/Classes/AnimGraphNode_BlendListBase.h
@@ -26,5 +26,5 @@ protected:
 	// removes removed pins and adjusts array indices of remained pins
 	void RemovePinsFromOldPins(TArray<UEdGraphPin*>& OldPins, int32 RemovedArrayIndex);
 
-	TArray<int32> RemovedPinArrayIndices;
+	int32 RemovedPinArrayIndex;
 };
diff --git a/Engine/Source/Editor/AnimGraph/Classes/AnimGraphNode_BlendListByEnum.h b/Engine/Source/Editor/AnimGraph/Classes/AnimGraphNode_BlendListByEnum.h
--- a/Engine/Source/Editor/AnimGraph/Classes/AnimGraphNode_BlendListByEnum.h
+++ b/Engine/Source/Editor/AnimGraph/Classes/AnimGraphNode_BlendListByEnum.h
@@ -32,7 +32,6 @@ public:
 	virtual FText GetTooltipText() const override;
 	virtual FText GetNodeTitle(ENodeTitleType::Type TitleType) const override;
 	virtual void Serialize(FArchive& Ar) override;
-	virtual void AllocateDefaultPins() override;
 	// End of UEdGraphNode interface
 
 	// UK2Node interface
@@ -48,9 +47,8 @@ public:
 	virtual void PreloadRequiredAssets() override;
 	// End of UAnimGraphNode_Base interface
 
-	void RemoveUnconnectedPinsFromBlendList();
 	//@TODO: Generalize this behavior (returning a list of actions/delegates maybe?)
-	void RemovePinFromBlendList(UEdGraphPin* Pin);
+	virtual void RemovePinFromBlendList(UEdGraphPin* Pin);
 
 	// INodeDependingOnEnumInterface
 	virtual class UEnum* GetEnum() const override { return BoundEnum; }
diff --git a/Engine/Source/Editor/AnimGraph/Private/AnimGraphNode_BlendListBase.cpp b/Engine/Source/Editor/AnimGraph/Private/AnimGraphNode_BlendListBase.cpp
--- a/Engine/Source/Editor/AnimGraph/Private/AnimGraphNode_BlendListBase.cpp
+++ b/Engine/Source/Editor/AnimGraph/Private/AnimGraphNode_BlendListBase.cpp
@@ -7,6 +7,7 @@
 
 UAnimGraphNode_BlendListBase::UAnimGraphNode_BlendListBase(const FObjectInitializer& ObjectInitializer)
 	: Super(ObjectInitializer)
+	, RemovedPinArrayIndex(INDEX_NONE)
 {
 }
 
@@ -95,12 +96,11 @@ void UAnimGraphNode_BlendListBase::ReallocatePinsDuringReconstruction(TArray<UEd
 	Super::ReallocatePinsDuringReconstruction(OldPins);
 
 	// Delete Pins by removed pin info 
-	for (const auto& i : RemovedPinArrayIndices)
+	if (RemovedPinArrayIndex != INDEX_NONE)
 	{
-		RemovePinsFromOldPins(OldPins, i);
+		RemovePinsFromOldPins(OldPins, RemovedPinArrayIndex);
+		// Clears removed pin info to avoid to remove multiple times
+		// @TODO : Considering receiving RemovedPinArrayIndex as an argument of ReconstructNode()
+		RemovedPinArrayIndex = INDEX_NONE;
 	}
-
-	// Clears removed pin info to avoid to remove multiple times
-	// @TODO : Considering receiving RemovedPinArrayIndex as an argument of ReconstructNode()
-	RemovedPinArrayIndices.Empty();
 }
diff --git a/Engine/Source/Editor/AnimGraph/Private/AnimGraphNode_BlendListByEnum.cpp b/Engine/Source/Editor/AnimGraph/Private/AnimGraphNode_BlendListByEnum.cpp
--- a/Engine/Source/Editor/AnimGraph/Private/AnimGraphNode_BlendListByEnum.cpp
+++ b/Engine/Source/Editor/AnimGraph/Private/AnimGraphNode_BlendListByEnum.cpp
@@ -19,6 +19,8 @@
 UAnimGraphNode_BlendListByEnum::UAnimGraphNode_BlendListByEnum(const FObjectInitializer& ObjectInitializer)
 	: Super(ObjectInitializer)
 {
+	// Make sure we start out with a pin
+	Node.AddPose();
 }
 
 FString UAnimGraphNode_BlendListByEnum::GetNodeCategory() const
@@ -96,12 +98,6 @@ void UAnimGraphNode_BlendListByEnum::GetContextMenuActions(const FGraphNodeConte
 					FUIAction Action = FUIAction( FExecuteAction::CreateUObject( const_cast<UAnimGraphNode_BlendListByEnum*>(this), &UAnimGraphNode_BlendListByEnum::RemovePinFromBlendList, const_cast<UEdGraphPin*>(Context.Pin)) );
 					Context.MenuBuilder->AddMenuEntry( LOCTEXT("RemovePose", "Remove Pose"), FText::GetEmpty(), FSlateIcon(), Action );
 				}
-
-				if (VisibleEnumEntries.Num() > 0)
-				{
-					FUIAction Action = FUIAction(FExecuteAction::CreateUObject(const_cast<UAnimGraphNode_BlendListByEnum*>(this), &UAnimGraphNode_BlendListByEnum::RemoveUnconnectedPinsFromBlendList));
-					Context.MenuBuilder->AddMenuEntry(LOCTEXT("RemoveUnconnectedPoses", "Remove unconnected Poses"), FText::GetEmpty(), FSlateIcon(), Action);
-				}
 			}
 		}
 
@@ -153,51 +149,6 @@ void UAnimGraphNode_BlendListByEnum::ExposeEnumElementAsPin(FName EnumElementNam
 	}
 }
 
-void UAnimGraphNode_BlendListByEnum::RemoveUnconnectedPinsFromBlendList()
-{
-	FScopedTransaction Transaction(LOCTEXT("RemoveUnconnectedPins", "Remove unconnected pins"));
-	Modify();
-
-	// inverse loop as otherwise VisibleEnumEntries.RemoveAt is incorrect
-	for (int32 i = Pins.Num() - 1; i >= 0; --i)
-	{
-		// skip if output pin OR if it has a connection
-		if (Pins[i]->Direction != EGPD_Input || Pins[i]->LinkedTo.Num() > 0)
-		{
-			continue;
-		}
-
-		int32 RawArrayIndex = 0;
-		bool bIsPosePin = false;
-		bool bIsTimePin = false;
-		GetPinInformation(Pins[i]->PinName.ToString(), /*out*/ RawArrayIndex, /*out*/ bIsPosePin, /*out*/ bIsTimePin);
-
-		// Only IsPosePin otherwise arrays go out of bounds
-		const int32 ExposedEnumIndex = bIsPosePin ? (RawArrayIndex - 1) : INDEX_NONE;
-
-		if (ExposedEnumIndex != INDEX_NONE)
-		{
-			// Record it as no longer exposed
-			VisibleEnumEntries.RemoveAt(ExposedEnumIndex);
-
-			// Remove the pose from the node
-			UProperty* AssociatedProperty;
-			int32 ArrayIndex;
-			GetPinAssociatedProperty(GetFNodeType(), Pins[i], /*out*/ AssociatedProperty, /*out*/ ArrayIndex);
-
-			ensure(ArrayIndex == (ExposedEnumIndex + 1));
-
-			RemovedPinArrayIndices.Add(ArrayIndex);
-			Node.RemovePose(ArrayIndex);
-			Pins[i]->SetSavePinIfOrphaned(false);
-		}
-	}
-
-	ReconstructNode();
-	//@TODO: Just want to invalidate the visual representation currently
-	FBlueprintEditorUtils::MarkBlueprintAsStructurallyModified(GetBlueprint());
-}
-
 void UAnimGraphNode_BlendListByEnum::RemovePinFromBlendList(UEdGraphPin* Pin)
 {
 	int32 RawArrayIndex = 0;
@@ -223,7 +174,7 @@ void UAnimGraphNode_BlendListByEnum::RemovePinFromBlendList(UEdGraphPin* Pin)
 		ensure(ArrayIndex == (ExposedEnumIndex + 1));
 
 		// setting up removed pins info 
-		RemovedPinArrayIndices.Add(ArrayIndex);
+		RemovedPinArrayIndex = ArrayIndex;
 		Node.RemovePose(ArrayIndex);
 		Pin->SetSavePinIfOrphaned(false);
 		ReconstructNode();
@@ -338,26 +289,6 @@ void UAnimGraphNode_BlendListByEnum::Serialize(FArchive& Ar)
 	}
 }
 
-void UAnimGraphNode_BlendListByEnum::AllocateDefaultPins()
-{
-	Super::AllocateDefaultPins();
-
-	// Make sure we start out with a pin
-	Node.AddPose();
-
-	if (BoundEnum && VisibleEnumEntries.Num() <= 0)
-	{
-		const int32 MaxIndex = BoundEnum->NumEnums() - 1; // we don't want to show _MAX enum
-		for (int32 Index = 0; Index < MaxIndex; ++Index)
-		{
-			VisibleEnumEntries.Add(BoundEnum->GetNameByIndex(Index));
-			Node.AddPose();
-		}
-	}
-
-	ReconstructNode();
-}
-
 void UAnimGraphNode_BlendListByEnum::ValidateAnimNodeDuringCompilation(class USkeleton* ForSkeleton, class FCompilerResultsLog& MessageLog)
 {
 	Super::ValidateAnimNodeDuringCompilation(ForSkeleton, MessageLog);
diff --git a/Engine/Source/Editor/AnimGraph/Private/AnimGraphNode_BlendListByInt.cpp b/Engine/Source/Editor/AnimGraph/Private/AnimGraphNode_BlendListByInt.cpp
--- a/Engine/Source/Editor/AnimGraph/Private/AnimGraphNode_BlendListByInt.cpp
+++ b/Engine/Source/Editor/AnimGraph/Private/AnimGraphNode_BlendListByInt.cpp
@@ -60,7 +60,7 @@ void UAnimGraphNode_BlendListByInt::RemovePinFromBlendList(UEdGraphPin* Pin)
 	{
 		//@TODO: ANIMREFACTOR: Need to handle moving pins below up correctly
 		// setting up removed pins info
-		RemovedPinArrayIndices.Add(ArrayIndex);
+		RemovedPinArrayIndex = ArrayIndex;
 		Node.RemovePose(ArrayIndex);
 		// removes the selected pin and related properties in reconstructNode()
 		// @TODO: Considering passing "RemovedPinArrayIndex" to ReconstructNode as the argument
diff --git a/Engine/Source/Editor/AnimGraph/Private/AnimGraphNode_LayeredBoneBlend.cpp b/Engine/Source/Editor/AnimGraph/Private/AnimGraphNode_LayeredBoneBlend.cpp
--- a/Engine/Source/Editor/AnimGraph/Private/AnimGraphNode_LayeredBoneBlend.cpp
+++ b/Engine/Source/Editor/AnimGraph/Private/AnimGraphNode_LayeredBoneBlend.cpp
@@ -61,7 +61,7 @@ void UAnimGraphNode_LayeredBoneBlend::RemovePinFromBlendByFilter(UEdGraphPin* Pi
 	{
 		//@TODO: ANIMREFACTOR: Need to handle moving pins below up correctly
 		// setting up removed pins info 
-		RemovedPinArrayIndices.Add(ArrayIndex);
+		RemovedPinArrayIndex = ArrayIndex;
 		Node.RemovePose(ArrayIndex);
 		ReconstructNode();
 		FBlueprintEditorUtils::MarkBlueprintAsStructurallyModified(GetBlueprint());
