Animation Layer node doesn't pass through variables found in instanced classes
#fix Reintroduce missing behaviour for generating pin exposing widget in details view for LayerNode

diff --git a/Engine/Source/Editor/AnimGraph/Classes/AnimGraphNode_CustomProperty.h b/Engine/Source/Editor/AnimGraph/Classes/AnimGraphNode_CustomProperty.h
--- a/Engine/Source/Editor/AnimGraph/Classes/AnimGraphNode_CustomProperty.h
+++ b/Engine/Source/Editor/AnimGraph/Classes/AnimGraphNode_CustomProperty.h
@@ -29,6 +29,7 @@ public:
 	//~ End UEdGraphNode Interface.
 
 	// UAnimGraphNode_Base interface
+	virtual void CustomizeDetails(IDetailLayoutBuilder& DetailBuilder) override;
 	// Gets the property on InOwnerInstanceClass that corresponds to InInputPin
 	void GetInstancePinProperty(const UClass* InOwnerInstanceClass, UEdGraphPin* InInputPin, UProperty*& OutProperty);
 	// Gets the unique name for the property linked to a given pin
diff --git a/Engine/Source/Editor/AnimGraph/Classes/AnimGraphNode_SubInstanceBase.h b/Engine/Source/Editor/AnimGraph/Classes/AnimGraphNode_SubInstanceBase.h
--- a/Engine/Source/Editor/AnimGraph/Classes/AnimGraphNode_SubInstanceBase.h
+++ b/Engine/Source/Editor/AnimGraph/Classes/AnimGraphNode_SubInstanceBase.h
@@ -45,6 +45,9 @@ public:
 protected:
 	// Finds out whether there is a loop in the graph formed by sub instances from this node
 	bool HasInstanceLoop();
+	
+	/** Generates widgets for exposing/hiding Pins for this node using the rpovided detail builder */
+	void GenerateExposedPinsDetails(IDetailLayoutBuilder &DetailBuilder);
 
 	// Finds out whether there is a loop in the graph formed by sub instances from CurrNode, used by HasInstanceLoop. VisitedNodes and NodeStack are required
 	// to track the graph links
diff --git a/Engine/Source/Editor/AnimGraph/Private/AnimGraphNode_CustomProperty.cpp b/Engine/Source/Editor/AnimGraph/Private/AnimGraphNode_CustomProperty.cpp
--- a/Engine/Source/Editor/AnimGraph/Private/AnimGraphNode_CustomProperty.cpp
+++ b/Engine/Source/Editor/AnimGraph/Private/AnimGraphNode_CustomProperty.cpp
@@ -242,6 +242,22 @@ bool UAnimGraphNode_CustomProperty::HasExternalDependencies(TArray<class UStruct
 	return InstanceClassToUse || bSuperResult;
 }
 
+void UAnimGraphNode_CustomProperty::CustomizeDetails(IDetailLayoutBuilder& DetailBuilder)
+{
+	Super::CustomizeDetails(DetailBuilder);
+
+	IDetailCategoryBuilder& CategoryBuilder = DetailBuilder.EditCategory(FName(TEXT("Settings")));
+
+	// Customize InstanceClass
+	{
+		TSharedRef<IPropertyHandle> ClassHandle = DetailBuilder.GetProperty(TEXT("Node.InstanceClass"), GetClass());
+		if (ClassHandle->IsValidHandle())
+		{
+			ClassHandle->SetOnPropertyValueChanged(FSimpleDelegate::CreateUObject(this, &UAnimGraphNode_CustomProperty::OnStructuralPropertyChanged, &DetailBuilder));
+		}
+	}
+}
+
 void UAnimGraphNode_CustomProperty::GetExposableProperties( TArray<UProperty*>& OutExposableProperties) const
 {
 	OutExposableProperties.Empty();
diff --git a/Engine/Source/Editor/AnimGraph/Private/AnimGraphNode_Layer.cpp b/Engine/Source/Editor/AnimGraph/Private/AnimGraphNode_Layer.cpp
--- a/Engine/Source/Editor/AnimGraph/Private/AnimGraphNode_Layer.cpp
+++ b/Engine/Source/Editor/AnimGraph/Private/AnimGraphNode_Layer.cpp
@@ -294,6 +294,7 @@ void UAnimGraphNode_Layer::CustomizeDetails(IDetailLayoutBuilder& DetailBuilder)
 		];
 	}
 
+	GenerateExposedPinsDetails(DetailBuilder);
 	UAnimGraphNode_CustomProperty::CustomizeDetails(DetailBuilder);
 
 	// Customize InstanceClass with unique visibility (identical to parent class apart from this)
diff --git a/Engine/Source/Editor/AnimGraph/Private/AnimGraphNode_SubInstanceBase.cpp b/Engine/Source/Editor/AnimGraph/Private/AnimGraphNode_SubInstanceBase.cpp
--- a/Engine/Source/Editor/AnimGraph/Private/AnimGraphNode_SubInstanceBase.cpp
+++ b/Engine/Source/Editor/AnimGraph/Private/AnimGraphNode_SubInstanceBase.cpp
@@ -227,6 +227,35 @@ void UAnimGraphNode_SubInstanceBase::CustomizeDetails(IDetailLayoutBuilder& Deta
 {
 	Super::CustomizeDetails(DetailBuilder);
 
+	GenerateExposedPinsDetails(DetailBuilder);
+
+	IDetailCategoryBuilder& CategoryBuilder = DetailBuilder.EditCategory(FName(TEXT("Settings")));
+
+	// Customize InstanceClass
+	{
+		TSharedRef<IPropertyHandle> ClassHandle = DetailBuilder.GetProperty(TEXT("Node.InstanceClass"), GetClass());
+		ClassHandle->MarkHiddenByCustomization();
+
+		FDetailWidgetRow& ClassWidgetRow = CategoryBuilder.AddCustomRow(LOCTEXT("FilterStringInstanceClass", "Instance Class"));
+		ClassWidgetRow.NameContent()
+		[
+			ClassHandle->CreatePropertyNameWidget()
+		]
+		.ValueContent()
+		.MinDesiredWidth(250.0f)
+		[
+			SNew(SObjectPropertyEntryBox)
+			.ObjectPath_UObject(this, &UAnimGraphNode_SubInstanceBase::GetCurrentInstanceBlueprintPath)
+			.AllowedClass(UAnimBlueprint::StaticClass())
+			.NewAssetFactories(TArray<UFactory*>())
+			.OnShouldFilterAsset(FOnShouldFilterAsset::CreateUObject(this, &UAnimGraphNode_SubInstanceBase::OnShouldFilterInstanceBlueprint))
+			.OnObjectChanged(FOnSetObject::CreateUObject(this, &UAnimGraphNode_SubInstanceBase::OnSetInstanceBlueprint, &DetailBuilder))
+		];
+	}
+}
+
+void UAnimGraphNode_SubInstanceBase::GenerateExposedPinsDetails(IDetailLayoutBuilder &DetailBuilder)
+{
 	// We dont allow multi-select here
 	if(DetailBuilder.GetSelectedObjects().Num() > 1)
 	{
@@ -323,30 +352,6 @@ void UAnimGraphNode_SubInstanceBase::CustomizeDetails(IDetailLayoutBuilder& Deta
 			];
 		}
 	}
-
-	IDetailCategoryBuilder& CategoryBuilder = DetailBuilder.EditCategory(FName(TEXT("Settings")));
-
-	// Customize InstanceClass
-	{
-		TSharedRef<IPropertyHandle> ClassHandle = DetailBuilder.GetProperty(TEXT("Node.InstanceClass"), GetClass());
-		ClassHandle->MarkHiddenByCustomization();
-
-		FDetailWidgetRow& ClassWidgetRow = CategoryBuilder.AddCustomRow(LOCTEXT("FilterStringInstanceClass", "Instance Class"));
-		ClassWidgetRow.NameContent()
-		[
-			ClassHandle->CreatePropertyNameWidget()
-		]
-		.ValueContent()
-		.MinDesiredWidth(250.0f)
-		[
-			SNew(SObjectPropertyEntryBox)
-			.ObjectPath_UObject(this, &UAnimGraphNode_SubInstanceBase::GetCurrentInstanceBlueprintPath)
-			.AllowedClass(UAnimBlueprint::StaticClass())
-			.NewAssetFactories(TArray<UFactory*>())
-			.OnShouldFilterAsset(FOnShouldFilterAsset::CreateUObject(this, &UAnimGraphNode_SubInstanceBase::OnShouldFilterInstanceBlueprint))
-			.OnObjectChanged(FOnSetObject::CreateUObject(this, &UAnimGraphNode_SubInstanceBase::OnSetInstanceBlueprint, &DetailBuilder))
-		];
-	}
 }
 
 bool UAnimGraphNode_SubInstanceBase::IsStructuralProperty(UProperty* InProperty) const
