Disallowing nesting Layer nodes inside of an Interface Animation Layer graph, still valid to nest a Layer node inside of a regular Animation Layer.

diff --git a/Engine/Source/Editor/AnimGraph/Private/AnimGraphNode_Layer.cpp b/Engine/Source/Editor/AnimGraph/Private/AnimGraphNode_Layer.cpp
--- a/Engine/Source/Editor/AnimGraph/Private/AnimGraphNode_Layer.cpp
+++ b/Engine/Source/Editor/AnimGraph/Private/AnimGraphNode_Layer.cpp
@@ -16,6 +16,7 @@
 #include "Widgets/Layout/SBox.h"
 #include "UObject/CoreRedirects.h"
 #include "Kismet2/KismetEditorUtilities.h"
+#include "AnimationStateGraph.h"
 
 #define LOCTEXT_NAMESPACE "LayerNode"
 
@@ -150,9 +151,22 @@ void UAnimGraphNode_Layer::ValidateAnimNodeDuringCompilation(USkeleton* ForSkele
 			TArray<UEdGraph*> Graphs;
 			CurrentBlueprint->GetAllGraphs(Graphs);
 
+			auto ValidateOuterGraph = [this, OriginalThis, &MessageLog](const UEdGraph* InGraph)
+			{
+				static const FName DefaultAnimGraphName("AnimGraph");
+				if (InGraph->Nodes.Contains(OriginalThis))
+				{
+					if (!InGraph->IsA<UAnimationStateGraph>() && InGraph->GetFName() != DefaultAnimGraphName)
+					{
+						MessageLog.Error(*FText::Format(LOCTEXT("NestedLayer", "Layer node @@ is part of Animation Layer Graph '{0}', layers cannot be nested."), FText::FromName(InGraph->GetFName())).ToString(), this);
+					}
+				}
+			};
+			
 			for(const UEdGraph* Graph : Graphs)
 			{
 				CheckGraph(Graph);
+				ValidateOuterGraph(Graph);				
 			}
 		}
 	}
