Renaming an animation interface graph now also fixes up layer nodes that use that graph

diff --git a/Engine/Source/Editor/AnimGraph/Classes/AnimGraphNode_LinkedAnimLayer.h b/Engine/Source/Editor/AnimGraph/Classes/AnimGraphNode_LinkedAnimLayer.h
--- a/Engine/Source/Editor/AnimGraph/Classes/AnimGraphNode_LinkedAnimLayer.h
+++ b/Engine/Source/Editor/AnimGraph/Classes/AnimGraphNode_LinkedAnimLayer.h
@@ -25,6 +25,13 @@ public:
 	UPROPERTY(EditAnywhere, Category = Settings)
 	FAnimNode_LinkedAnimLayer Node;
 
+	/** Guid of the named layer graph we refer to */
+	UPROPERTY()
+	FGuid InterfaceGuid;
+
+	// Begin UObject interface
+	virtual void Serialize(FArchive& Ar) override;
+
 	//~ Begin UEdGraphNode Interface.
 	virtual FText GetTooltipText() const override;
 	virtual FText GetNodeTitle(ENodeTitleType::Type TitleType) const override;
@@ -37,6 +44,9 @@ public:
 	// UAnimGraphNode_Base interface
 	virtual void CustomizeDetails(IDetailLayoutBuilder& DetailBuilder) override;
 
+	// Optionally updates layer GUID if it is invalid
+	void UpdateGuidForLayer();
+
 protected:
 	// ----- UI CALLBACKS ----- //
 	// Handlers for layer combo
@@ -63,6 +73,9 @@ protected:
 
 	// Helper function to get the interface currently in use by the selected layer
 	TSubclassOf<UInterface> GetInterfaceForLayer() const;
+
+	// Helper function to get the interface graph GUID currently in use by the selected layer
+	FGuid GetGuidForLayer() const;
 };
 
 UE_DEPRECATED(4.24, "UAnimGraphNode_Layer has been renamed to UAnimGraphNode_LinkedAnimLayer")
diff --git a/Engine/Source/Editor/AnimGraph/Classes/AnimationGraphSchema.h b/Engine/Source/Editor/AnimGraph/Classes/AnimationGraphSchema.h
--- a/Engine/Source/Editor/AnimGraph/Classes/AnimationGraphSchema.h
+++ b/Engine/Source/Editor/AnimGraph/Classes/AnimationGraphSchema.h
@@ -99,6 +99,9 @@ class UAnimationGraphSchema : public UEdGraphSchema_K2
 	/** Conforms an anim graph to an interface function */
 	ANIMGRAPH_API static void ConformAnimGraphToInterface(UBlueprint* InBlueprint, UEdGraph& InGraph, UFunction* InFunction);
 
+	/** Conforms anim layer nodes to an interface desc by GUID */
+	ANIMGRAPH_API static void ConformAnimLayersByGuid(const UAnimBlueprint* InAnimBlueprint, const FBPInterfaceDescription& CurrentInterfaceDesc);
+
 	UE_DEPRECATED(4.24, "Function renamed, please use GetPositionForNewLinkedInputPoseNode")
 	ANIMGRAPH_API static FVector2D GetPositionForNewSubInputNode(UEdGraph& InGraph) { return GetPositionForNewLinkedInputPoseNode(InGraph); }
 
diff --git a/Engine/Source/Editor/AnimGraph/Private/AnimGraphNode_LinkedAnimLayer.cpp b/Engine/Source/Editor/AnimGraph/Private/AnimGraphNode_LinkedAnimLayer.cpp
--- a/Engine/Source/Editor/AnimGraph/Private/AnimGraphNode_LinkedAnimLayer.cpp
+++ b/Engine/Source/Editor/AnimGraph/Private/AnimGraphNode_LinkedAnimLayer.cpp
@@ -17,9 +17,28 @@
 #include "UObject/CoreRedirects.h"
 #include "Kismet2/KismetEditorUtilities.h"
 #include "AnimationStateGraph.h"
+#include "UObject/FortniteMainBranchObjectVersion.h"
 
 #define LOCTEXT_NAMESPACE "LinkedAnimLayerNode"
 
+void UAnimGraphNode_LinkedAnimLayer::Serialize(FArchive& Ar)
+{
+	Super::Serialize(Ar);
+
+	Ar.UsingCustomVersion(FFortniteMainBranchObjectVersion::GUID);
+
+	if (Ar.IsLoading())
+	{
+		if (Ar.CustomVer(FFortniteMainBranchObjectVersion::GUID) < FFortniteMainBranchObjectVersion::AnimLayerGuidConformation)
+		{
+			if (!InterfaceGuid.IsValid())
+			{
+				InterfaceGuid = GetGuidForLayer();
+			}
+		}
+	}
+}
+
 FText UAnimGraphNode_LinkedAnimLayer::GetTooltipText() const
 {
 	return LOCTEXT("ToolTip", "Runs another linked animation layer graph to process animation");
@@ -509,18 +528,14 @@ TSubclassOf<UInterface> UAnimGraphNode_LinkedAnimLayer::GetInterfaceForLayer() c
 {
 	if (UAnimBlueprint* CurrentBlueprint = Cast<UAnimBlueprint>(GetBlueprint()))
 	{
-		UClass* TargetClass = *CurrentBlueprint->SkeletonGeneratedClass;
-		if(TargetClass)
+		// Find layer with this name in interfaces
+		for(FBPInterfaceDescription& InterfaceDesc : CurrentBlueprint->ImplementedInterfaces)
 		{
-			// Find layer with this name in interfaces
-			for(FBPInterfaceDescription& InterfaceDesc : CurrentBlueprint->ImplementedInterfaces)
+			for(UEdGraph* InterfaceGraph : InterfaceDesc.Graphs)
 			{
-				for(UEdGraph* InterfaceGraph : InterfaceDesc.Graphs)
+				if(InterfaceGraph->GetFName() == Node.Layer)
 				{
-					if(InterfaceGraph->GetFName() == Node.Layer)
-					{
-						return InterfaceDesc.Interface;
-					}
+					return InterfaceDesc.Interface;
 				}
 			}
 		}
@@ -529,6 +544,34 @@ TSubclassOf<UInterface> UAnimGraphNode_LinkedAnimLayer::GetInterfaceForLayer() c
 	return nullptr;
 }
 
+void UAnimGraphNode_Layer::UpdateGuidForLayer()
+{
+	if (!InterfaceGuid.IsValid())
+	{
+		InterfaceGuid = GetGuidForLayer();
+	}
+}
+
+FGuid UAnimGraphNode_Layer::GetGuidForLayer() const
+{
+	if (UAnimBlueprint* CurrentBlueprint = Cast<UAnimBlueprint>(GetBlueprint()))
+	{
+		// Find layer with this name in interfaces
+		for (FBPInterfaceDescription& InterfaceDesc : CurrentBlueprint->ImplementedInterfaces)
+		{
+			for (UEdGraph* InterfaceGraph : InterfaceDesc.Graphs)
+			{
+				if (InterfaceGraph->GetFName() == Node.Layer)
+				{
+					return InterfaceGraph->InterfaceGuid;
+				}
+			}
+		}
+	}
+
+	return FGuid();
+}
+
 void UAnimGraphNode_LinkedAnimLayer::OnLayerChanged(IDetailLayoutBuilder* DetailBuilder)
 {
 	OnStructuralPropertyChanged(DetailBuilder);
@@ -536,6 +579,9 @@ void UAnimGraphNode_LinkedAnimLayer::OnLayerChanged(IDetailLayoutBuilder* Detail
 	// Get the interface for this layer. If null, then we are using a 'self' layer.
 	Node.Interface = GetInterfaceForLayer();
 
+	// Update the Guid for conforming
+	InterfaceGuid = GetGuidForLayer();
+
 	if(Node.Interface.Get() == nullptr)
 	{
 		// Self layers cannot have override implementations
diff --git a/Engine/Source/Editor/AnimGraph/Private/AnimationGraphSchema.cpp b/Engine/Source/Editor/AnimGraph/Private/AnimationGraphSchema.cpp
--- a/Engine/Source/Editor/AnimGraph/Private/AnimationGraphSchema.cpp
+++ b/Engine/Source/Editor/AnimGraph/Private/AnimationGraphSchema.cpp
@@ -34,6 +34,7 @@
 #include "ScopedTransaction.h"
 #include "Animation/AnimMontage.h"
 #include "AnimGraphNode_LinkedInputPose.h"
+#include "AnimGraphNode_LinkedAnimLayer.h"
 
 #define LOCTEXT_NAMESPACE "AnimationGraphSchema"
 
@@ -755,6 +756,40 @@ void UAnimationGraphSchema::ConformAnimGraphToInterface(UBlueprint* InBlueprint,
 	}
 }
 
+void UAnimationGraphSchema::ConformAnimLayersByGuid(const UAnimBlueprint* InAnimBlueprint, const FBPInterfaceDescription& CurrentInterfaceDesc)
+{
+	const UBlueprint* InterfaceBlueprint = CastChecked<UBlueprint>(CurrentInterfaceDesc.Interface->ClassGeneratedBy);
+
+	TArray<UEdGraph*> InterfaceGraphs;
+	InterfaceBlueprint->GetAllGraphs(InterfaceGraphs);
+
+	TArray<UEdGraph*> Graphs;
+	InAnimBlueprint->GetAllGraphs(Graphs);
+
+	for (UEdGraph* Graph : Graphs)
+	{
+		TArray<UAnimGraphNode_LinkedInputLayer*> LayerNodes;
+		Graph->GetNodesOfClass<UAnimGraphNode_Layer>(LayerNodes);
+
+		for (UAnimGraphNode_LinkedInputLayer* LayerNode : LayerNodes)
+		{
+			LayerNode->UpdateGuidForLayer();
+
+			if (LayerNode->InterfaceGuid.IsValid())
+			{
+				for (UEdGraph* InterfaceGraph : InterfaceGraphs)
+				{
+					// Check to see if GUID matches but name does not and update if so
+					if (InterfaceGraph->GraphGuid == LayerNode->InterfaceGuid && InterfaceGraph->GetFName() != LayerNode->Node.Layer)
+					{
+						LayerNode->Node.Layer = InterfaceGraph->GetFName();
+					}
+				}
+			}
+		}
+	}
+}
+
 FVector2D UAnimationGraphSchema::GetPositionForNewLinkedInputPoseNode(UEdGraph& InGraph)
 {
 	TArray<UAnimGraphNode_LinkedInputPose*> LinkedInputPoseNodes;
diff --git a/Engine/Source/Editor/UnrealEd/Private/Kismet2/BlueprintEditorUtils.cpp b/Engine/Source/Editor/UnrealEd/Private/Kismet2/BlueprintEditorUtils.cpp
--- a/Engine/Source/Editor/UnrealEd/Private/Kismet2/BlueprintEditorUtils.cpp
+++ b/Engine/Source/Editor/UnrealEd/Private/Kismet2/BlueprintEditorUtils.cpp
@@ -6650,6 +6650,12 @@ void FBlueprintEditorUtils::ConformImplementedEvents(UBlueprint* Blueprint)
 /** Helper function for ConformImplementedInterfaces */
 static void ConformInterfaceByGUID(const UBlueprint* Blueprint, const FBPInterfaceDescription& CurrentInterfaceDesc)
 {
+	// Conform anim layers before interface graphs as GUIDs may need to be set up in older assets.
+	if (const UAnimBlueprint * AnimBlueprint = Cast<UAnimBlueprint>(Blueprint))
+	{
+		UAnimationGraphSchema::ConformAnimLayersByGuid(AnimBlueprint, CurrentInterfaceDesc);
+	}
+
 	// Attempt to conform by GUID if we have a blueprint interface
 	// This just make sure that GUID-linked functions preserve their names
 	const UBlueprint* InterfaceBlueprint = CastChecked<UBlueprint>(CurrentInterfaceDesc.Interface->ClassGeneratedBy);
diff --git a/Engine/Source/Runtime/Core/Public/UObject/FortniteMainBranchObjectVersion.h b/Engine/Source/Runtime/Core/Public/UObject/FortniteMainBranchObjectVersion.h
--- a/Engine/Source/Runtime/Core/Public/UObject/FortniteMainBranchObjectVersion.h
+++ b/Engine/Source/Runtime/Core/Public/UObject/FortniteMainBranchObjectVersion.h
@@ -99,6 +99,9 @@ struct CORE_API FFortniteMainBranchObjectVersion
 		// Replaced lake collision component
 		ReplaceLakeCollision,
 
+		// Anim layer node names are now conformed by Guid
+		AnimLayerGuidConformation,
+
 		// -----<new versions can be added above this line>-------------------------------------------------
 		VersionPlusOne,
 		LatestVersion = VersionPlusOne - 1
