Fix build breakage

diff --git a/Engine/Source/Editor/AnimGraph/Private/AnimationGraphSchema.cpp b/Engine/Source/Editor/AnimGraph/Private/AnimationGraphSchema.cpp
--- a/Engine/Source/Editor/AnimGraph/Private/AnimationGraphSchema.cpp
+++ b/Engine/Source/Editor/AnimGraph/Private/AnimationGraphSchema.cpp
@@ -768,10 +768,10 @@ void UAnimationGraphSchema::ConformAnimLayersByGuid(const UAnimBlueprint* InAnim
 
 	for (UEdGraph* Graph : Graphs)
 	{
-		TArray<UAnimGraphNode_LinkedInputLayer*> LayerNodes;
-		Graph->GetNodesOfClass<UAnimGraphNode_Layer>(LayerNodes);
+		TArray<UAnimGraphNode_LinkedAnimLayer*> LayerNodes;
+		Graph->GetNodesOfClass<UAnimGraphNode_LinkedAnimLayer>(LayerNodes);
 
-		for (UAnimGraphNode_LinkedInputLayer* LayerNode : LayerNodes)
+		for (UAnimGraphNode_LinkedAnimLayer* LayerNode : LayerNodes)
 		{
 			LayerNode->UpdateGuidForLayer();
 
