Correctly displaying current class on layer nodes

This change means that layer nodes can now display their (potentially dynamically-assigned) linked class when the debugger is inspecting a running instance.

diff --git a/Engine/Source/Editor/AnimGraph/Classes/AnimGraphNode_Base.h b/Engine/Source/Editor/AnimGraph/Classes/AnimGraphNode_Base.h
--- a/Engine/Source/Editor/AnimGraph/Classes/AnimGraphNode_Base.h
+++ b/Engine/Source/Editor/AnimGraph/Classes/AnimGraphNode_Base.h
@@ -288,6 +288,11 @@ class ANIMGRAPH_API UAnimGraphNode_Base : public UK2Node
 	 */
 	bool IsPinExposedAndLinked(const FString& InPinName, const EEdGraphPinDirection Direction = EGPD_MAX) const;
 
+	// Event that is broadcast to inform observers that the node title has changed
+	// The default SAnimationGraphNode uses this to invalidate cached node title text
+	DECLARE_EVENT(UAnimGraphNode_Base, FOnNodeTitleChangedEvent);
+	FOnNodeTitleChangedEvent& OnNodeTitleChangedEvent() { return NodeTitleChangedEvent; }
+
 protected:
 	friend FAnimBlueprintCompilerContext;
 	friend FAnimGraphNodeDetails;
@@ -314,6 +319,8 @@ protected:
 
 	FOnNodePropertyChangedEvent PropertyChangeEvent;
 
+	FOnNodeTitleChangedEvent NodeTitleChangedEvent;
+
 private:
 	TArray<FName> OldShownPins;
 };
diff --git a/Engine/Source/Editor/AnimGraph/Classes/AnimGraphNode_LinkedAnimLayer.h b/Engine/Source/Editor/AnimGraph/Classes/AnimGraphNode_LinkedAnimLayer.h
--- a/Engine/Source/Editor/AnimGraph/Classes/AnimGraphNode_LinkedAnimLayer.h
+++ b/Engine/Source/Editor/AnimGraph/Classes/AnimGraphNode_LinkedAnimLayer.h
@@ -39,6 +39,7 @@ public:
 	virtual UObject* GetJumpTargetForDoubleClick() const override;
 	virtual void JumpToDefinition() const override;
 	virtual bool HasExternalDependencies(TArray<class UStruct*>* OptionalOutput /*= NULL*/) const override;
+	virtual void ReconstructNode() override;
 	//~ End UEdGraphNode Interface.
 
 	// UAnimGraphNode_Base interface
@@ -55,6 +56,8 @@ protected:
 	void OnLayerChanged(IDetailLayoutBuilder* DetailBuilder);
 	bool HasAvailableLayers() const;
 	bool HasValidNonSelfLayer() const;
+	void HandleSetObjectBeingDebugged(UObject* InDebugObj);
+	void HandleInstanceChanged();
 	// ----- END UI CALLBACKS ----- //
 
 	// Begin UAnimGraphNode_CustomProperty
@@ -76,6 +79,12 @@ protected:
 
 	// Helper function to get the interface graph GUID currently in use by the selected layer
 	FGuid GetGuidForLayer() const;
+
+	// Get the preview node, if any, when instanced in an animation blueprint and debugged
+	FAnimNode_LinkedAnimLayer* GetPreviewNode() const;
+
+	// Handle used to hook into object being debugged changing
+	FDelegateHandle SetObjectBeingDebuggedHandle;
 };
 
 UE_DEPRECATED(4.24, "UAnimGraphNode_Layer has been renamed to UAnimGraphNode_LinkedAnimLayer")
diff --git a/Engine/Source/Editor/AnimGraph/Private/AnimGraphNode_LinkedAnimLayer.cpp b/Engine/Source/Editor/AnimGraph/Private/AnimGraphNode_LinkedAnimLayer.cpp
--- a/Engine/Source/Editor/AnimGraph/Private/AnimGraphNode_LinkedAnimLayer.cpp
+++ b/Engine/Source/Editor/AnimGraph/Private/AnimGraphNode_LinkedAnimLayer.cpp
@@ -39,33 +39,83 @@ void UAnimGraphNode_LinkedAnimLayer::Serialize(FArchive& Ar)
 	}
 }
 
+void UAnimGraphNode_LinkedAnimLayer::ReconstructNode()
+{
+	if(SetObjectBeingDebuggedHandle.IsValid())
+	{
+		GetBlueprint()->OnSetObjectBeingDebugged().Remove(SetObjectBeingDebuggedHandle);
+	}
+
+	SetObjectBeingDebuggedHandle = GetBlueprint()->OnSetObjectBeingDebugged().AddUObject(this, &UAnimGraphNode_LinkedAnimLayer::HandleSetObjectBeingDebugged);
+
+	Super::ReconstructNode();
+}
+
 FText UAnimGraphNode_LinkedAnimLayer::GetTooltipText() const
 {
 	return LOCTEXT("ToolTip", "Runs another linked animation layer graph to process animation");
 }
 
-FText UAnimGraphNode_LinkedAnimLayer::GetNodeTitle(ENodeTitleType::Type TitleType) const
+FAnimNode_LinkedAnimLayer* UAnimGraphNode_LinkedAnimLayer::GetPreviewNode() const
 {
-	UClass* TargetClass = *Node.Interface;
-	UAnimBlueprint* TargetAnimBlueprint = TargetClass ? CastChecked<UAnimBlueprint>(TargetClass->ClassGeneratedBy) : nullptr;
+	FAnimNode_LinkedAnimLayer* PreviewNode = nullptr;
+	USkeletalMeshComponent* Component = nullptr;
 
-	FFormatNamedArguments Args;
-	Args.Add(TEXT("NodeTitle"), LOCTEXT("Title", "Linked Anim Layer"));
-	Args.Add(TEXT("TargetClass"), TargetAnimBlueprint ? FText::FromString(TargetAnimBlueprint->GetName()) : LOCTEXT("ClassSelf", "Self"));
+	// look for a valid component in the object being debugged,
+	// we might be set to something other than the preview.
+	UObject* ObjectBeingDebugged = GetAnimBlueprint()->GetObjectBeingDebugged();
+	if (ObjectBeingDebugged)
+	{
+		UAnimInstance* InstanceBeingDebugged = Cast<UAnimInstance>(ObjectBeingDebugged);
+		if (InstanceBeingDebugged)
+		{
+			Component = InstanceBeingDebugged->GetSkelMeshComponent();
+		}
+	}
 
-	if(TitleType == ENodeTitleType::MenuTitle)
+	if (Component != nullptr && Component->GetAnimInstance() != nullptr)
 	{
-		return LOCTEXT("NodeTitle", "Layer");
+		PreviewNode = static_cast<FAnimNode_LinkedAnimLayer*>(FindDebugAnimNode(Component));
 	}
-	else if(TitleType == ENodeTitleType::ListView)
+
+	return PreviewNode;
+}
+
+FText UAnimGraphNode_LinkedAnimLayer::GetNodeTitle(ENodeTitleType::Type TitleType) const
+{
+	UClass* TargetClass = *Node.Interface;
+	UAnimBlueprint* TargetAnimBlueprintInterface = TargetClass ? CastChecked<UAnimBlueprint>(TargetClass->ClassGeneratedBy) : nullptr;
+
+	if(TitleType == ENodeTitleType::MenuTitle)
 	{
-		Args.Add(TEXT("Layer"), Node.Layer == NAME_None ? LOCTEXT("LayerNone", "None") : FText::FromName(Node.Layer));
-		return FText::Format(LOCTEXT("TitleListFormatOutputPose", "{NodeTitle}: {Layer} - {TargetClass}"), Args);
+		return LOCTEXT("NodeTitle", "Linked Anim Layer");
 	}
 	else
 	{
+		FFormatNamedArguments Args;
+		Args.Add(TEXT("NodeTitle"), LOCTEXT("NodeTitle", "Linked Anim Layer"));
+		Args.Add(TEXT("TargetClass"), TargetAnimBlueprintInterface ? FText::FromString(TargetAnimBlueprintInterface->GetName()) : LOCTEXT("InterfaceNone", "None"));
 		Args.Add(TEXT("Layer"), Node.Layer == NAME_None ? LOCTEXT("LayerNone", "None") : FText::FromName(Node.Layer));
-		return FText::Format(LOCTEXT("TitleFormatOutputPose", "{NodeTitle}: {Layer}\n{TargetClass}"), Args);
+
+		if (FAnimNode_LinkedAnimLayer* PreviewNode = GetPreviewNode())
+		{
+			if (UAnimInstance* PreviewAnimInstance = PreviewNode->GetTargetInstance<UAnimInstance>())
+			{
+				if (UClass* PreviewTargetClass = PreviewAnimInstance->GetClass())
+				{
+					Args.Add(TEXT("TargetClass"), PreviewTargetClass == GetAnimBlueprint()->GeneratedClass ? LOCTEXT("ClassSelf", "Self") : FText::FromName(PreviewTargetClass->GetFName()));
+				}
+			}
+		}
+
+		if (TitleType == ENodeTitleType::ListView)
+		{
+			return FText::Format(LOCTEXT("TitleListFormatOutputPose", "{NodeTitle}: {Layer} - {TargetClass}"), Args);
+		}
+		else
+		{
+			return FText::Format(LOCTEXT("TitleFormatOutputPose", "{NodeTitle}: {Layer}\n{TargetClass}"), Args);
+		}
 	}
 }
 
@@ -632,4 +682,20 @@ bool UAnimGraphNode_LinkedAnimLayer::HasValidNonSelfLayer() const
 	return false;
 }
 
+void UAnimGraphNode_LinkedAnimLayer::HandleSetObjectBeingDebugged(UObject* InDebugObj)
+{
+	NodeTitleChangedEvent.Broadcast();
+
+	FAnimNode_LinkedAnimLayer* PreviewNode = GetPreviewNode();
+	if(PreviewNode)
+	{
+		PreviewNode->OnInstanceChanged().AddUObject(this, &UAnimGraphNode_LinkedAnimLayer::HandleInstanceChanged);
+	}
+}
+
+void UAnimGraphNode_LinkedAnimLayer::HandleInstanceChanged()
+{
+	NodeTitleChangedEvent.Broadcast();
+}
+
 #undef LOCTEXT_NAMESPACE
diff --git a/Engine/Source/Editor/AnimationBlueprintEditor/Private/AnimationNodes/SAnimationGraphNode.cpp b/Engine/Source/Editor/AnimationBlueprintEditor/Private/AnimationNodes/SAnimationGraphNode.cpp
--- a/Engine/Source/Editor/AnimationBlueprintEditor/Private/AnimationNodes/SAnimationGraphNode.cpp
+++ b/Engine/Source/Editor/AnimationBlueprintEditor/Private/AnimationNodes/SAnimationGraphNode.cpp
@@ -191,3 +191,23 @@ FReply SAnimationGraphNode::SpawnColourPicker()
 
 	return FReply::Handled();
 }
+
+TSharedRef<SWidget> SAnimationGraphNode::CreateTitleWidget(TSharedPtr<SNodeTitle> InNodeTitle)
+{
+	// Store title widget reference
+	NodeTitle = InNodeTitle;
+
+	// hook up invalidation delegate
+	UAnimGraphNode_Base* AnimGraphNode = CastChecked<UAnimGraphNode_Base>(GraphNode);
+	AnimGraphNode->OnNodeTitleChangedEvent().AddSP(this, &SAnimationGraphNode::HandleNodeTitleChanged);
+
+	return SGraphNodeK2Base::CreateTitleWidget(InNodeTitle);
+}
+
+void SAnimationGraphNode::HandleNodeTitleChanged()
+{
+	if(NodeTitle.IsValid())
+	{
+		NodeTitle->MarkDirty();
+	}
+}
\ No newline at end of file
diff --git a/Engine/Source/Editor/AnimationBlueprintEditor/Private/AnimationNodes/SAnimationGraphNode.h b/Engine/Source/Editor/AnimationBlueprintEditor/Private/AnimationNodes/SAnimationGraphNode.h
--- a/Engine/Source/Editor/AnimationBlueprintEditor/Private/AnimationNodes/SAnimationGraphNode.h
+++ b/Engine/Source/Editor/AnimationBlueprintEditor/Private/AnimationNodes/SAnimationGraphNode.h
@@ -27,6 +27,7 @@ protected:
 
 	// SGraphNode interface
 	virtual TArray<FOverlayWidgetInfo> GetOverlayWidgets(bool bSelected, const FVector2D& WidgetSize) const override;
+	virtual TSharedRef<SWidget> CreateTitleWidget(TSharedPtr<SNodeTitle> InNodeTitle) override;
 	// End of SGraphNode interface
 
 private:
@@ -35,11 +36,17 @@ private:
 
 	FReply SpawnColourPicker();
 
+	// Handle the node informing us that the title has changed
+	void HandleNodeTitleChanged();
+
 	/** Keep a reference to the indicator widget handing around */
 	TSharedPtr<SWidget> IndicatorWidget;
 
 	/** Keep a reference to the pose view indicator widget handing around */
 	TSharedPtr<SWidget> PoseViewWidget;
 
+	/** Cache the node title so we can invalidate it */
+	TSharedPtr<SNodeTitle> NodeTitle;
+
 	TWeakObjectPtr<class UPoseWatch> PoseWatch;
 };
diff --git a/Engine/Source/Editor/GraphEditor/Private/SGraphNode.cpp b/Engine/Source/Editor/GraphEditor/Private/SGraphNode.cpp
--- a/Engine/Source/Editor/GraphEditor/Private/SGraphNode.cpp
+++ b/Engine/Source/Editor/GraphEditor/Private/SGraphNode.cpp
@@ -89,6 +89,11 @@ FVector2D SNodeTitle::GetTitleSize() const
 	return CachedSize;
 }
 
+void SNodeTitle::MarkDirty()
+{
+	NodeTitleCache.MarkDirty();
+}
+
 void SNodeTitle::RebuildWidget()
 {
 	// Create the box to contain the lines
diff --git a/Engine/Source/Editor/GraphEditor/Public/SGraphNode.h b/Engine/Source/Editor/GraphEditor/Public/SGraphNode.h
--- a/Engine/Source/Editor/GraphEditor/Public/SGraphNode.h
+++ b/Engine/Source/Editor/GraphEditor/Public/SGraphNode.h
@@ -60,6 +60,9 @@ public:
 	/** Get the size of this title the last time it was drawn */
 	FVector2D GetTitleSize() const;
 
+	/** Mark the cached title dirty */
+	void MarkDirty();
+
 protected:
 	TWeakObjectPtr<UEdGraphNode> GraphNode;
 	FNodeTextCache NodeTitleCache;
diff --git a/Engine/Source/Runtime/Engine/Classes/Animation/AnimNode_LinkedAnimLayer.h b/Engine/Source/Runtime/Engine/Classes/Animation/AnimNode_LinkedAnimLayer.h
--- a/Engine/Source/Runtime/Engine/Classes/Animation/AnimNode_LinkedAnimLayer.h
+++ b/Engine/Source/Runtime/Engine/Classes/Animation/AnimNode_LinkedAnimLayer.h
@@ -49,8 +49,18 @@ public:
 		return *Interface;
 	}
 
+#if WITH_EDITOR
+	// Event fired when the instance we are running has changed
+	FSimpleMulticastDelegate& OnInstanceChanged() { return OnInstanceChangedEvent; }
+#endif
+
 protected:
 	void InitializeSelfLayer(const UAnimInstance* SelfAnimInstance);
+
+#if WITH_EDITOR
+	// Event fired when the instance we are running has changed
+	FSimpleMulticastDelegate OnInstanceChangedEvent;
+#endif
 };
 
 UE_DEPRECATED(4.24, "FAnimNode_Layer has been renamed to FAnimNode_LinkedAnimLayer")
diff --git a/Engine/Source/Runtime/Engine/Private/Animation/AnimNode_LinkedAnimLayer.cpp b/Engine/Source/Runtime/Engine/Private/Animation/AnimNode_LinkedAnimLayer.cpp
--- a/Engine/Source/Runtime/Engine/Private/Animation/AnimNode_LinkedAnimLayer.cpp
+++ b/Engine/Source/Runtime/Engine/Private/Animation/AnimNode_LinkedAnimLayer.cpp
@@ -66,4 +66,8 @@ void FAnimNode_LinkedAnimLayer::SetLinkedLayerInstance(const UAnimInstance* InOw
 	{
 		InitializeSelfLayer(InOwningAnimInstance);
 	}
+
+#if WITH_EDITOR
+	OnInstanceChangedEvent.Broadcast();
+#endif
 }
