Fix non-unity build

diff --git a/Engine/Source/Editor/AnimGraph/Classes/AnimationGraphSchema.h b/Engine/Source/Editor/AnimGraph/Classes/AnimationGraphSchema.h
--- a/Engine/Source/Editor/AnimGraph/Classes/AnimationGraphSchema.h
+++ b/Engine/Source/Editor/AnimGraph/Classes/AnimationGraphSchema.h
@@ -12,6 +12,8 @@
 
 class FMenuBuilder;
 class UAnimationAsset;
+class UAnimBlueprint;
+struct FBPInterfaceDescription;
 
 UCLASS(MinimalAPI)
 class UAnimationGraphSchema : public UEdGraphSchema_K2
