Fix deprecation warnings
#rnx

diff --git a/Engine/Source/Editor/AnimGraph/Private/AnimGraphNode_LinkedAnimLayer.cpp b/Engine/Source/Editor/AnimGraph/Private/AnimGraphNode_LinkedAnimLayer.cpp
--- a/Engine/Source/Editor/AnimGraph/Private/AnimGraphNode_LinkedAnimLayer.cpp
+++ b/Engine/Source/Editor/AnimGraph/Private/AnimGraphNode_LinkedAnimLayer.cpp
@@ -594,7 +594,7 @@ TSubclassOf<UInterface> UAnimGraphNode_LinkedAnimLayer::GetInterfaceForLayer() c
 	return nullptr;
 }
 
-void UAnimGraphNode_Layer::UpdateGuidForLayer()
+void UAnimGraphNode_LinkedAnimLayer::UpdateGuidForLayer()
 {
 	if (!InterfaceGuid.IsValid())
 	{
@@ -602,7 +602,7 @@ void UAnimGraphNode_Layer::UpdateGuidForLayer()
 	}
 }
 
-FGuid UAnimGraphNode_Layer::GetGuidForLayer() const
+FGuid UAnimGraphNode_LinkedAnimLayer::GetGuidForLayer() const
 {
 	if (UAnimBlueprint* CurrentBlueprint = Cast<UAnimBlueprint>(GetBlueprint()))
 	{
