Fixed blend logic reverting to standard blend when switching from Custom to Inertialization.

UAnimationCustomTransitionSchema::HandleGraphBeingDeleted resets LogicType to standard blend, so I made UAnimStateTransitionNode::PostEditChangeProperty restore the desired logic type after graph removal.

diff --git a/Engine/Source/Editor/AnimGraph/Private/AnimStateTransitionNode.cpp b/Engine/Source/Editor/AnimGraph/Private/AnimStateTransitionNode.cpp
--- a/Engine/Source/Editor/AnimGraph/Private/AnimStateTransitionNode.cpp
+++ b/Engine/Source/Editor/AnimGraph/Private/AnimStateTransitionNode.cpp
@@ -331,9 +331,14 @@ void UAnimStateTransitionNode::PostEditChangeProperty(struct FPropertyChangedEve
 		}
 		else if (CustomTransitionGraph != NULL)
 		{
+			// UAnimationCustomTransitionSchema::HandleGraphBeingDeleted resets logic type, so we'll need to restore it after RemoveGraph
+			const TEnumAsByte<ETransitionLogicType::Type> DesiredLogicType = LogicType;
+
 			UBlueprint* Blueprint = FBlueprintEditorUtils::FindBlueprintForNodeChecked(this);
 			FBlueprintEditorUtils::RemoveGraph(Blueprint, CustomTransitionGraph);
 			CustomTransitionGraph = NULL;
+
+			LogicType = DesiredLogicType;
 		}	
 	}
 
