Clamp transition crossfade durations to 0

diff --git a/Engine/Source/Editor/AnimGraph/Classes/AnimStateTransitionNode.h b/Engine/Source/Editor/AnimGraph/Classes/AnimStateTransitionNode.h
--- a/Engine/Source/Editor/AnimGraph/Classes/AnimStateTransitionNode.h
+++ b/Engine/Source/Editor/AnimGraph/Classes/AnimStateTransitionNode.h
@@ -35,7 +35,7 @@ class UAnimStateTransitionNode : public UAnimStateNodeBase
 	int32 PriorityOrder;
 
 	// The duration to cross-fade for
-	UPROPERTY(EditAnywhere, Config, Category=Transition)
+	UPROPERTY(EditAnywhere, Config, Category=Transition, Meta=(ClampMin="0.0"))
 	float CrossfadeDuration;
 
 	// The type of blending to use in the crossfade
