Fixed crash when deleting a state machine inside an animation layer while the layer node is still visible

diff --git a/Engine/Source/Editor/AnimGraph/Private/AnimGraphNode_LinkedAnimLayer.cpp b/Engine/Source/Editor/AnimGraph/Private/AnimGraphNode_LinkedAnimLayer.cpp
--- a/Engine/Source/Editor/AnimGraph/Private/AnimGraphNode_LinkedAnimLayer.cpp
+++ b/Engine/Source/Editor/AnimGraph/Private/AnimGraphNode_LinkedAnimLayer.cpp
@@ -684,12 +684,15 @@ bool UAnimGraphNode_LinkedAnimLayer::HasValidNonSelfLayer() const
 
 void UAnimGraphNode_LinkedAnimLayer::HandleSetObjectBeingDebugged(UObject* InDebugObj)
 {
-	NodeTitleChangedEvent.Broadcast();
-
-	FAnimNode_LinkedAnimLayer* PreviewNode = GetPreviewNode();
-	if(PreviewNode)
+	if(HasValidBlueprint())
 	{
-		PreviewNode->OnInstanceChanged().AddUObject(this, &UAnimGraphNode_LinkedAnimLayer::HandleInstanceChanged);
+		NodeTitleChangedEvent.Broadcast();
+
+		FAnimNode_LinkedAnimLayer* PreviewNode = GetPreviewNode();
+		if(PreviewNode)
+		{
+			PreviewNode->OnInstanceChanged().AddUObject(this, &UAnimGraphNode_LinkedAnimLayer::HandleInstanceChanged);
+		}
 	}
 }
 
