Animation layer input's type can no longer be reset to default

diff --git a/Engine/Source/Editor/AnimGraph/Private/AnimBlueprintPinInfoDetails.cpp b/Engine/Source/Editor/AnimGraph/Private/AnimBlueprintPinInfoDetails.cpp
--- a/Engine/Source/Editor/AnimGraph/Private/AnimBlueprintPinInfoDetails.cpp
+++ b/Engine/Source/Editor/AnimGraph/Private/AnimBlueprintPinInfoDetails.cpp
@@ -42,6 +42,8 @@ void FAnimBlueprintFunctionPinInfoDetails::CustomizeHeader(TSharedRef<IPropertyH
 	NamePropertyHandle = InStructPropertyHandle->GetChildHandle(GET_MEMBER_NAME_CHECKED(FAnimBlueprintFunctionPinInfo, Name));
 	TypePropertyHandle = InStructPropertyHandle->GetChildHandle(GET_MEMBER_NAME_CHECKED(FAnimBlueprintFunctionPinInfo, Type));
 
+	StructPropertyHandle->MarkResetToDefaultCustomized(true);
+
 	InHeaderRow
 	.DiffersFromDefault(false)
 	.NameContent()
