Fixed details panel expansion state not persisting when adding an animation layer

diff --git a/Engine/Source/Editor/AnimGraph/Private/AnimGraphDetails.cpp b/Engine/Source/Editor/AnimGraph/Private/AnimGraphDetails.cpp
--- a/Engine/Source/Editor/AnimGraph/Private/AnimGraphDetails.cpp
+++ b/Engine/Source/Editor/AnimGraph/Private/AnimGraphDetails.cpp
@@ -135,6 +135,7 @@ void FAnimGraphDetails::CustomizeDetails(IDetailLayoutBuilder& DetailLayout)
 	}
 
 	IDetailCategoryBuilder& InputsCategory = DetailLayout.EditCategory("Inputs", LOCTEXT("LinkedInputPoseInputsCategory", "Inputs"));
+	InputsCategory.RestoreExpansionState(true);
 
 	DetailLayoutBuilder = &DetailLayout;
 
@@ -200,9 +201,11 @@ void FAnimGraphDetails::CustomizeDetails(IDetailLayoutBuilder& DetailLayout)
 
 			TArray<UObject*> ExternalObjects;
 			ExternalObjects.Add(LinkedInputPoseNode);
-			if(IDetailPropertyRow* LinkedInputPoseRow = InputsCategory.AddExternalObjects(ExternalObjects))
+			FAddPropertyParams AddPropertyParams;
+			AddPropertyParams.UniqueId(LinkedInputPoseNode->GetFName());
+
+			if(IDetailPropertyRow* LinkedInputPoseRow = InputsCategory.AddExternalObjects(ExternalObjects, EPropertyLocation::Default, AddPropertyParams))
 			{
-				LinkedInputPoseRow->ShouldAutoExpand(true);
 				LinkedInputPoseRow->CustomWidget()
 				.NameContent()
 				[
