Fix 'Add pin for element' section not listing items in Animation Blueprint Node context menu

diff --git a/Engine/Source/Editor/AnimGraph/Private/AnimGraphNode_BlendListByEnum.cpp b/Engine/Source/Editor/AnimGraph/Private/AnimGraphNode_BlendListByEnum.cpp
--- a/Engine/Source/Editor/AnimGraph/Private/AnimGraphNode_BlendListByEnum.cpp
+++ b/Engine/Source/Editor/AnimGraph/Private/AnimGraphNode_BlendListByEnum.cpp
@@ -103,7 +103,7 @@ void UAnimGraphNode_BlendListByEnum::GetNodeContextMenuActions(UToolMenu* Menu,
 		}
 
 		// Offer to add any not-currently-visible pins
-		bool bAddedHeader = false;
+		FToolMenuSection* Section = nullptr;
 		const int32 MaxIndex = BoundEnum->NumEnums() - 1; // we don't want to show _MAX enum
 		for (int32 Index = 0; Index < MaxIndex; ++Index)
 		{
@@ -113,21 +113,13 @@ void UAnimGraphNode_BlendListByEnum::GetNodeContextMenuActions(UToolMenu* Menu,
 				FText PrettyElementName = BoundEnum->GetDisplayNameTextByIndex(Index);
 
 				// Offer to add this entry
-				if (!bAddedHeader)
+				if (!Section)
 				{
-					bAddedHeader = true;
-					{
-						FToolMenuSection& Section = Menu->AddSection("AnimGraphNodeAddElementPin", LOCTEXT("ExposeHeader", "Add pin for element"));
-						FUIAction Action = FUIAction( FExecuteAction::CreateUObject( const_cast<UAnimGraphNode_BlendListByEnum*>(this), &UAnimGraphNode_BlendListByEnum::ExposeEnumElementAsPin, ElementName) );
-						Section.AddMenuEntry("PrettyElementName", PrettyElementName, PrettyElementName, FSlateIcon(), Action);
-					}
-				}
-				else
-				{
-					FUIAction Action = FUIAction( FExecuteAction::CreateUObject( const_cast<UAnimGraphNode_BlendListByEnum*>(this), &UAnimGraphNode_BlendListByEnum::ExposeEnumElementAsPin, ElementName) );
-					FToolMenuSection& Section = Menu->AddSection("PrettyElementName");
-					Section.AddMenuEntry("PrettyElementName", PrettyElementName, PrettyElementName, FSlateIcon(), Action);
+					Section = &Menu->AddSection("AnimGraphNodeAddElementPin", LOCTEXT("ExposeHeader", "Add pin for element"));
 				}
+
+				FUIAction Action = FUIAction( FExecuteAction::CreateUObject( const_cast<UAnimGraphNode_BlendListByEnum*>(this), &UAnimGraphNode_BlendListByEnum::ExposeEnumElementAsPin, ElementName) );
+				Section->AddMenuEntry(NAME_None, PrettyElementName, PrettyElementName, FSlateIcon(), Action);
 			}
 		}
 	}
