Fixing CIS warnings

diff --git a/Engine/Source/Editor/AnimGraph/Private/AnimGraphNode_LinkedAnimLayer.cpp b/Engine/Source/Editor/AnimGraph/Private/AnimGraphNode_LinkedAnimLayer.cpp
--- a/Engine/Source/Editor/AnimGraph/Private/AnimGraphNode_LinkedAnimLayer.cpp
+++ b/Engine/Source/Editor/AnimGraph/Private/AnimGraphNode_LinkedAnimLayer.cpp
@@ -506,7 +506,7 @@ FString UAnimGraphNode_LinkedAnimLayer::GetCurrentInstanceBlueprintPath() const
 	return FString();
 }
 
-void UAnimGraphNode_Layer::GetExposableProperties(TArray<FProperty*>& OutExposableProperties) const
+void UAnimGraphNode_LinkedAnimLayer::GetExposableProperties(TArray<FProperty*>& OutExposableProperties) const
 {
 	UClass* TargetClass = GetTargetSkeletonClass();
 	if(TargetClass)
diff --git a/Engine/Source/Editor/Kismet/Private/BlueprintCompilationManager.cpp b/Engine/Source/Editor/Kismet/Private/BlueprintCompilationManager.cpp
--- a/Engine/Source/Editor/Kismet/Private/BlueprintCompilationManager.cpp
+++ b/Engine/Source/Editor/Kismet/Private/BlueprintCompilationManager.cpp
@@ -2870,7 +2870,7 @@ FArchive& FBlueprintCompilationManagerImpl::FFixupBytecodeReferences::operator<<
 		{
 			UClass* NewClass = OwningClass->GetAuthoritativeClass();
 			ensure(NewClass);
-			if (NewClass != OwningClass)
+			if (NewClass && NewClass != OwningClass)
 			{
 				// drill into new class finding equivalent object:
 				TArray<FFieldVariant> OwnerChain;
