Fix crash when native control rig is selected as an option

diff --git a/Engine/Source/Editor/AnimGraph/Private/AnimGraphNode_CustomProperty.cpp b/Engine/Source/Editor/AnimGraph/Private/AnimGraphNode_CustomProperty.cpp
--- a/Engine/Source/Editor/AnimGraph/Private/AnimGraphNode_CustomProperty.cpp
+++ b/Engine/Source/Editor/AnimGraph/Private/AnimGraphNode_CustomProperty.cpp
@@ -305,7 +305,7 @@ UClass* UAnimGraphNode_CustomProperty::GetTargetClass() const
 UClass* UAnimGraphNode_CustomProperty::GetTargetSkeletonClass() const
 {
 	UClass* TargetClass = GetTargetClass();
-	if(TargetClass)
+	if(TargetClass && TargetClass->ClassGeneratedBy)
 	{
 		UBlueprint* Blueprint = CastChecked<UBlueprint>(TargetClass->ClassGeneratedBy);
 		if(Blueprint)
