Remove Persona as a dependency as it is unneeded and causes a cyclic dependency
#rnx

diff --git a/Engine/Source/Editor/AnimGraph/AnimGraph.Build.cs b/Engine/Source/Editor/AnimGraph/AnimGraph.Build.cs
--- a/Engine/Source/Editor/AnimGraph/AnimGraph.Build.cs
+++ b/Engine/Source/Editor/AnimGraph/AnimGraph.Build.cs
@@ -35,7 +35,6 @@ public class AnimGraph : ModuleRules
                 "ContentBrowser",
 				"KismetWidgets",
 				"ToolMenus",
-				"Persona",
 				"EditorWidgets",
 			}
 		);
