Restricted anim BP compiler subsystem APIs

Removed direct access to currently compiling BP class to restrict unguarded mutation.
Const-corrected various accessors and overrides.
Moved UObject-based compiler subsystems to 'handlers' and removed UObject dependency.
Moved from set of virtual functions on subsystems/handlers to a restricted set of contexts passed to specific multicast delegates that handlers subscribe to.
Removed anim class subsystems.
Moved property access code to engine module (editor code is still in a plugin).
Fixed nativized builds where anim BPs have nativized/non-nativized classes in the hierarchy

diff --git a/Engine/Plugins/Developer/PropertyAccessNode/PropertyAccessNode.uplugin b/Engine/Plugins/Developer/PropertyAccessNode/PropertyAccessNode.uplugin
--- a/Engine/Plugins/Developer/PropertyAccessNode/PropertyAccessNode.uplugin
+++ b/Engine/Plugins/Developer/PropertyAccessNode/PropertyAccessNode.uplugin
@@ -26,7 +26,7 @@
 	"Plugins":
 	[
 		{
-			"Name": "PropertyAccess",
+			"Name": "PropertyAccessEditor",
 			"Enabled": true
 		}
 	]
diff --git a/Engine/Plugins/Developer/PropertyAccessNode/Source/PropertyAccessNode/Private/K2Node_PropertyAccess.cpp b/Engine/Plugins/Developer/PropertyAccessNode/Source/PropertyAccessNode/Private/K2Node_PropertyAccess.cpp
--- a/Engine/Plugins/Developer/PropertyAccessNode/Source/PropertyAccessNode/Private/K2Node_PropertyAccess.cpp
+++ b/Engine/Plugins/Developer/PropertyAccessNode/Source/PropertyAccessNode/Private/K2Node_PropertyAccess.cpp
@@ -15,27 +15,25 @@
 #include "K2Node_VariableGet.h"
 #include "AnimationGraph.h"
 #include "Kismet2/BlueprintEditorUtils.h"
-#include "IPropertyAccessCompilerSubsystem.h"
-#include "AnimBlueprintCompilerSubsystem.h"
+#include "PropertyAccessCompilerHandler.h"
 #include "Modules/ModuleManager.h"
 #include "IPropertyAccessEditor.h"
 #include "IPropertyAccessCompiler.h"
 #include "Features/IModularFeatures.h"
+#include "IAnimBlueprintCompilationContext.h"
 
 #define LOCTEXT_NAMESPACE "K2Node_PropertyAccess"
 
-void UK2Node_PropertyAccess::CreateClassVariablesFromBlueprint(FKismetCompilerContext& InCompilerContext)
+void UK2Node_PropertyAccess::CreateClassVariablesFromBlueprint(IAnimBlueprintVariableCreationContext& InCreationContext)
 {
 	GeneratedPropertyName = NAME_None;
 
 	if(ResolvedPinType != FEdGraphPinType() && ResolvedPinType.PinCategory != UEdGraphSchema_K2::PC_Wildcard)
 	{
 		// Create internal generated destination property
-		if(FProperty* DestProperty = FKismetCompilerUtilities::CreatePropertyOnScope(InCompilerContext.NewClass, GetFName(), ResolvedPinType, InCompilerContext.NewClass, CPF_None, CastChecked<UEdGraphSchema_K2>(GetSchema()), InCompilerContext.MessageLog))
+		if(FProperty* DestProperty = InCreationContext.CreateVariable(GetFName(), ResolvedPinType))
 		{
 			GeneratedPropertyName = DestProperty->GetFName();
-
-			FKismetCompilerUtilities::LinkAddedProperty(InCompilerContext.NewClass, DestProperty);
 		}
 	}
 }
@@ -50,9 +48,10 @@ void UK2Node_PropertyAccess::ExpandNode(FKismetCompilerContext& InCompilerContex
 		DestPropertyPath.Add(GeneratedPropertyName.ToString());
 
 		// Create a copy event in the complied generated class
-		IPropertyAccessCompilerSubsystem* PropertyAccessSubsystem = UAnimBlueprintCompilerSubsystem::FindSubsystemWithInterface<IPropertyAccessCompilerSubsystem>(InCompilerContext);
-		check(PropertyAccessSubsystem);
-		PropertyAccessSubsystem->AddCopy(Path, DestPropertyPath, EPropertyAccessBatchType::Batched, this);
+		TUniquePtr<IAnimBlueprintCompilationContext> CompilationContext = IAnimBlueprintCompilationContext::Get(InCompilerContext);
+		FPropertyAccessCompilerHandler* PropertyAccessHandler = CompilationContext->GetHandler<FPropertyAccessCompilerHandler>("PropertyAccessCompilerHandler");
+		check(PropertyAccessHandler);
+		PropertyAccessHandler->AddCopy(Path, DestPropertyPath, EPropertyAccessBatchType::Batched, this);
 
 		// Replace us with a get node
 		UK2Node_VariableGet* VariableGetNode = InCompilerContext.SpawnIntermediateNode<UK2Node_VariableGet>(this, InSourceGraph);
diff --git a/Engine/Plugins/Developer/PropertyAccessNode/Source/PropertyAccessNode/Private/K2Node_PropertyAccess.h b/Engine/Plugins/Developer/PropertyAccessNode/Source/PropertyAccessNode/Private/K2Node_PropertyAccess.h
--- a/Engine/Plugins/Developer/PropertyAccessNode/Source/PropertyAccessNode/Private/K2Node_PropertyAccess.h
+++ b/Engine/Plugins/Developer/PropertyAccessNode/Source/PropertyAccessNode/Private/K2Node_PropertyAccess.h
@@ -64,7 +64,7 @@ private:
 	int32 ResolvedArrayIndex = INDEX_NONE;
 
 	// IClassVariableCreator interface
-	virtual void CreateClassVariablesFromBlueprint(FKismetCompilerContext& InCompilerContext) override;
+	virtual void CreateClassVariablesFromBlueprint(IAnimBlueprintVariableCreationContext& InCreationContext) override;
 
 	// UEdGraphNode interface
 	virtual void AllocateDefaultPins() override;
diff --git a/Engine/Plugins/Runtime/PropertyAccess/Source/PropertyAccessEditor/Private/AnimBlueprintCompilerHandler_PropertyAccess.cpp b/Engine/Plugins/Runtime/PropertyAccess/Source/PropertyAccessEditor/Private/AnimBlueprintCompilerHandler_PropertyAccess.cpp
new file mode 100644
--- /dev/null
+++ b/Engine/Plugins/Runtime/PropertyAccess/Source/PropertyAccessEditor/Private/AnimBlueprintCompilerHandler_PropertyAccess.cpp
@@ -0,0 +1,52 @@
+// Copyright Epic Games, Inc. All Rights Reserved.
+
+#include "AnimBlueprintCompilerHandler_PropertyAccess.h"
+#include "Animation/AnimBlueprintGeneratedClass.h"
+#include "IAnimBlueprintGeneratedClassCompiledData.h"
+#include "IAnimBlueprintCompilerCreationContext.h"
+#include "IAnimBlueprintCompilationContext.h"
+#include "IAnimBlueprintCompilationBracketContext.h"
+
+FAnimBlueprintCompilerHandler_PropertyAccess::FAnimBlueprintCompilerHandler_PropertyAccess(IAnimBlueprintCompilerCreationContext& InCreationContext)
+{
+	InCreationContext.OnStartCompilingClass().AddRaw(this, &FAnimBlueprintCompilerHandler_PropertyAccess::StartCompilingClass);
+	InCreationContext.OnFinishCompilingClass().AddRaw(this, &FAnimBlueprintCompilerHandler_PropertyAccess::FinishCompilingClass);
+}
+
+int32 FAnimBlueprintCompilerHandler_PropertyAccess::AddCopy(TArrayView<FString> InSourcePath, TArrayView<FString> InDestPath, EPropertyAccessBatchType InBatchType, UObject* InObject)
+{
+	return PropertyAccessLibraryCompiler.AddCopy(InSourcePath, InDestPath, InBatchType, InObject);
+}
+
+void FAnimBlueprintCompilerHandler_PropertyAccess::StartCompilingClass(const UClass* InClass, IAnimBlueprintCompilationBracketContext& InCompilationContext, IAnimBlueprintGeneratedClassCompiledData& OutCompiledData)
+{
+	PropertyAccessLibraryCompiler.Setup(InClass, &OutCompiledData.GetPropertyAccessLibrary());
+}
+
+void FAnimBlueprintCompilerHandler_PropertyAccess::FinishCompilingClass(const UClass* InClass, IAnimBlueprintCompilationBracketContext& InCompilationContext, IAnimBlueprintGeneratedClassCompiledData& OutCompiledData)
+{
+	OnPreLibraryCompiledDelegate.Broadcast();
+
+	if(!PropertyAccessLibraryCompiler.FinishCompilation())
+	{
+		PropertyAccessLibraryCompiler.IterateErrors([&InCompilationContext](const FText& InErrorText, UObject* InObject)
+		{
+			// Output any property access errors as warnings
+			if(InObject)
+			{
+				InCompilationContext.GetMessageLog().Warning(*InErrorText.ToString(), InObject);
+			}
+			else
+			{
+				InCompilationContext.GetMessageLog().Warning(*InErrorText.ToString());
+			}
+		});
+	}
+
+	OnPostLibraryCompiledDelegate.Broadcast(OutCompiledData);
+}
+
+int32 FAnimBlueprintCompilerHandler_PropertyAccess::MapCopyIndex(int32 InIndex) const
+{
+	return PropertyAccessLibraryCompiler.MapCopyIndex(InIndex);
+}
\ No newline at end of file
diff --git a/Engine/Plugins/Runtime/PropertyAccess/Source/PropertyAccessEditor/Private/AnimBlueprintCompilerHandler_PropertyAccess.h b/Engine/Plugins/Runtime/PropertyAccess/Source/PropertyAccessEditor/Private/AnimBlueprintCompilerHandler_PropertyAccess.h
new file mode 100644
--- /dev/null
+++ b/Engine/Plugins/Runtime/PropertyAccess/Source/PropertyAccessEditor/Private/AnimBlueprintCompilerHandler_PropertyAccess.h
@@ -0,0 +1,38 @@
+// Copyright Epic Games, Inc. All Rights Reserved.
+
+#pragma once
+
+#include "CoreMinimal.h"
+#include "IAnimBlueprintCompilerHandler.h"
+#include "PropertyAccessCompilerHandler.h"
+#include "PropertyAccessEditor.h"
+
+class IAnimBlueprintCompilerCreationContext;
+class IAnimBlueprintCompilationBracketContext;
+class IAnimBlueprintGeneratedClassCompiledData;
+
+class FAnimBlueprintCompilerHandler_PropertyAccess : public FPropertyAccessCompilerHandler
+{
+public:
+	FAnimBlueprintCompilerHandler_PropertyAccess(IAnimBlueprintCompilerCreationContext& InCreationContext);
+
+	// FPropertyAccessCompilerHandler interface
+	virtual int32 AddCopy(TArrayView<FString> InSourcePath, TArrayView<FString> InDestPath, EPropertyAccessBatchType InBatchType, UObject* InObject = nullptr) override;
+	virtual FSimpleMulticastDelegate& OnPreLibraryCompiled() override { return OnPreLibraryCompiledDelegate; }
+	virtual FOnPostLibraryCompiled& OnPostLibraryCompiled() override { return OnPostLibraryCompiledDelegate; }
+	virtual int32 MapCopyIndex(int32 InIndex) const override;
+
+private:
+	void StartCompilingClass(const UClass* InClass, IAnimBlueprintCompilationBracketContext& InCompilationContext, IAnimBlueprintGeneratedClassCompiledData& OutCompiledData);
+	void FinishCompilingClass(const UClass* InClass, IAnimBlueprintCompilationBracketContext& InCompilationContext, IAnimBlueprintGeneratedClassCompiledData& OutCompiledData);
+
+private:
+	// Property access library compiler
+	FPropertyAccessLibraryCompiler PropertyAccessLibraryCompiler;
+
+	// Delegate called before the library is compiled
+	FSimpleMulticastDelegate OnPreLibraryCompiledDelegate;
+
+	// Delegate called when the library is compiled (whether successfully or not)
+	FOnPostLibraryCompiled OnPostLibraryCompiledDelegate;
+};
\ No newline at end of file
diff --git a/Engine/Plugins/Runtime/PropertyAccess/Source/PropertyAccessEditor/Private/PropertyAccessEditor.cpp b/Engine/Plugins/Runtime/PropertyAccess/Source/PropertyAccessEditor/Private/PropertyAccessEditor.cpp
--- a/Engine/Plugins/Runtime/PropertyAccess/Source/PropertyAccessEditor/Private/PropertyAccessEditor.cpp
+++ b/Engine/Plugins/Runtime/PropertyAccess/Source/PropertyAccessEditor/Private/PropertyAccessEditor.cpp
@@ -14,7 +14,7 @@ struct FPropertyAccessEditorSystem
 {
 	struct FResolveSegmentsContext
 	{
-		FResolveSegmentsContext(UStruct* InStruct, TArrayView<FString> InPath, FPropertyAccessPath& InAccessPath)
+		FResolveSegmentsContext(const UStruct* InStruct, TArrayView<FString> InPath, FPropertyAccessPath& InAccessPath)
 			: Struct(InStruct)
 			, CurrentStruct(InStruct)
 			, Path(InPath)
@@ -22,10 +22,10 @@ struct FPropertyAccessEditorSystem
 		{}
 
 		// Starting struct
-		UStruct* Struct;
+		const UStruct* Struct;
 
-		// Starting struct
-		UStruct* CurrentStruct;
+		// Current struct
+		const UStruct* CurrentStruct;
 
 		// Path as FStrings with optional array markup
 		TArrayView<FString> Path;
@@ -247,7 +247,7 @@ struct FPropertyAccessEditorSystem
 		}
 	}
 
-	static EPropertyAccessResolveResult ResolveLeafProperty(UStruct* InStruct, TArrayView<FString> InPath, FProperty*& OutProperty, int32& OutArrayIndex)
+	static EPropertyAccessResolveResult ResolveLeafProperty(const UStruct* InStruct, TArrayView<FString> InPath, FProperty*& OutProperty, int32& OutArrayIndex)
 	{
 		FPropertyAccessPath AccessPath;
 		FResolveSegmentsContext Context(InStruct, InPath, AccessPath);
@@ -376,7 +376,7 @@ struct FPropertyAccessEditorSystem
 		return EPropertyAccessCopyType::None;
 	}
 
-	static bool CompileCopy(UStruct* InStruct, FPropertyAccessLibrary& InLibrary, FPropertyAccessLibraryCompiler::FQueuedCopy& OutCopy)
+	static bool CompileCopy(const UStruct* InStruct, FPropertyAccessLibrary& InLibrary, FPropertyAccessLibraryCompiler::FQueuedCopy& OutCopy)
 	{
 		FPropertyAccessPath SrcAccessPath;
 		FPropertyAccessPath DestAccessPath;
@@ -430,7 +430,7 @@ struct FPropertyAccessEditorSystem
 
 namespace PropertyAccess
 {
-	EPropertyAccessResolveResult ResolveLeafProperty(UStruct* InStruct, TArrayView<FString> InPath, FProperty*& OutProperty, int32& OutArrayIndex)
+	EPropertyAccessResolveResult ResolveLeafProperty(const UStruct* InStruct, TArrayView<FString> InPath, FProperty*& OutProperty, int32& OutArrayIndex)
 	{
 		return ::FPropertyAccessEditorSystem::ResolveLeafProperty(InStruct, InPath, OutProperty, OutArrayIndex);
 	}
@@ -567,6 +567,14 @@ FPropertyAccessLibraryCompiler::FPropertyAccessLibraryCompiler()
 {
 }
 
+void FPropertyAccessLibraryCompiler::BeginCompilation(const UClass* InClass)
+{
+	if(Class && Library)
+	{
+		*Library = FPropertyAccessLibrary();
+	}
+}
+
 int32 FPropertyAccessLibraryCompiler::AddCopy(TArrayView<FString> InSourcePath, TArrayView<FString> InDestPath, EPropertyAccessBatchType InBatchType, UObject* InAssociatedObject)
 {
 	FQueuedCopy QueuedCopy;
@@ -632,4 +640,10 @@ int32 FPropertyAccessLibraryCompiler::MapCopyIndex(int32 InIndex) const
 	return INDEX_NONE;
 }
 
+void FPropertyAccessLibraryCompiler::Setup(const UClass* InClass, FPropertyAccessLibrary* InLibrary)
+{
+	Class = InClass;
+	Library = InLibrary;
+}
+
 #undef LOCTEXT_NAMESPACE
\ No newline at end of file
diff --git a/Engine/Plugins/Runtime/PropertyAccess/Source/PropertyAccessEditor/Private/PropertyAccessEditor.h b/Engine/Plugins/Runtime/PropertyAccess/Source/PropertyAccessEditor/Private/PropertyAccessEditor.h
--- a/Engine/Plugins/Runtime/PropertyAccess/Source/PropertyAccessEditor/Private/PropertyAccessEditor.h
+++ b/Engine/Plugins/Runtime/PropertyAccess/Source/PropertyAccessEditor/Private/PropertyAccessEditor.h
@@ -13,7 +13,7 @@ struct FEdGraphPinType;
 namespace PropertyAccess
 {
 	/** Resolve a property path to a structure, returning the leaf property and array index if any. @return true if resolution succeeded */
-	extern EPropertyAccessResolveResult ResolveLeafProperty(UStruct* InStruct, TArrayView<FString> InPath, FProperty*& OutProperty, int32& OutArrayIndex);
+	extern EPropertyAccessResolveResult ResolveLeafProperty(const UStruct* InStruct, TArrayView<FString> InPath, FProperty*& OutProperty, int32& OutArrayIndex);
 
 	// Get the compatibility of the two supplied properties. Ordering matters for promotion (A->B).
 	extern EPropertyAccessCompatibility GetPropertyCompatibility(const FProperty* InPropertyA, const FProperty* InPropertyB);
@@ -32,11 +32,16 @@ public:
 	FPropertyAccessLibraryCompiler();
 
 	// IPropertyAccessLibraryCompiler interface
+	virtual void BeginCompilation(const UClass* InClass) override;
 	virtual int32 AddCopy(TArrayView<FString> InSourcePath, TArrayView<FString> InDestPath, EPropertyAccessBatchType InBatchType, UObject* InAssociatedObject = nullptr) override;
 	virtual bool FinishCompilation() override;
 	virtual void IterateErrors(TFunctionRef<void(const FText&, UObject*)> InFunction) const override;
 	virtual int32 MapCopyIndex(int32 InIndex) const override;
 
+	// Set up the compiler ready for compilation
+	void Setup(const UClass* InClass, FPropertyAccessLibrary* InLibrary);
+
+public:
 	// Stored copy info for processing in FinishCompilation()
 	struct FQueuedCopy
 	{
@@ -58,7 +63,7 @@ protected:
 	FPropertyAccessLibrary* Library;
 
 	// The class we are compiling the library for
-	UClass* Class;
+	const UClass* Class;
 
 	// All copies to process in FinishCompilation()
 	TArray<FQueuedCopy> QueuedCopies;
diff --git a/Engine/Plugins/Runtime/PropertyAccess/Source/PropertyAccessEditor/Private/PropertyAccessEditorModule.cpp b/Engine/Plugins/Runtime/PropertyAccess/Source/PropertyAccessEditor/Private/PropertyAccessEditorModule.cpp
--- a/Engine/Plugins/Runtime/PropertyAccess/Source/PropertyAccessEditor/Private/PropertyAccessEditorModule.cpp
+++ b/Engine/Plugins/Runtime/PropertyAccess/Source/PropertyAccessEditor/Private/PropertyAccessEditorModule.cpp
@@ -8,6 +8,8 @@
 #include "PropertyAccessEditor.h"
 #include "Modules/ModuleInterface.h"
 #include "Features/IModularFeatures.h"
+#include "IAnimBlueprintCompilerHandlerCollection.h"
+#include "AnimBlueprintCompilerHandler_PropertyAccess.h"
 
 class FPropertyAccessEditorModule : public IPropertyAccessEditor, public IModuleInterface
 {
@@ -16,10 +18,19 @@ public:
 	virtual void StartupModule() override
 	{
 		IModularFeatures::Get().RegisterModularFeature("PropertyAccessEditor", this);
+
+		// Register node compilation handlers
+		IAnimBlueprintCompilerHandlerCollection::RegisterHandler("PropertyAccessCompilerHandler", [](IAnimBlueprintCompilerCreationContext& InCreationContext)
+		{
+			return MakeUnique<FAnimBlueprintCompilerHandler_PropertyAccess>(InCreationContext);
+		});
 	}
 
 	virtual void ShutdownModule() override
 	{
+		// Register node compilation handlers
+		IAnimBlueprintCompilerHandlerCollection::UnregisterHandler("PropertyAccessCompilerHandler");
+
 		IModularFeatures::Get().UnregisterModularFeature("PropertyAccessEditor", this);
 	}
 
@@ -30,7 +41,7 @@ public:
 			.Args(InArgs);
 	}
 
-	virtual EPropertyAccessResolveResult ResolveLeafProperty(UStruct* InStruct, TArrayView<FString> InPath, FProperty*& OutProperty, int32& OutArrayIndex) const override
+	virtual EPropertyAccessResolveResult ResolveLeafProperty(const UStruct* InStruct, TArrayView<FString> InPath, FProperty*& OutProperty, int32& OutArrayIndex) const override
 	{
 		return PropertyAccess::ResolveLeafProperty(InStruct, InPath, OutProperty, OutArrayIndex);
 	}
diff --git a/Engine/Source/Developer/BlueprintCompilerCppBackend/Private/BlueprintCompilerCppBackendAnim.cpp b/Engine/Source/Developer/BlueprintCompilerCppBackend/Private/BlueprintCompilerCppBackendAnim.cpp
--- a/Engine/Source/Developer/BlueprintCompilerCppBackend/Private/BlueprintCompilerCppBackendAnim.cpp
+++ b/Engine/Source/Developer/BlueprintCompilerCppBackend/Private/BlueprintCompilerCppBackendAnim.cpp
@@ -7,7 +7,6 @@
 #include "BlueprintCompilerCppBackendUtils.h"
 #include "Animation/AnimClassData.h"
 #include "Animation/AnimNodeBase.h"
-#include "Animation/AnimBlueprintClassSubsystem.h"
 
 void FBackendHelperAnim::AddHeaders(FEmitterLocalContext& Context)
 {
@@ -22,22 +21,11 @@ void FBackendHelperAnim::CreateAnimClassData(FEmitterLocalContext& Context)
 {
 	if (UAnimBlueprintGeneratedClass* AnimClass = Cast<UAnimBlueprintGeneratedClass>(Context.GetCurrentlyGeneratedClass()))
 	{
-		UAnimClassData* AnimClassData = NewObject<UAnimClassData>(AnimClass, TEXT("AnimClassData"));
-		AnimClassData->CopyFrom(AnimClass);
-		FEmitDefaultValueHelper::HandleClassSubobject(Context, AnimClassData, FEmitterLocalContext::EClassSubobjectList::MiscConvertedSubobjects, true, false);
-
-		const FString LocalNativeName = Context.ClassSubobjectsMap.FindChecked(AnimClassData);
-
-		auto CreateClassSubobjects = [&Context, &AnimClassData](bool bCreate, bool bInitialize)
-		{
-			for(UAnimBlueprintClassSubsystem* Subsystem : AnimClassData->GetSubsystems())
-			{
-				FEmitDefaultValueHelper::HandleClassSubobject(Context, Subsystem, FEmitterLocalContext::EClassSubobjectList::MiscConvertedSubobjects, bCreate, bInitialize);
-			}
-		};
+		const FString LocalNativeName = Context.GenerateUniqueLocalName();
+		Context.AddLine(FString::Printf(TEXT("UAnimClassData* %s = NewObject<UAnimClassData>(InDynamicClass, TEXT(\"AnimClassData\"));"), *LocalNativeName));
 
-		CreateClassSubobjects(true, false);
-		CreateClassSubobjects(false, true);
+		UAnimClassData* AnimClassData = NewObject<UAnimClassData>(GetTransientPackage(), TEXT("AnimClassData"));
+		AnimClassData->CopyFrom(AnimClass);
 
 		UObject* ObjectArchetype = AnimClassData->GetArchetype();
 		for (const FProperty* Property : TFieldRange<const FProperty>(UAnimClassData::StaticClass()))
@@ -51,9 +39,6 @@ void FBackendHelperAnim::CreateAnimClassData(FEmitterLocalContext& Context)
 
 		Context.AddLine(FString::Printf(TEXT("InDynamicClass->%s = %s;"), GET_MEMBER_NAME_STRING_CHECKED(UDynamicClass, AnimClassImplementation), *LocalNativeName));
 		Context.AddLine(FString::Printf(TEXT("%s->DynamicClassInitialization(InDynamicClass);"), *LocalNativeName));
-
-		// rename the temp anim class data out of the way
-		AnimClassData->Rename(nullptr, GetTransientPackage());
 	}
 }
 
diff --git a/Engine/Source/Editor/AnimGraph/Classes/AnimGraphNode_AssetPlayerBase.h b/Engine/Source/Editor/AnimGraph/Classes/AnimGraphNode_AssetPlayerBase.h
--- a/Engine/Source/Editor/AnimGraph/Classes/AnimGraphNode_AssetPlayerBase.h
+++ b/Engine/Source/Editor/AnimGraph/Classes/AnimGraphNode_AssetPlayerBase.h
@@ -29,7 +29,7 @@ public:
 	virtual void PinDefaultValueChanged(UEdGraphPin* Pin) override;
 
 	/** UAnimGraphNode_Base interface */
-	virtual void OnProcessDuringCompilation(FAnimBlueprintCompilerContext& InCompilerContext) override;
+	virtual void OnProcessDuringCompilation(IAnimBlueprintCompilationContext& InCompilationContext, IAnimBlueprintGeneratedClassCompiledData& OutCompiledData) override;
 
 	virtual void SetAnimationAsset(UAnimationAsset* Asset) { check(false); /*Base function called*/ }
 };
diff --git a/Engine/Source/Editor/AnimGraph/Classes/AnimGraphNode_Base.h b/Engine/Source/Editor/AnimGraph/Classes/AnimGraphNode_Base.h
--- a/Engine/Source/Editor/AnimGraph/Classes/AnimGraphNode_Base.h
+++ b/Engine/Source/Editor/AnimGraph/Classes/AnimGraphNode_Base.h
@@ -13,7 +13,6 @@
 #include "K2Node.h"
 #include "AnimGraphNode_Base.generated.h"
 
-class FAnimBlueprintCompilerContext;
 class FAnimGraphNodeDetails;
 class FBlueprintActionDatabaseRegistrar;
 class FCanvas;
@@ -23,6 +22,9 @@ class IDetailLayoutBuilder;
 class UAnimGraphNode_Base;
 class UEdGraphSchema;
 class USkeletalMeshComponent;
+class IAnimBlueprintCompilerHandlerCollection;
+class IAnimBlueprintGeneratedClassCompiledData;
+class IAnimBlueprintCompilationContext;
 
 struct FPoseLinkMappingRecord
 {
@@ -254,15 +256,6 @@ class ANIMGRAPH_API UAnimGraphNode_Base : public UK2Node
 
 	// Replace references to animations that exist in the supplied maps 	
 	virtual void ReplaceReferredAnimations(const TMap<UAnimationAsset*, UAnimationAsset*>& AnimAssetReplacementMap) {};
-	
-	// Process this node's data during compilation (override point)
-	virtual void OnProcessDuringCompilation(FAnimBlueprintCompilerContext& InCompilerContext) {}
-
-	// Get all the subsystems that we want to add to the class to support this node
-	virtual void GetRequiredClassSubsystems(TArray<TSubclassOf<UAnimBlueprintClassSubsystem>>& OutSubsystemClasses) const {}
-
-	// Process this node's data during compilation
-	void ProcessDuringCompilation(FAnimBlueprintCompilerContext& InCompilerContext);
 
 	// Helper function for GetAllAnimationSequencesReferred
 	void HandleAnimReferenceCollection(UAnimationAsset* AnimAsset, TArray<UAnimationAsset*>& AnimationAssets) const;
@@ -360,7 +353,7 @@ class ANIMGRAPH_API UAnimGraphNode_Base : public UK2Node
 protected:
 	friend class FAnimBlueprintCompilerContext;
 	friend class FAnimGraphNodeDetails;
-	friend class UAnimBlueprintCompilerSubsystem_Base;
+	friend class FAnimBlueprintCompilerHandler_Base;
 
 	// Gets the animation FNode type represented by this ed graph node
 	UScriptStruct* GetFNodeType() const;
@@ -379,6 +372,12 @@ protected:
 	/** Get the property (and possibly array index) associated with the supplied pin */
 	virtual void GetPinAssociatedProperty(const UScriptStruct* NodeType, const UEdGraphPin* InputPin, FProperty*& OutProperty, int32& OutIndex) const;
 
+	// Process this node's data during compilation
+	void ProcessDuringCompilation(IAnimBlueprintCompilationContext& InCompilationContext, IAnimBlueprintGeneratedClassCompiledData& OutCompiledData);
+
+	// Process this node's data during compilation (override point)
+	virtual void OnProcessDuringCompilation(IAnimBlueprintCompilationContext& InCompilationContext, IAnimBlueprintGeneratedClassCompiledData& OutCompiledData) {}
+
 	// Allocates or reallocates pins
 	void InternalPinCreation(TArray<UEdGraphPin*>* OldPins);
 
diff --git a/Engine/Source/Editor/AnimGraph/Classes/AnimGraphNode_CustomProperty.h b/Engine/Source/Editor/AnimGraph/Classes/AnimGraphNode_CustomProperty.h
--- a/Engine/Source/Editor/AnimGraph/Classes/AnimGraphNode_CustomProperty.h
+++ b/Engine/Source/Editor/AnimGraph/Classes/AnimGraphNode_CustomProperty.h
@@ -14,6 +14,7 @@
 class FCompilerResultsLog;
 class IDetailLayoutBuilder;
 class IPropertyHandle;
+class FAnimBlueprintCompilerHandler;
 
 UCLASS(Abstract)
 class ANIMGRAPH_API UAnimGraphNode_CustomProperty : public UAnimGraphNode_Base, public IClassVariableCreator
@@ -23,7 +24,7 @@ class ANIMGRAPH_API UAnimGraphNode_CustomProperty : public UAnimGraphNode_Base,
 public:
 
 	// IClassVariableCreator interface
-	virtual void CreateClassVariablesFromBlueprint(FKismetCompilerContext& InCompilerContext) override;
+	virtual void CreateClassVariablesFromBlueprint(IAnimBlueprintVariableCreationContext& InCreationContext) override;
 
 	//~ Begin UEdGraphNode Interface.
 	virtual void ValidateAnimNodeDuringCompilation(USkeleton* ForSkeleton, FCompilerResultsLog& MessageLog) override;
@@ -34,10 +35,10 @@ public:
 
 	// UAnimGraphNode_Base interface
 	virtual void CustomizeDetails(IDetailLayoutBuilder& DetailBuilder) override;
-	virtual void OnProcessDuringCompilation(FAnimBlueprintCompilerContext& InCompilerContext) override;
+	virtual void OnProcessDuringCompilation(IAnimBlueprintCompilationContext& InCompilationContext, IAnimBlueprintGeneratedClassCompiledData& OutCompiledData) override;
 
 	// Gets the property on InOwnerInstanceClass that corresponds to InInputPin
-	void GetInstancePinProperty(const UClass* InOwnerInstanceClass, UEdGraphPin* InInputPin, FProperty*& OutProperty);
+	void GetInstancePinProperty(const IAnimBlueprintCompilationContext& InCompilationContext, UEdGraphPin* InInputPin, FProperty*& OutProperty);
 	// Gets the unique name for the property linked to a given pin
 	FString GetPinTargetVariableName(const UEdGraphPin* InPin) const;
 	// Gets Target Class this properties to link
diff --git a/Engine/Source/Editor/AnimGraph/Classes/AnimGraphNode_LinkedAnimGraphBase.h b/Engine/Source/Editor/AnimGraph/Classes/AnimGraphNode_LinkedAnimGraphBase.h
--- a/Engine/Source/Editor/AnimGraph/Classes/AnimGraphNode_LinkedAnimGraphBase.h
+++ b/Engine/Source/Editor/AnimGraph/Classes/AnimGraphNode_LinkedAnimGraphBase.h
@@ -7,7 +7,6 @@
 #include "Misc/Guid.h"
 #include "AnimGraphNode_CustomProperty.h"
 #include "IClassVariableCreator.h"
-#include "Animation/AnimInstanceSubsystemData.h"
 
 #include "AnimGraphNode_LinkedAnimGraphBase.generated.h"
 
@@ -44,7 +43,7 @@ public:
 	virtual const FAnimNode_LinkedAnimGraph* GetLinkedAnimGraphNode() const PURE_VIRTUAL(UAnimGraphNode_LinkedAnimGraphBase::GetLinkedAnimGraphNode, return nullptr;);
 
 protected:
-	friend class UAnimBlueprintCompilerSubsystem_LinkedAnimGraph;
+	friend class FAnimBlueprintCompilerHandler_LinkedAnimGraph;
 
 	// Called pre-compilation to allocate pose links
 	void AllocatePoseLinks();
diff --git a/Engine/Source/Editor/AnimGraph/Classes/AnimGraphNode_LinkedInputPose.h b/Engine/Source/Editor/AnimGraph/Classes/AnimGraphNode_LinkedInputPose.h
--- a/Engine/Source/Editor/AnimGraph/Classes/AnimGraphNode_LinkedInputPose.h
+++ b/Engine/Source/Editor/AnimGraph/Classes/AnimGraphNode_LinkedInputPose.h
@@ -66,7 +66,7 @@ public:
 	virtual void PostEditChangeProperty(FPropertyChangedEvent& PropertyChangedEvent) override;
 
 	/** IClassVariableCreator interface */
-	virtual void CreateClassVariablesFromBlueprint(FKismetCompilerContext& InCompilerContext) override;
+	virtual void CreateClassVariablesFromBlueprint(IAnimBlueprintVariableCreationContext& InCreationContext) override;
 
 	/** UEdGraphNode interface */
 	virtual FLinearColor GetNodeTitleColor() const override;
diff --git a/Engine/Source/Editor/AnimGraph/Classes/AnimGraphNode_Root.h b/Engine/Source/Editor/AnimGraph/Classes/AnimGraphNode_Root.h
--- a/Engine/Source/Editor/AnimGraph/Classes/AnimGraphNode_Root.h
+++ b/Engine/Source/Editor/AnimGraph/Classes/AnimGraphNode_Root.h
@@ -30,7 +30,7 @@ class UAnimGraphNode_Root : public UAnimGraphNode_Base
 
 	//~ Begin UAnimGraphNode_Base Interface
 	virtual bool IsSinkNode() const override;
-	virtual void OnProcessDuringCompilation(FAnimBlueprintCompilerContext& InCompilerContext) override;
+	virtual void OnProcessDuringCompilation(IAnimBlueprintCompilationContext& InCompilationContext, IAnimBlueprintGeneratedClassCompiledData& OutCompiledData) override;
 
 	// Get the link to the documentation
 	virtual FString GetDocumentationLink() const override;
diff --git a/Engine/Source/Editor/AnimGraph/Private/AnimBlueprintCompiler.cpp b/Engine/Source/Editor/AnimGraph/Private/AnimBlueprintCompiler.cpp
--- a/Engine/Source/Editor/AnimGraph/Private/AnimBlueprintCompiler.cpp
+++ b/Engine/Source/Editor/AnimGraph/Private/AnimBlueprintCompiler.cpp
@@ -37,7 +37,9 @@
 #include "Algo/Transform.h"
 #include "Algo/Accumulate.h"
 #include "IClassVariableCreator.h"
-#include "Animation/AnimInstanceSubsystemData.h"
+#include "AnimBlueprintGeneratedClassCompiledData.h"
+#include "AnimBlueprintCompilationContext.h"
+#include "AnimBlueprintVariableCreationContext.h"
 
 #define LOCTEXT_NAMESPACE "AnimBlueprintCompiler"
 
@@ -49,8 +51,10 @@ FAnimBlueprintCompilerContext::FAnimBlueprintCompilerContext(UAnimBlueprint* Sou
 	, AnimBlueprint(SourceSketch)
 	, bIsDerivedAnimBlueprint(false)
 {
-	AnimBlueprintCompilerSubsystemCollection.RegisterContext(this);
-	AnimBlueprintCompilerSubsystemCollection.Initialize(GetTransientPackage());
+	AnimBlueprintCompilerHandlerCollection.Initialize(this);
+
+	// Add the animation graph schema to skip default function processing on them
+	KnownGraphSchemas.AddUnique(UAnimationGraphSchema::StaticClass());
 
 	// Make sure the skeleton has finished preloading
 	if (AnimBlueprint->TargetSkeleton != nullptr)
@@ -138,8 +142,6 @@ FAnimBlueprintCompilerContext::FAnimBlueprintCompilerContext(UAnimBlueprint* Sou
 FAnimBlueprintCompilerContext::~FAnimBlueprintCompilerContext()
 {
 	DestroyAnimGraphStubFunctions();
-
-	AnimBlueprintCompilerSubsystemCollection.Deinitialize();
 }
 
 void FAnimBlueprintCompilerContext::ForAllSubGraphs(UEdGraph* InGraph, TFunctionRef<void(UEdGraph*)> InPerGraphFunction)
@@ -164,9 +166,10 @@ void FAnimBlueprintCompilerContext::CreateClassVariablesFromBlueprint()
 		{
 			TArray<IClassVariableCreator*> ClassVariableCreators;
 			InGraph->GetNodesOfClass(ClassVariableCreators);
+			FAnimBlueprintVariableCreationContext CreationContext(this);
 			for(IClassVariableCreator* ClassVariableCreator : ClassVariableCreators)
 			{
-				ClassVariableCreator->CreateClassVariablesFromBlueprint(*this);
+				ClassVariableCreator->CreateClassVariablesFromBlueprint(CreationContext);
 			}
 		};
 
@@ -252,62 +255,9 @@ void FAnimBlueprintCompilerContext::ProcessAnimationNode(UAnimGraphNode_Base* Vi
 	NewAnimBlueprintDebugData.NodePropertyIndexToNodeMap.Add(AllocatedIndex, TrueSourceObject);
 	NewAnimBlueprintClass->GetDebugData().RegisterClassPropertyAssociation(TrueSourceObject, NewProperty);
 
-	// Add to the classes subsystem map
-	TArray<TSubclassOf<UAnimBlueprintClassSubsystem>> SubsystemClasses;
-	VisualAnimNode->GetRequiredClassSubsystems(SubsystemClasses);
-	AddSubsystemClasses(SubsystemClasses);
-
-	VisualAnimNode->ProcessDuringCompilation(*this);
-}
-
-void FAnimBlueprintCompilerContext::ProcessClassSubsystems()
-{
-	// Sort subsystems by class name
-	NewAnimBlueprintClass->Subsystems.Sort([](UAnimBlueprintClassSubsystem& InSubsystemA, UAnimBlueprintClassSubsystem& InSubsystemB)
-	{
-		return InSubsystemA.GetClass()->GetName() < InSubsystemB.GetClass()->GetName();
-	});
-
-	// Rebuild subsystem maps
-	NewAnimBlueprintClass->RebuildSubsystemMaps();
-
-	// Process all gathered class subsystems
-	for(UAnimBlueprintClassSubsystem* Subsystem : NewAnimBlueprintClass->Subsystems)
-	{
-		ProcessClassSubsystem(Subsystem);
-	}
-}
-
-void FAnimBlueprintCompilerContext::ProcessClassSubsystem(UAnimBlueprintClassSubsystem* InSubsystem)
-{
-	const FString SubsystemVariableName = ClassScopeNetNameMap.MakeValidName(InSubsystem);
-
-	const UScriptStruct* InstanceDataType = InSubsystem->GetInstanceDataType();
-	check(InstanceDataType->IsChildOf(FAnimInstanceSubsystemData::StaticStruct()));
-
-	FEdGraphPinType SubsystemVariableType;
-	SubsystemVariableType.PinCategory = UAnimationGraphSchema::PC_Struct;
-	SubsystemVariableType.PinSubCategoryObject = MakeWeakObjectPtr(const_cast<UScriptStruct*>(InstanceDataType));
-
-	FStructProperty* NewProperty = CastField<FStructProperty>(CreateVariable(FName(*SubsystemVariableName), SubsystemVariableType));
-	if (NewProperty == nullptr)
-	{
-		MessageLog.Error(*FText::Format(LOCTEXT("SubsystemPropertyCreationFailed", "Failed to create subsystem property for '{0}'"), FText::FromString(InSubsystem->GetName())).ToString());
-	}
-}
-
-void FAnimBlueprintCompilerContext::AddSubsystemClasses(const TArray<TSubclassOf<UAnimBlueprintClassSubsystem>>& SubsystemClasses)
-{
-	for(const TSubclassOf<UAnimBlueprintClassSubsystem>& NodeSubsystemClass : SubsystemClasses)
-	{
-		UAnimBlueprintClassSubsystem* ExistingSubsystem = NewAnimBlueprintClass->SubsystemMap.FindRef(NodeSubsystemClass);
-		if(ExistingSubsystem == nullptr)
-		{
-			UAnimBlueprintClassSubsystem* NewSubsystem = NewObject<UAnimBlueprintClassSubsystem>(NewAnimBlueprintClass, NodeSubsystemClass.Get());
-			NewAnimBlueprintClass->Subsystems.Add(NewSubsystem);
-			NewAnimBlueprintClass->SubsystemMap.Add(NodeSubsystemClass, NewSubsystem);
-		}
-	}
+	FAnimBlueprintGeneratedClassCompiledData CompiledData(NewAnimBlueprintClass);
+	FAnimBlueprintCompilationContext CompilerContext(this);
+	VisualAnimNode->ProcessDuringCompilation(CompilerContext, CompiledData);
 }
 
 int32 FAnimBlueprintCompilerContext::GetAllocationIndexOfNode(UAnimGraphNode_Base* VisualAnimNode)
@@ -324,12 +274,11 @@ bool FAnimBlueprintCompilerContext::ShouldForceKeepNode(const UEdGraphNode* Node
 	return Node->IsA<UAnimGraphNode_Base>();
 }
 
-void FAnimBlueprintCompilerContext::PostExpansionStep(UEdGraph* Graph)
+void FAnimBlueprintCompilerContext::PostExpansionStep(const UEdGraph* Graph)
 {
-	ForEachSubsystem([Graph](UAnimBlueprintCompilerSubsystem* InSubsystem)
-	{
-		InSubsystem->PostExpansionStep(Graph);
-	});
+	FAnimBlueprintGeneratedClassCompiledData CompiledData(NewAnimBlueprintClass);
+	FAnimBlueprintPostExpansionStepContext CompilerContext(this);
+	OnPostExpansionStepDelegate.Broadcast(Graph, CompilerContext, CompiledData);
 }
 
 void FAnimBlueprintCompilerContext::PruneIsolatedAnimationNodes(const TArray<UAnimGraphNode_Base*>& RootSet, TArray<UAnimGraphNode_Base*>& GraphNodes)
@@ -401,7 +350,7 @@ void FAnimBlueprintCompilerContext::ProcessAnimationNodes(TArray<UAnimGraphNode_
 	}
 }
 
-void FAnimBlueprintCompilerContext::GetLinkedAnimNodes(UAnimGraphNode_Base* InGraphNode, TArray<UAnimGraphNode_Base*> &LinkedAnimNodes)
+void FAnimBlueprintCompilerContext::GetLinkedAnimNodes(UAnimGraphNode_Base* InGraphNode, TArray<UAnimGraphNode_Base*> &LinkedAnimNodes) const
 {
 	for(UEdGraphPin* Pin : InGraphNode->Pins)
 	{
@@ -419,7 +368,7 @@ void FAnimBlueprintCompilerContext::GetLinkedAnimNodes(UAnimGraphNode_Base* InGr
 	}
 }
 
-void FAnimBlueprintCompilerContext::GetLinkedAnimNodes_TraversePin(UEdGraphPin* InPin, TArray<UAnimGraphNode_Base*>& LinkedAnimNodes)
+void FAnimBlueprintCompilerContext::GetLinkedAnimNodes_TraversePin(UEdGraphPin* InPin, TArray<UAnimGraphNode_Base*>& LinkedAnimNodes) const
 {
 	if(!InPin)
 	{
@@ -446,13 +395,13 @@ void FAnimBlueprintCompilerContext::GetLinkedAnimNodes_TraversePin(UEdGraphPin*
 	}
 }
 
-void FAnimBlueprintCompilerContext::GetLinkedAnimNodes_ProcessAnimNode(UAnimGraphNode_Base* AnimNode, TArray<UAnimGraphNode_Base *> &LinkedAnimNodes)
+void FAnimBlueprintCompilerContext::GetLinkedAnimNodes_ProcessAnimNode(UAnimGraphNode_Base* AnimNode, TArray<UAnimGraphNode_Base *>& LinkedAnimNodes) const
 {
 	if(!AllocatedAnimNodes.Contains(AnimNode))
 	{
 		UAnimGraphNode_Base* TrueSourceNode = MessageLog.FindSourceObjectTypeChecked<UAnimGraphNode_Base>(AnimNode);
 
-		if(UAnimGraphNode_Base** AllocatedNode = SourceNodeToProcessedNodeMap.Find(TrueSourceNode))
+		if(UAnimGraphNode_Base*const* AllocatedNode = SourceNodeToProcessedNodeMap.Find(TrueSourceNode))
 		{
 			LinkedAnimNodes.Add(*AllocatedNode);
 		}
@@ -482,7 +431,7 @@ void FAnimBlueprintCompilerContext::ProcessAllAnimationNodes()
 	ConsolidatedEventGraph->GetNodesOfClass<UAnimGraphNode_Base>(RootAnimNodeList);
 
 	// We recursively build the node lists for pre- and post-processing phases to make sure
-	// we catch any subsystem-relevant nodes in sub-graphs
+	// we catch any handler-relevant nodes in sub-graphs
 	TArray<UAnimGraphNode_Base*> AllSubGraphsAnimNodeList;
 	ForAllSubGraphs(ConsolidatedEventGraph, [&AllSubGraphsAnimNodeList](UEdGraph* InGraph)
 	{
@@ -513,26 +462,14 @@ void FAnimBlueprintCompilerContext::ProcessAllAnimationNodes()
 		// Validate the graph
 		ValidateGraphIsWellFormed(ConsolidatedEventGraph);
 
-		ForEachSubsystem([this, &AllSubGraphsAnimNodeList](UAnimBlueprintCompilerSubsystem* InSubsystem)
-		{
-			// Add any subsystem classes that need to be added without paying heed to node connectivity
-			TArray<TSubclassOf<UAnimBlueprintClassSubsystem>> SubsystemClasses;
-			InSubsystem->GetRequiredClassSubsystems(SubsystemClasses);
-			AddSubsystemClasses(SubsystemClasses);
-
-			InSubsystem->PreProcessAnimationNodes(AllSubGraphsAnimNodeList);
-		});
+		FAnimBlueprintGeneratedClassCompiledData CompiledData(NewAnimBlueprintClass);
+		FAnimBlueprintCompilationContext CompilerContext(this);
+		OnPreProcessAnimationNodesDelegate.Broadcast(AllSubGraphsAnimNodeList, CompilerContext, CompiledData);
 
 		// Process the animation nodes
 		ProcessAnimationNodes(RootAnimNodeList);
 
-		ForEachSubsystem([&AllSubGraphsAnimNodeList](UAnimBlueprintCompilerSubsystem* InSubsystem)
-		{
-			InSubsystem->PostProcessAnimationNodes(AllSubGraphsAnimNodeList);
-		});
-
-		// Process all the gathered class subsystems
-		ProcessClassSubsystems();
+		OnPostProcessAnimationNodesDelegate.Broadcast(AllSubGraphsAnimNodeList, CompilerContext, CompiledData);
 	}
 	else
 	{
@@ -682,11 +619,9 @@ void FAnimBlueprintCompilerContext::CopyTermDefaultsToDefaultObject(UObject* Def
 			}
 		}   
 
-		// Let each subsystem copy any data it wants
-		ForEachSubsystem([DefaultObject](UAnimBlueprintCompilerSubsystem* InSubsystem)
-		{
-			InSubsystem->CopyTermDefaultsToDefaultObject(DefaultObject);
-		});
+		FAnimBlueprintGeneratedClassCompiledData CompiledData(NewAnimBlueprintClass);
+		FAnimBlueprintCopyTermDefaultsContext CompilerContext(this);
+		OnCopyTermDefaultsToDefaultObjectDelegate.Broadcast(DefaultObject, CompilerContext, CompiledData);
 
 		UAnimBlueprintGeneratedClass* AnimBlueprintGeneratedClass = CastChecked<UAnimBlueprintGeneratedClass>(NewClass);
 
@@ -884,26 +819,13 @@ void FAnimBlueprintCompilerContext::MergeUbergraphPagesIn(UEdGraph* Ubergraph)
 
 void FAnimBlueprintCompilerContext::ProcessOneFunctionGraph(UEdGraph* SourceGraph, bool bInternalFunction)
 {
-	if (SourceGraph->Schema->IsChildOf(UAnimationGraphSchema::StaticClass()))
+	if(!KnownGraphSchemas.FindByPredicate([SourceGraph](const TSubclassOf<UEdGraphSchema>& InSchemaClass)
 	{
-		// Animation graph
-		// Do nothing, as this graph has already been processed
-	}
-	else
+		return SourceGraph->Schema->IsChildOf(InSchemaClass.Get());
+	}))
 	{
-		bool bProcessGraph = true;
-
-		// Let each subsystem opt out
-		ForEachSubsystem([SourceGraph, &bProcessGraph](UAnimBlueprintCompilerSubsystem* InSubsystem)
-		{
-			bProcessGraph &= InSubsystem->ShouldProcessFunctionGraph(SourceGraph);
-		});
-
-		if(bProcessGraph)
-		{
-			// Let the regular K2 compiler handle this one
-			Super::ProcessOneFunctionGraph(SourceGraph, bInternalFunction);
-		}
+		// Not known as a schema that this compiler looks at, pass to the default
+		Super::ProcessOneFunctionGraph(SourceGraph, bInternalFunction);
 	}
 }
 
@@ -931,19 +853,9 @@ void FAnimBlueprintCompilerContext::SpawnNewClass(const FString& NewClassName)
 	}
 	NewClass = NewAnimBlueprintClass;
 
-	// Add any subsystem classes that need to be added without paying heed to node connectivity
-	ForEachSubsystem([this](UAnimBlueprintCompilerSubsystem* InSubsystem)
-	{
-		TArray<TSubclassOf<UAnimBlueprintClassSubsystem>> SubsystemClasses;
-		InSubsystem->GetRequiredClassSubsystems(SubsystemClasses);
-		AddSubsystemClasses(SubsystemClasses);
-	});
-
-	// Give each subsystem a crack
-	ForEachSubsystem([this](UAnimBlueprintCompilerSubsystem* InSubsystem)
-	{
-		InSubsystem->StartCompilingClass(NewAnimBlueprintClass);
-	});
+	FAnimBlueprintGeneratedClassCompiledData CompiledData(NewAnimBlueprintClass);
+	FAnimBlueprintCompilationBracketContext CompilerContext(this);
+	OnStartCompilingClassDelegate.Broadcast(NewAnimBlueprintClass, CompilerContext, CompiledData);
 }
 
 void FAnimBlueprintCompilerContext::OnPostCDOCompiled()
@@ -986,10 +898,6 @@ void FAnimBlueprintCompilerContext::CleanAndSanitizeClass(UBlueprintGeneratedCla
 	NewAnimBlueprintClass->EvaluateGraphExposedInputs.Empty();
 	NewAnimBlueprintClass->GraphAssetPlayerInformation.Empty();
 	NewAnimBlueprintClass->GraphBlendOptions.Empty();
-	NewAnimBlueprintClass->Subsystems.Empty();
-	NewAnimBlueprintClass->SubsystemMap.Empty();
-	NewAnimBlueprintClass->SubsystemInterfaceMap.Empty();
-	NewAnimBlueprintClass->SubsystemProperties.Empty();
 
 	// Copy over runtime data from the blueprint to the class
 	NewAnimBlueprintClass->TargetSkeleton = AnimBlueprint->TargetSkeleton;
@@ -997,19 +905,9 @@ void FAnimBlueprintCompilerContext::CleanAndSanitizeClass(UBlueprintGeneratedCla
 	UAnimBlueprint* RootAnimBP = UAnimBlueprint::FindRootAnimBlueprint(AnimBlueprint);
 	bIsDerivedAnimBlueprint = RootAnimBP != NULL;
 
-	// Add any class subsystems that need to be added without paying heed to node connectivity
-	ForEachSubsystem([this](UAnimBlueprintCompilerSubsystem* InSubsystem)
-	{
-		TArray<TSubclassOf<UAnimBlueprintClassSubsystem>> SubsystemClasses;
-		InSubsystem->GetRequiredClassSubsystems(SubsystemClasses);
-		AddSubsystemClasses(SubsystemClasses);
-	});
-
-	// Give each subsystem a crack
-	ForEachSubsystem([this](UAnimBlueprintCompilerSubsystem* InSubsystem)
-	{
-		InSubsystem->StartCompilingClass(NewAnimBlueprintClass);
-	});
+	FAnimBlueprintGeneratedClassCompiledData CompiledData(NewAnimBlueprintClass);
+	FAnimBlueprintCompilationBracketContext CompilerContext(this);
+	OnStartCompilingClassDelegate.Broadcast(NewAnimBlueprintClass, CompilerContext, CompiledData);
 }
 
 void FAnimBlueprintCompilerContext::FinishCompilingClass(UClass* Class)
@@ -1052,11 +950,9 @@ void FAnimBlueprintCompilerContext::FinishCompilingClass(UClass* Class)
 		}
 	}
 
-	// Give each subsystem a crack
-	ForEachSubsystem([Class](UAnimBlueprintCompilerSubsystem* InSubsystem)
-	{
-		InSubsystem->FinishCompilingClass(Class);
-	});
+	FAnimBlueprintGeneratedClassCompiledData CompiledData(NewAnimBlueprintClass);
+	FAnimBlueprintCompilationBracketContext CompilerContext(this);
+	OnFinishCompilingClassDelegate.Broadcast(AnimBlueprintGeneratedClass, CompilerContext, CompiledData);
 
 	Super::FinishCompilingClass(Class);
 }
@@ -1355,15 +1251,6 @@ void FAnimBlueprintCompilerContext::SetCalculatedMetaDataAndFlags(UFunction* Fun
 	}
 }
 
-void FAnimBlueprintCompilerContext::ForEachSubsystem(TFunctionRef<void(UAnimBlueprintCompilerSubsystem*)> InFunction)
-{
-	const TArray<UAnimBlueprintCompilerSubsystem*>& Subsystems = AnimBlueprintCompilerSubsystemCollection.GetSubsystemArray(TSubclassOf<UAnimBlueprintCompilerSubsystem>(UAnimBlueprintCompilerSubsystem::StaticClass()));
-	for(UAnimBlueprintCompilerSubsystem* Subsystem : Subsystems)
-	{
-		InFunction(Subsystem);
-	}
-}
-
 //////////////////////////////////////////////////////////////////////////
 
 #undef LOCTEXT_NAMESPACE
\ No newline at end of file
diff --git a/Engine/Source/Editor/AnimGraph/Private/AnimBlueprintCompiler.h b/Engine/Source/Editor/AnimGraph/Private/AnimBlueprintCompiler.h
--- a/Engine/Source/Editor/AnimGraph/Private/AnimBlueprintCompiler.h
+++ b/Engine/Source/Editor/AnimGraph/Private/AnimBlueprintCompiler.h
@@ -7,7 +7,9 @@
 #include "Animation/AnimNodeBase.h"
 #include "AnimGraphNode_Base.h"
 #include "KismetCompilerModule.h"
-#include "AnimBlueprintCompilerSubsystemCollection.h"
+#include "AnimBlueprintCompilerHandlerCollection.h"
+#include "IAnimBlueprintCompilerHandlerCollection.h"
+#include "IAnimBlueprintCompilerCreationContext.h"
 
 class UAnimationGraphSchema;
 class UAnimGraphNode_SaveCachedPose;
@@ -35,13 +37,16 @@ struct FPoseLinkMappingRecord;
 struct FAnimGraphNodePropertyBinding;
 class FAnimBlueprintCompilerContext;
 
-#define ANIM_FUNC_DECORATOR	TEXT("__AnimFunc")
-
 //////////////////////////////////////////////////////////////////////////
 // FAnimBlueprintCompilerContext
 class FAnimBlueprintCompilerContext : public FKismetCompilerContext
 {
-	friend class UAnimBlueprintCompilerSubsystem;
+	friend class FAnimBlueprintCompilerCreationContext;
+	friend class FAnimBlueprintCompilationContext;
+	friend class FAnimBlueprintVariableCreationContext;
+	friend class FAnimBlueprintCompilationBracketContext;
+	friend class FAnimBlueprintPostExpansionStepContext;
+	friend class FAnimBlueprintCopyTermDefaultsContext;
 
 protected:
 	typedef FKismetCompilerContext Super;
@@ -49,30 +54,6 @@ public:
 	FAnimBlueprintCompilerContext(UAnimBlueprint* SourceSketch, FCompilerResultsLog& InMessageLog, const FKismetCompilerOptions& InCompileOptions);
 	virtual ~FAnimBlueprintCompilerContext();
 
-	virtual void PostCompile() override;
-
-	// Get a Subsystem of specified type
-	template <typename TSubsystemClass>
-	TSubsystemClass* GetSubsystem() const
-	{
-		checkSlow(this != nullptr);
-		return AnimBlueprintCompilerSubsystemCollection.GetSubsystem<TSubsystemClass>(TSubsystemClass::StaticClass());
-	}
-
-	// Find the first subsystem implementing the specified interface
-	template<typename InterfaceClass>
-	InterfaceClass* FindSubsystemWithInterface() const
-	{
-		checkSlow(this != nullptr);
-		return AnimBlueprintCompilerSubsystemCollection.FindSubsystemWithInterface<InterfaceClass>(InterfaceClass::UClassType::StaticClass());
-	}
-
-	// Allow nodes to create variables by exposing this
-	using FKismetCompilerContext::CreateVariable;
-
-	// Expose expansion functionality
-	using FKismetCompilerContext::ExpansionStep;
-
 protected:
 	// Implementation of FKismetCompilerContext interface
 	virtual void CreateClassVariablesFromBlueprint() override;
@@ -83,6 +64,7 @@ protected:
 	virtual void OnNewClassSet(UBlueprintGeneratedClass* ClassToUse) override;
 	virtual void OnPostCDOCompiled() override;
 	virtual void CopyTermDefaultsToDefaultObject(UObject* DefaultObject) override;
+	virtual void PostCompile() override;
 	virtual void PostCompileDiagnostics() override;
 	virtual void EnsureProperGeneratedClass(UClass*& TargetClass) override;
 	virtual void CleanAndSanitizeClass(UBlueprintGeneratedClass* ClassToClean, UObject*& InOldCDO) override;
@@ -90,7 +72,7 @@ protected:
 	virtual void PrecompileFunction(FKismetFunctionContext& Context, EInternalCompilerFlags InternalFlags) override;
 	virtual void SetCalculatedMetaDataAndFlags(UFunction* Function, UK2Node_FunctionEntry* EntryNode, const UEdGraphSchema_K2* Schema ) override;
 	virtual bool ShouldForceKeepNode(const UEdGraphNode* Node) const override;
-	virtual void PostExpansionStep(UEdGraph* Graph) override;
+	virtual void PostExpansionStep(const UEdGraph* Graph) override;
 	// End of FKismetCompilerContext interface
 
 protected:
@@ -124,16 +106,34 @@ protected:
 	// True if any parent class is also generated from an animation blueprint
 	bool bIsDerivedAnimBlueprint;
 
-	// Subsystems that this context is hosting
-	FAnimBlueprintCompilerSubsystemCollection AnimBlueprintCompilerSubsystemCollection;
+	// Handlers that this context is hosting
+	FAnimBlueprintCompilerHandlerCollection AnimBlueprintCompilerHandlerCollection;
+
+	// Graph schema classes that this compiler is aware of - they will skip default function processing
+	TArray<TSubclassOf<UEdGraphSchema>> KnownGraphSchemas;
+
+	/** Delegate fired when the class starts compiling. The class may be new or recycled. */
+	FOnStartCompilingClass OnStartCompilingClassDelegate;
 
-	// Expose compile options to subsystems
+	/** Delegate fired before all animation nodes are processed */
+	FOnPreProcessAnimationNodes OnPreProcessAnimationNodesDelegate;
+
+	/** Delegate fired after all animation nodes are processed */
+	FOnPostProcessAnimationNodes OnPostProcessAnimationNodesDelegate;
+
+	/** Delegate fired post- graph expansion */
+	FOnPostExpansionStep OnPostExpansionStepDelegate;
+
+	/** Delegate fired when the class has finished compiling */
+	FOnFinishCompilingClass OnFinishCompilingClassDelegate;
+
+	/** Delegate fired when data is being copied to the CDO */
+	FOnCopyTermDefaultsToDefaultObject OnCopyTermDefaultsToDefaultObjectDelegate;
+
+	// Expose compile options to handlers
 	using FKismetCompilerContext::CompileOptions;
 
 private:
-	// Run a function for each subsystem we have registered
-	void ForEachSubsystem(TFunctionRef<void(UAnimBlueprintCompilerSubsystem*)> InFunction);
-
 	// Run a function on the passed-in graph and each subgraph of it
 	void ForAllSubGraphs(UEdGraph* InGraph, TFunctionRef<void(UEdGraph*)> InPerGraphFunction);
 
@@ -153,9 +153,9 @@ private:
 	void ProcessAnimationNodes(TArray<UAnimGraphNode_Base*>& AnimNodeList);
 
 	// Gets all anim graph nodes that are piped into the provided node (traverses input pins)
-	void GetLinkedAnimNodes(UAnimGraphNode_Base* InGraphNode, TArray<UAnimGraphNode_Base*>& LinkedAnimNodes);
-	void GetLinkedAnimNodes_TraversePin(UEdGraphPin* InPin, TArray<UAnimGraphNode_Base*>& LinkedAnimNodes);
-	void GetLinkedAnimNodes_ProcessAnimNode(UAnimGraphNode_Base* AnimNode, TArray<UAnimGraphNode_Base*>& LinkedAnimNodes);
+	void GetLinkedAnimNodes(UAnimGraphNode_Base* InGraphNode, TArray<UAnimGraphNode_Base*>& LinkedAnimNodes) const;
+	void GetLinkedAnimNodes_TraversePin(UEdGraphPin* InPin, TArray<UAnimGraphNode_Base*>& LinkedAnimNodes) const;
+	void GetLinkedAnimNodes_ProcessAnimNode(UAnimGraphNode_Base* AnimNode, TArray<UAnimGraphNode_Base*>& LinkedAnimNodes) const;
 
 	// Returns the allocation index of the specified node, processing it if it was pending
 	int32 GetAllocationIndexOfNode(UAnimGraphNode_Base* VisualAnimNode);
@@ -168,14 +168,5 @@ private:
 
 	// Expands split pins for a graph
 	void ExpandSplitPins(UEdGraph* InGraph);
-
-	// Add the subsystem classes to the currently compiling class
-	void AddSubsystemClasses(const TArray<TSubclassOf<UAnimBlueprintClassSubsystem>>& SubsystemClasses);
-
-	// Process all the gathered class subsystems
-	void ProcessClassSubsystems();
-
-	// Process a single class subsystem
-	void ProcessClassSubsystem(UAnimBlueprintClassSubsystem* InSubsystem);
 };
 
diff --git a/Engine/Source/Editor/AnimGraph/Private/AnimBlueprintCompilerCreationContext.cpp b/Engine/Source/Editor/AnimGraph/Private/AnimBlueprintCompilerCreationContext.cpp
new file mode 100644
--- /dev/null
+++ b/Engine/Source/Editor/AnimGraph/Private/AnimBlueprintCompilerCreationContext.cpp
@@ -0,0 +1,40 @@
+// Copyright Epic Games, Inc. All Rights Reserved.
+
+#include "AnimBlueprintCompilerCreationContext.h"
+#include "AnimBlueprintCompiler.h"
+
+FOnStartCompilingClass& FAnimBlueprintCompilerCreationContext::OnStartCompilingClass()
+{
+	return CompilerContext->OnStartCompilingClassDelegate;
+}
+
+FOnPreProcessAnimationNodes& FAnimBlueprintCompilerCreationContext::OnPreProcessAnimationNodes()
+{
+	return CompilerContext->OnPreProcessAnimationNodesDelegate;
+}
+
+FOnPostProcessAnimationNodes& FAnimBlueprintCompilerCreationContext::OnPostProcessAnimationNodes()
+{
+	return CompilerContext->OnPostProcessAnimationNodesDelegate;
+}
+
+FOnPostExpansionStep& FAnimBlueprintCompilerCreationContext::OnPostExpansionStep()
+{
+	return CompilerContext->OnPostExpansionStepDelegate;
+}
+
+FOnFinishCompilingClass& FAnimBlueprintCompilerCreationContext::OnFinishCompilingClass()
+{
+	return CompilerContext->OnFinishCompilingClassDelegate;
+}
+
+FOnCopyTermDefaultsToDefaultObject& FAnimBlueprintCompilerCreationContext::OnCopyTermDefaultsToDefaultObject()
+{
+	return CompilerContext->OnCopyTermDefaultsToDefaultObjectDelegate;
+}
+
+void FAnimBlueprintCompilerCreationContext::RegisterKnownGraphSchema(TSubclassOf<UEdGraphSchema> InGraphSchemaClass)
+{
+	CompilerContext->KnownGraphSchemas.AddUnique(InGraphSchemaClass);
+}
+
diff --git a/Engine/Source/Editor/AnimGraph/Private/AnimBlueprintCompilerCreationContext.h b/Engine/Source/Editor/AnimGraph/Private/AnimBlueprintCompilerCreationContext.h
new file mode 100644
--- /dev/null
+++ b/Engine/Source/Editor/AnimGraph/Private/AnimBlueprintCompilerCreationContext.h
@@ -0,0 +1,30 @@
+// Copyright Epic Games, Inc. All Rights Reserved.
+
+#pragma once
+
+#include "CoreMinimal.h"
+#include "IAnimBlueprintCompilerCreationContext.h"
+
+class FAnimBlueprintCompilerContext;
+
+class FAnimBlueprintCompilerCreationContext : public IAnimBlueprintCompilerCreationContext
+{
+private:
+	friend class FAnimBlueprintCompilerHandlerCollection;
+
+	FAnimBlueprintCompilerCreationContext(FAnimBlueprintCompilerContext* InCompilerContext)
+		: CompilerContext(InCompilerContext)
+	{}
+
+	// IAnimBlueprintCompilerCreationContext interface
+	virtual FOnStartCompilingClass& OnStartCompilingClass() override;
+	virtual FOnPreProcessAnimationNodes& OnPreProcessAnimationNodes() override;
+	virtual FOnPostProcessAnimationNodes& OnPostProcessAnimationNodes() override;
+	virtual FOnPostExpansionStep& OnPostExpansionStep() override;
+	virtual FOnFinishCompilingClass& OnFinishCompilingClass() override;
+	virtual FOnCopyTermDefaultsToDefaultObject& OnCopyTermDefaultsToDefaultObject() override;
+	virtual void RegisterKnownGraphSchema(TSubclassOf<UEdGraphSchema> InGraphSchemaClass) override;
+
+private:
+	FAnimBlueprintCompilerContext* CompilerContext;
+};
\ No newline at end of file
diff --git a/Engine/Source/Editor/AnimGraph/Private/AnimBlueprintCompilerHandlerCollection.cpp b/Engine/Source/Editor/AnimGraph/Private/AnimBlueprintCompilerHandlerCollection.cpp
new file mode 100644
--- /dev/null
+++ b/Engine/Source/Editor/AnimGraph/Private/AnimBlueprintCompilerHandlerCollection.cpp
@@ -0,0 +1,38 @@
+// Copyright Epic Games, Inc. All Rights Reserved.
+
+#include "AnimBlueprintCompilerHandlerCollection.h"
+#include "AnimBlueprintCompilerCreationContext.h"
+
+/** All of our registered handler factories */
+static TMap<FName, TFunction<TUniquePtr<IAnimBlueprintCompilerHandler>(IAnimBlueprintCompilerCreationContext&)>> HandlerFactories;
+
+void IAnimBlueprintCompilerHandlerCollection::RegisterHandler(FName InName, TFunction<TUniquePtr<IAnimBlueprintCompilerHandler>(IAnimBlueprintCompilerCreationContext&)> InFunction)
+{
+	HandlerFactories.Add(InName, InFunction);
+}
+
+void IAnimBlueprintCompilerHandlerCollection::UnregisterHandler(FName InName)
+{
+	HandlerFactories.Remove(InName);
+}
+
+void FAnimBlueprintCompilerHandlerCollection::Initialize(FAnimBlueprintCompilerContext* InCompilerContext)
+{
+	FAnimBlueprintCompilerCreationContext CreationContext(InCompilerContext);
+
+	// Create all of the registered handlers
+	for(const TPair<FName, TFunction<TUniquePtr<IAnimBlueprintCompilerHandler>(IAnimBlueprintCompilerCreationContext&)>>& HandlerFactoryPair : HandlerFactories)
+	{
+		Handlers.Add(HandlerFactoryPair.Key, HandlerFactoryPair.Value(CreationContext));
+	}
+}
+
+IAnimBlueprintCompilerHandler* FAnimBlueprintCompilerHandlerCollection::GetHandlerByName(FName InName) const
+{
+	if(const TUniquePtr<IAnimBlueprintCompilerHandler>* FoundHandler = Handlers.Find(InName))
+	{
+		return (*FoundHandler).Get();
+	}
+
+	return nullptr;
+}
\ No newline at end of file
diff --git a/Engine/Source/Editor/AnimGraph/Private/AnimBlueprintCompilerHandlerCollection.h b/Engine/Source/Editor/AnimGraph/Private/AnimBlueprintCompilerHandlerCollection.h
new file mode 100644
--- /dev/null
+++ b/Engine/Source/Editor/AnimGraph/Private/AnimBlueprintCompilerHandlerCollection.h
@@ -0,0 +1,24 @@
+// Copyright Epic Games, Inc. All Rights Reserved.
+
+#pragma once
+
+#include "CoreMinimal.h"
+#include "IAnimBlueprintCompilerHandler.h"
+#include "IAnimBlueprintCompilerHandlerCollection.h"
+
+class FAnimBlueprintCompilerContext;
+
+/** Handler collection for the anim blueprint compiler */
+class FAnimBlueprintCompilerHandlerCollection : public IAnimBlueprintCompilerHandlerCollection
+{
+private:
+	friend class FAnimBlueprintCompilerContext;
+
+	void Initialize(FAnimBlueprintCompilerContext* InCompilerContext);
+
+	/** Get a named handler */
+	virtual IAnimBlueprintCompilerHandler* GetHandlerByName(FName InName) const override;
+
+	/** All the currently constructed handlers */
+	TMap<FName, TUniquePtr<IAnimBlueprintCompilerHandler>> Handlers;
+};
\ No newline at end of file
