Restricted anim BP compiler subsystem APIs

diff --git a/Engine/Source/Editor/AnimGraph/Private/AnimBlueprintCompilerHandler_Base.cpp b/Engine/Source/Editor/AnimGraph/Private/AnimBlueprintCompilerHandler_Base.cpp
new file mode 100644
--- /dev/null
+++ b/Engine/Source/Editor/AnimGraph/Private/AnimBlueprintCompilerHandler_Base.cpp
@@ -0,0 +1,1089 @@
+// Copyright Epic Games, Inc. All Rights Reserved.
+
+#include "AnimBlueprintCompilerHandler_Base.h"
+#include "AnimGraphNode_Base.h"
+#include "AnimationGraphSchema.h"
+#include "AnimGraphNode_CustomProperty.h"
+#include "K2Node_CustomEvent.h"
+#include "K2Node_VariableSet.h"
+#include "K2Node_StructMemberSet.h"
+#include "K2Node_StructMemberGet.h"
+#include "K2Node_CallArrayFunction.h"
+#include "Kismet/KismetArrayLibrary.h"
+#include "K2Node_Knot.h"
+#include "String/ParseTokens.h"
+#include "K2Node_VariableGet.h"
+#include "K2Node_BreakStruct.h"
+#include "K2Node_MakeStruct.h"
+#include "Kismet/KismetMathLibrary.h"
+#include "K2Node_TransitionRuleGetter.h"
+#include "Algo/Accumulate.h"
+#include "K2Node_GetArrayItem.h"
+#include "Animation/AnimNode_LinkedAnimGraph.h"
+#include "Kismet2/BlueprintEditorUtils.h"
+#include "PropertyAccessCompilerHandler.h"
+#include "IPropertyAccessEditor.h"
+#include "IPropertyAccessCompiler.h"
+#include "IAnimBlueprintGeneratedClassCompiledData.h"
+#include "IAnimBlueprintCompilerCreationContext.h"
+#include "IAnimBlueprintCompilationContext.h"
+#include "IAnimBlueprintCopyTermDefaultsContext.h"
+#include "IAnimBlueprintPostExpansionStepContext.h"
+#include "IAnimBlueprintCompilationBracketContext.h"
+
+#define LOCTEXT_NAMESPACE "AnimBlueprintCompilerHandler_Base"
+
+FAnimBlueprintCompilerHandler_Base::FAnimBlueprintCompilerHandler_Base(IAnimBlueprintCompilerCreationContext& InCreationContext)
+{
+	InCreationContext.OnStartCompilingClass().AddRaw(this, &FAnimBlueprintCompilerHandler_Base::StartCompilingClass);
+	InCreationContext.OnFinishCompilingClass().AddRaw(this, &FAnimBlueprintCompilerHandler_Base::FinishCompilingClass);
+	InCreationContext.OnPostExpansionStep().AddRaw(this, &FAnimBlueprintCompilerHandler_Base::PostExpansionStep);
+	InCreationContext.OnCopyTermDefaultsToDefaultObject().AddRaw(this, &FAnimBlueprintCompilerHandler_Base::CopyTermDefaultsToDefaultObject);
+}
+
+void FAnimBlueprintCompilerHandler_Base::CopyTermDefaultsToDefaultObject(UObject* InDefaultObject, IAnimBlueprintCopyTermDefaultsContext& InCompilationContext, IAnimBlueprintGeneratedClassCompiledData& OutCompiledData)
+{
+	UAnimInstance* DefaultAnimInstance = Cast<UAnimInstance>(InDefaultObject);
+
+	if(DefaultAnimInstance)
+	{
+		// And patch in constant values that don't need to be re-evaluated every frame
+		for (auto LiteralLinkIt = ValidAnimNodePinConstants.CreateIterator(); LiteralLinkIt; ++LiteralLinkIt)
+		{
+			FEffectiveConstantRecord& ConstantRecord = *LiteralLinkIt;
+
+			//const FString ArrayClause = (ConstantRecord.ArrayIndex != INDEX_NONE) ? FString::Printf(TEXT("[%d]"), ConstantRecord.ArrayIndex) : FString();
+			//const FString ValueClause = ConstantRecord.LiteralSourcePin->GetDefaultAsString();
+			//GetMessageLog().Note(*FString::Printf(TEXT("Want to set %s.%s%s to %s"), *ConstantRecord.NodeVariableProperty->GetName(), *ConstantRecord.ConstantProperty->GetName(), *ArrayClause, *ValueClause));
+
+			if (!ConstantRecord.Apply(DefaultAnimInstance))
+			{
+				InCompilationContext.GetMessageLog().Error(TEXT("ICE: Failed to push literal value from @@ into CDO"), ConstantRecord.LiteralSourcePin);
+			}
+		}
+
+		for (const FEffectiveConstantRecord& ConstantRecord : ValidAnimNodePinConstants)
+		{
+			UAnimGraphNode_Base* Node = CastChecked<UAnimGraphNode_Base>(ConstantRecord.LiteralSourcePin->GetOwningNode());
+			UAnimGraphNode_Base* TrueNode = InCompilationContext.GetMessageLog().FindSourceObjectTypeChecked<UAnimGraphNode_Base>(Node);
+			TrueNode->BlueprintUsage = EBlueprintUsage::DoesNotUseBlueprint;
+		}
+
+		for(const FEvaluationHandlerRecord& EvaluationHandler : ValidEvaluationHandlerList)
+		{
+			if(EvaluationHandler.EvaluationHandlerIdx != INDEX_NONE && EvaluationHandler.ServicedProperties.Num() > 0)
+			{
+				const FAnimNodeSinglePropertyHandler& Handler = EvaluationHandler.ServicedProperties.CreateConstIterator()->Value;
+				check(Handler.CopyRecords.Num() > 0);
+				if(Handler.CopyRecords[0].DestPin != nullptr)
+				{
+					UAnimGraphNode_Base* Node = CastChecked<UAnimGraphNode_Base>(Handler.CopyRecords[0].DestPin->GetOwningNode());
+					UAnimGraphNode_Base* TrueNode = InCompilationContext.GetMessageLog().FindSourceObjectTypeChecked<UAnimGraphNode_Base>(Node);	
+
+					const FExposedValueHandler& ValueHandler = OutCompiledData.GetExposedValueHandlers()[ EvaluationHandler.EvaluationHandlerIdx ];
+					TrueNode->BlueprintUsage = ValueHandler.BoundFunction != NAME_None ? EBlueprintUsage::UsesBlueprint : EBlueprintUsage::DoesNotUseBlueprint; 
+
+#if WITH_EDITORONLY_DATA // ANIMINST_PostCompileValidation
+					const bool bWarnAboutBlueprintUsage = InCompilationContext.GetAnimBlueprint()->bWarnAboutBlueprintUsage || DefaultAnimInstance->PCV_ShouldWarnAboutNodesNotUsingFastPath();
+					const bool bNotifyAboutBlueprintUsage = DefaultAnimInstance->PCV_ShouldNotifyAboutNodesNotUsingFastPath();
+#else
+					const bool bWarnAboutBlueprintUsage = InCompilationContext.GetAnimBlueprint()->bWarnAboutBlueprintUsage;
+					const bool bNotifyAboutBlueprintUsage = false;
+#endif
+					if ((TrueNode->BlueprintUsage == EBlueprintUsage::UsesBlueprint) && (bWarnAboutBlueprintUsage || bNotifyAboutBlueprintUsage))
+					{
+						const FString MessageString = LOCTEXT("BlueprintUsageWarning", "Node @@ uses Blueprint to update its values, access member variables directly or use a constant value for better performance.").ToString();
+						if (bWarnAboutBlueprintUsage)
+						{
+							InCompilationContext.GetMessageLog().Warning(*MessageString, Node);
+						}
+						else
+						{
+							InCompilationContext.GetMessageLog().Note(*MessageString, Node);
+						}
+					}
+				}
+			}
+		}
+	}
+}
+
+void FAnimBlueprintCompilerHandler_Base::PostExpansionStep(const UEdGraph* InGraph, IAnimBlueprintPostExpansionStepContext& InCompilationContext, IAnimBlueprintGeneratedClassCompiledData& OutCompiledData)
+{
+	UEdGraph* ConsolidatedEventGraph = InCompilationContext.GetConsolidatedEventGraph();
+	if(InGraph == ConsolidatedEventGraph)
+	{
+		FPropertyAccessCompilerHandler* PropertyAccessHandler = InCompilationContext.GetHandler<FPropertyAccessCompilerHandler>("PropertyAccessCompilerHandler");
+
+		// Skip fast-path generation if the property access system is unavailable.
+		// Disable fast-path generation for nativized anim BPs, we dont run the VM anyways and 
+		// the property names are 'decorated' by the backend, so records dont match.
+		// Note that this wont prevent property access 'binding' copy records from running, only
+		// old-style 'fast-path' records that are derived from BP pure chains
+		if(PropertyAccessHandler != nullptr && !InCompilationContext.GetCompileOptions().DoesRequireCppCodeGeneration())
+		{
+			for(FEvaluationHandlerRecord& HandlerRecord : ValidEvaluationHandlerList)
+			{
+				HandlerRecord.BuildFastPathCopyRecords(*this);
+
+				if(HandlerRecord.IsFastPath())
+				{
+					for(UEdGraphNode* CustomEventNode : HandlerRecord.CustomEventNodes)
+					{
+						// Remove custom event nodes as we dont need it any more
+						ConsolidatedEventGraph->RemoveNode(CustomEventNode);
+					}
+				}
+			}
+		}
+
+		// Cull out all anim nodes as they dont contribute to execution at all
+		for (int32 NodeIndex = 0; NodeIndex < ConsolidatedEventGraph->Nodes.Num(); ++NodeIndex)
+		{
+			if(UAnimGraphNode_Base* Node = Cast<UAnimGraphNode_Base>(ConsolidatedEventGraph->Nodes[NodeIndex]))
+			{
+				Node->BreakAllNodeLinks();
+				ConsolidatedEventGraph->Nodes.RemoveAtSwap(NodeIndex);
+				--NodeIndex;
+			}
+		}
+	}
+}
+
+void FAnimBlueprintCompilerHandler_Base::StartCompilingClass(const UClass* InClass, IAnimBlueprintCompilationBracketContext& InCompilationContext, IAnimBlueprintGeneratedClassCompiledData& OutCompiledData)
+{
+	FPropertyAccessCompilerHandler* PropertyAccessHandler = InCompilationContext.GetHandler<FPropertyAccessCompilerHandler>("PropertyAccessCompilerHandler");
+	if(PropertyAccessHandler)
+	{
+		if(!PreLibraryCompiledDelegateHandle.IsValid())
+		{
+			PreLibraryCompiledDelegateHandle = PropertyAccessHandler->OnPreLibraryCompiled().AddLambda([this, PropertyAccessHandler, InClass]()
+			{
+				if(IModularFeatures::Get().IsModularFeatureAvailable("PropertyAccessEditor"))
+				{
+					IPropertyAccessEditor& PropertyAccessEditor = IModularFeatures::Get().GetModularFeature<IPropertyAccessEditor>("PropertyAccessEditor");
+
+					// Build the classes property access library before the library is compiled
+					for(FEvaluationHandlerRecord& HandlerRecord : ValidEvaluationHandlerList)
+					{
+						for(TPair<FName, FAnimNodeSinglePropertyHandler>& PropertyHandler : HandlerRecord.ServicedProperties)
+						{
+							for(FPropertyCopyRecord& Record : PropertyHandler.Value.CopyRecords)
+							{
+								if(Record.IsFastPath())
+								{
+									// Check if the resolved copy
+									FProperty* LeafProperty = nullptr;
+									int32 ArrayIndex = INDEX_NONE;
+									EPropertyAccessResolveResult Result = PropertyAccessEditor.ResolveLeafProperty(InClass, Record.SourcePropertyPath, LeafProperty, ArrayIndex);
+
+									// Batch all external accesses, we cant call them safely from a worker thread.
+									Record.LibraryBatchType = Result == EPropertyAccessResolveResult::SucceededExternal ? EPropertyAccessBatchType::Batched : EPropertyAccessBatchType::Unbatched;
+									Record.LibraryCopyIndex = PropertyAccessHandler->AddCopy(Record.SourcePropertyPath, Record.DestPropertyPath, Record.LibraryBatchType, HandlerRecord.AnimGraphNode);
+								}
+							}
+						}
+					}
+				}
+			});
+		}
+
+		if(!PostLibraryCompiledDelegateHandle.IsValid())
+		{
+			PostLibraryCompiledDelegateHandle = PropertyAccessHandler->OnPostLibraryCompiled().AddLambda([this, PropertyAccessHandler](IAnimBlueprintGeneratedClassCompiledData& OutCompiledData)
+			{
+				TArray<FExposedValueHandler>& ExposedValueHandlers = OutCompiledData.GetExposedValueHandlers();
+
+				for(FEvaluationHandlerRecord& HandlerRecord : ValidEvaluationHandlerList)
+				{
+					// Map global copy index to batched indices
+					for(TPair<FName, FAnimNodeSinglePropertyHandler>& PropertyHandler : HandlerRecord.ServicedProperties)
+					{
+						for(FPropertyCopyRecord& CopyRecord : PropertyHandler.Value.CopyRecords)
+						{
+							if(CopyRecord.IsFastPath())
+							{
+								CopyRecord.LibraryCopyIndex = PropertyAccessHandler->MapCopyIndex(CopyRecord.LibraryCopyIndex);
+							}
+						}
+					}
+
+					// Patch either fast-path copy records or generated function names into the class
+					HandlerRecord.EvaluationHandlerIdx = ExposedValueHandlers.Num();
+					FExposedValueHandler& ExposedValueHandler = ExposedValueHandlers.AddDefaulted_GetRef();
+					HandlerRecord.PatchFunctionNameAndCopyRecordsInto(ExposedValueHandler);
+				}
+			});
+		}
+	}
+}
+
+void FAnimBlueprintCompilerHandler_Base::FinishCompilingClass(const UClass* InClass, IAnimBlueprintCompilationBracketContext& InCompilationContext, IAnimBlueprintGeneratedClassCompiledData& OutCompiledData)
+{
+	FPropertyAccessCompilerHandler* PropertyAccessHandler = InCompilationContext.GetHandler<FPropertyAccessCompilerHandler>("PropertyAccessCompilerHandler");
+	if(PropertyAccessHandler == nullptr)
+	{
+		TArray<FExposedValueHandler>& ExposedValueHandlers = OutCompiledData.GetExposedValueHandlers();
+
+		// Without the property access system we need to patch generated function names here
+		for(FEvaluationHandlerRecord& HandlerRecord : ValidEvaluationHandlerList)
+		{
+			HandlerRecord.EvaluationHandlerIdx = ExposedValueHandlers.Num();
+			FExposedValueHandler& ExposedValueHandler = ExposedValueHandlers.AddDefaulted_GetRef();
+			HandlerRecord.PatchFunctionNameAndCopyRecordsInto(ExposedValueHandler);
+		}
+	}
+}
+
+void FAnimBlueprintCompilerHandler_Base::AddStructEvalHandlers(UAnimGraphNode_Base* InNode, IAnimBlueprintCompilationContext& InCompilationContext, IAnimBlueprintGeneratedClassCompiledData& OutCompiledData)
+{
+	const UAnimationGraphSchema* AnimGraphDefaultSchema = GetDefault<UAnimationGraphSchema>();
+
+	FEvaluationHandlerRecord& EvalHandler = PerNodeStructEvalHandlers.Add(InNode);
+
+	FStructProperty* NodeProperty = CastFieldChecked<FStructProperty>(InCompilationContext.GetAllocatedPropertiesByNode().FindChecked(InNode));
+
+	for (auto SourcePinIt = InNode->Pins.CreateIterator(); SourcePinIt; ++SourcePinIt)
+	{
+		UEdGraphPin* SourcePin = *SourcePinIt;
+		bool bConsumed = false;
+
+		// Register pose links for future use
+		if ((SourcePin->Direction == EGPD_Input) && (AnimGraphDefaultSchema->IsPosePin(SourcePin->PinType)))
+		{
+			// Input pose pin, going to need to be linked up
+			FPoseLinkMappingRecord LinkRecord = InNode->GetLinkIDLocation(NodeProperty->Struct, SourcePin);
+			if (LinkRecord.IsValid())
+			{
+				InCompilationContext.AddPoseLinkMappingRecord(LinkRecord);
+				bConsumed = true;
+			}
+		}
+		else
+		{
+			// The property source for our data, either the struct property for an anim node, or the
+			// owning anim instance if using a linked instance node.
+			FProperty* SourcePinProperty = nullptr;
+			int32 SourceArrayIndex = INDEX_NONE;
+			bool bInstancePropertyExists = false;
+
+			// We have special handling below if we're targeting a linked instance instead of our own instance properties
+			UAnimGraphNode_CustomProperty* CustomPropertyNode = Cast<UAnimGraphNode_CustomProperty>(InNode);
+
+			InNode->GetPinAssociatedProperty(NodeProperty->Struct, SourcePin, /*out*/ SourcePinProperty, /*out*/ SourceArrayIndex);
+
+			// Does this pin have an associated evaluation handler?
+			if(!SourcePinProperty && CustomPropertyNode)
+			{
+				// Custom property nodes use instance properties not node properties as they aren't UObjects
+				// and we can't store non-native properties there
+				CustomPropertyNode->GetInstancePinProperty(InCompilationContext, SourcePin, SourcePinProperty);
+				bInstancePropertyExists = true;
+			}
+			
+			if (SourcePinProperty != NULL)
+			{
+				if (SourcePin->LinkedTo.Num() == 0)
+				{
+					// Literal that can be pushed into the CDO instead of re-evaluated every frame
+					new (ValidAnimNodePinConstants) FEffectiveConstantRecord(NodeProperty, SourcePin, SourcePinProperty, SourceArrayIndex);
+					bConsumed = true;
+				}
+				else
+				{
+					// Dynamic value that needs to be wired up and evaluated each frame
+					const FString& EvaluationHandlerStr = SourcePinProperty->GetMetaData(AnimGraphDefaultSchema->NAME_OnEvaluate);
+					FName EvaluationHandlerName(*EvaluationHandlerStr);
+					if (EvaluationHandlerName != NAME_None)
+					{
+						// warn that NAME_OnEvaluate is deprecated:
+						InCompilationContext.GetMessageLog().Warning(*LOCTEXT("OnEvaluateDeprecated", "OnEvaluate meta data is deprecated, found on @@").ToString(), SourcePinProperty);
+					}
+					
+					ensure(EvalHandler.NodeVariableProperty == nullptr || EvalHandler.NodeVariableProperty == NodeProperty);
+					EvalHandler.AnimGraphNode = InNode;
+					EvalHandler.NodeVariableProperty = NodeProperty;
+					EvalHandler.RegisterPin(SourcePin, SourcePinProperty, SourceArrayIndex);
+					// if it's not instance property, ensure we mark it
+					EvalHandler.bServicesNodeProperties = EvalHandler.bServicesNodeProperties | !bInstancePropertyExists;
+
+					if (CustomPropertyNode)
+					{
+						EvalHandler.bServicesInstanceProperties = EvalHandler.bServicesInstanceProperties | bInstancePropertyExists;
+
+						FAnimNodeSinglePropertyHandler* SinglePropHandler = EvalHandler.ServicedProperties.Find(SourcePinProperty->GetFName());
+						check(SinglePropHandler); // Should have been added in RegisterPin
+
+						// Flag that the target property is actually on the instance class and not the node
+						SinglePropHandler->bInstanceIsTarget = bInstancePropertyExists;
+					}
+
+					bConsumed = true;
+				}
+
+				UEdGraphPin* TrueSourcePin = InCompilationContext.GetMessageLog().FindSourcePin(SourcePin);
+				if (TrueSourcePin)
+				{
+					OutCompiledData.GetBlueprintDebugData().RegisterClassPropertyAssociation(TrueSourcePin, SourcePinProperty);
+				}
+			}
+		}
+
+		if (!bConsumed && (SourcePin->Direction == EGPD_Input))
+		{
+			//@TODO: ANIMREFACTOR: It's probably OK to have certain pins ignored eventually, but this is very helpful during development
+			InCompilationContext.GetMessageLog().Note(TEXT("@@ was visible but ignored"), SourcePin);
+		}
+	}
+
+	// Add any property bindings
+	for(const TPair<FName, FAnimGraphNodePropertyBinding>& PropertyBinding : InNode->PropertyBindings)
+	{
+		if(PropertyBinding.Value.bIsBound)
+		{
+			EvalHandler.AnimGraphNode = InNode;
+			EvalHandler.NodeVariableProperty = NodeProperty;
+			EvalHandler.bServicesNodeProperties = true;
+
+			if (FProperty* Property = FindFProperty<FProperty>(NodeProperty->Struct, PropertyBinding.Key))
+			{
+				EvalHandler.RegisterPropertyBinding(Property, PropertyBinding.Value);
+			}
+			else
+			{
+				InCompilationContext.GetMessageLog().Warning(*FString::Printf(TEXT("ICE: @@ Failed to find a property '%s'"), *PropertyBinding.Key.ToString()), InNode);
+			}
+		}
+	}
+}
+
+void FAnimBlueprintCompilerHandler_Base::CreateEvaluationHandlerForNode(IAnimBlueprintCompilationContext& InCompilationContext, UAnimGraphNode_Base* InNode)
+{
+	if(FEvaluationHandlerRecord* RecordPtr = PerNodeStructEvalHandlers.Find(InNode))
+	{
+		// Generate a new event to update the value of these properties
+		FEvaluationHandlerRecord& Record = *RecordPtr;
+
+		if (Record.NodeVariableProperty)
+		{
+			CreateEvaluationHandler(InCompilationContext, InNode, Record);
+
+			int32 NewIndex = ValidEvaluationHandlerList.Add(Record);
+			ValidEvaluationHandlerMap.Add(InNode, NewIndex);
+		}
+	}
+}
+
+void FAnimBlueprintCompilerHandler_Base::CreateEvaluationHandler(IAnimBlueprintCompilationContext& InCompilationContext, UAnimGraphNode_Base* InNode, FEvaluationHandlerRecord& Record)
+{
+	// Shouldn't create a handler if there is nothing to work with
+	check(Record.ServicedProperties.Num() > 0);
+	check(Record.NodeVariableProperty != NULL);
+
+	const UEdGraphSchema_K2* K2Schema = GetDefault<UEdGraphSchema_K2>();
+	const UAnimationGraphSchema* AnimGraphDefaultSchema = GetDefault<UAnimationGraphSchema>();
+
+	// Use the node GUID for a stable name across compiles
+	FString FunctionName = FString::Printf(TEXT("%s_%s_%s_%s"), *AnimGraphDefaultSchema->DefaultEvaluationHandlerName.ToString(), *InNode->GetOuter()->GetName(), *InNode->GetClass()->GetName(), *InNode->NodeGuid.ToString());
+	Record.HandlerFunctionName = FName(*FunctionName);
+
+	// check function name isnt already used (data exists that can contain duplicate GUIDs) and apply a numeric extension until it is unique
+	int32 ExtensionIndex = 0;
+	FName* ExistingName = HandlerFunctionNames.Find(Record.HandlerFunctionName);
+	while (ExistingName != nullptr)
+	{
+		FunctionName = FString::Printf(TEXT("%s_%s_%s_%s_%d"), *AnimGraphDefaultSchema->DefaultEvaluationHandlerName.ToString(), *InNode->GetOuter()->GetName(), *InNode->GetClass()->GetName(), *InNode->NodeGuid.ToString(), ExtensionIndex);
+		Record.HandlerFunctionName = FName(*FunctionName);
+		ExistingName = HandlerFunctionNames.Find(Record.HandlerFunctionName);
+		ExtensionIndex++;
+	}
+
+	HandlerFunctionNames.Add(Record.HandlerFunctionName);
+
+	// Add a custom event in the graph
+	UK2Node_CustomEvent* CustomEventNode = InCompilationContext.SpawnIntermediateEventNode<UK2Node_CustomEvent>(InNode, nullptr, InCompilationContext.GetConsolidatedEventGraph());
+	CustomEventNode->bInternalEvent = true;
+	CustomEventNode->CustomFunctionName = Record.HandlerFunctionName;
+	CustomEventNode->AllocateDefaultPins();
+	Record.CustomEventNodes.Add(CustomEventNode);
+
+	// The ExecChain is the current exec output pin in the linear chain
+	UEdGraphPin* ExecChain = K2Schema->FindExecutionPin(*CustomEventNode, EGPD_Output);
+	if (Record.bServicesInstanceProperties)
+	{
+		// Need to create a variable set call for each serviced property in the handler
+		for (TPair<FName, FAnimNodeSinglePropertyHandler>& PropHandlerPair : Record.ServicedProperties)
+		{
+			FAnimNodeSinglePropertyHandler& PropHandler = PropHandlerPair.Value;
+			FName PropertyName = PropHandlerPair.Key;
+
+			// Should be true, we only want to deal with instance targets in here
+			if (PropHandler.bInstanceIsTarget)
+			{
+				for (FPropertyCopyRecord& CopyRecord : PropHandler.CopyRecords)
+				{
+					// New set node for the property
+					UK2Node_VariableSet* VarAssignNode = InCompilationContext.SpawnIntermediateNode<UK2Node_VariableSet>(InNode, InCompilationContext.GetConsolidatedEventGraph());
+					VarAssignNode->VariableReference.SetSelfMember(CopyRecord.DestProperty->GetFName());
+					VarAssignNode->AllocateDefaultPins();
+					Record.CustomEventNodes.Add(VarAssignNode);
+
+					// Wire up the exec line, and update the end of the chain
+					UEdGraphPin* ExecVariablesIn = K2Schema->FindExecutionPin(*VarAssignNode, EGPD_Input);
+					ExecChain->MakeLinkTo(ExecVariablesIn);
+					ExecChain = K2Schema->FindExecutionPin(*VarAssignNode, EGPD_Output);
+
+					// Find the property pin on the set node and configure
+					for (UEdGraphPin* TargetPin : VarAssignNode->Pins)
+					{
+						FName PinPropertyName(TargetPin->PinName);
+
+						if (PinPropertyName == PropertyName)
+						{
+							// This is us, wire up the variable
+							UEdGraphPin* DestPin = CopyRecord.DestPin;
+
+							// Copy the data (link up to the source nodes)
+							TargetPin->CopyPersistentDataFromOldPin(*DestPin);
+							InCompilationContext.GetMessageLog().NotifyIntermediatePinCreation(TargetPin, DestPin);
+
+							break;
+						}
+					}
+				}
+			}
+		}
+	}
+
+	if (Record.bServicesNodeProperties)
+	{
+		// Create a struct member write node to store the parameters into the animation node
+		UK2Node_StructMemberSet* AssignmentNode = InCompilationContext.SpawnIntermediateNode<UK2Node_StructMemberSet>(InNode, InCompilationContext.GetConsolidatedEventGraph());
+		AssignmentNode->VariableReference.SetSelfMember(Record.NodeVariableProperty->GetFName());
+		AssignmentNode->StructType = Record.NodeVariableProperty->Struct;
+		AssignmentNode->AllocateDefaultPins();
+		Record.CustomEventNodes.Add(AssignmentNode);
+
+		// Wire up the variable node execution wires
+		UEdGraphPin* ExecVariablesIn = K2Schema->FindExecutionPin(*AssignmentNode, EGPD_Input);
+		ExecChain->MakeLinkTo(ExecVariablesIn);
+		ExecChain = K2Schema->FindExecutionPin(*AssignmentNode, EGPD_Output);
+
+		// Run thru each property
+		TSet<FName> PropertiesBeingSet;
+
+		for (auto TargetPinIt = AssignmentNode->Pins.CreateIterator(); TargetPinIt; ++TargetPinIt)
+		{
+			UEdGraphPin* TargetPin = *TargetPinIt;
+			FName PropertyName(TargetPin->PinName);
+
+			// Does it get serviced by this handler?
+			if (FAnimNodeSinglePropertyHandler* SourceInfo = Record.ServicedProperties.Find(PropertyName))
+			{
+				if (TargetPin->PinType.IsArray())
+				{
+					// Grab the array that we need to set members for
+					UK2Node_StructMemberGet* FetchArrayNode = InCompilationContext.SpawnIntermediateNode<UK2Node_StructMemberGet>(InNode, InCompilationContext.GetConsolidatedEventGraph());
+					FetchArrayNode->VariableReference.SetSelfMember(Record.NodeVariableProperty->GetFName());
+					FetchArrayNode->StructType = Record.NodeVariableProperty->Struct;
+					FetchArrayNode->AllocatePinsForSingleMemberGet(PropertyName);
+					Record.CustomEventNodes.Add(FetchArrayNode);
+
+					UEdGraphPin* ArrayVariableNode = FetchArrayNode->FindPin(PropertyName);
+
+					if (SourceInfo->CopyRecords.Num() > 0)
+					{
+						// Set each element in the array
+						for (FPropertyCopyRecord& CopyRecord : SourceInfo->CopyRecords)
+						{
+							int32 ArrayIndex = CopyRecord.DestArrayIndex;
+							if(UEdGraphPin* DestPin = CopyRecord.DestPin)
+							{
+								// Create an array element set node
+								UK2Node_CallArrayFunction* ArrayNode = InCompilationContext.SpawnIntermediateNode<UK2Node_CallArrayFunction>(InNode, InCompilationContext.GetConsolidatedEventGraph());
+								ArrayNode->FunctionReference.SetExternalMember(GET_FUNCTION_NAME_CHECKED(UKismetArrayLibrary, Array_Set), UKismetArrayLibrary::StaticClass());
+								ArrayNode->AllocateDefaultPins();
+								Record.CustomEventNodes.Add(ArrayNode);
+
+								// Connect the execution chain
+								ExecChain->MakeLinkTo(ArrayNode->GetExecPin());
+								ExecChain = ArrayNode->GetThenPin();
+
+								// Connect the input array
+								UEdGraphPin* TargetArrayPin = ArrayNode->FindPinChecked(TEXT("TargetArray"));
+								TargetArrayPin->MakeLinkTo(ArrayVariableNode);
+								ArrayNode->PinConnectionListChanged(TargetArrayPin);
+
+								// Set the array index
+								UEdGraphPin* TargetIndexPin = ArrayNode->FindPinChecked(TEXT("Index"));
+								TargetIndexPin->DefaultValue = FString::FromInt(ArrayIndex);
+
+								// Wire up the data input
+								UEdGraphPin* TargetItemPin = ArrayNode->FindPinChecked(TEXT("Item"));
+								TargetItemPin->CopyPersistentDataFromOldPin(*DestPin);
+								InCompilationContext.GetMessageLog().NotifyIntermediatePinCreation(TargetItemPin, DestPin);
+							}
+						}
+					}
+				}
+				else
+				{
+					check(!TargetPin->PinType.IsContainer());
+					// Single property
+					if (SourceInfo->CopyRecords.Num() > 0 && SourceInfo->CopyRecords[0].DestPin != nullptr)
+					{
+						UEdGraphPin* DestPin = SourceInfo->CopyRecords[0].DestPin;
+
+						PropertiesBeingSet.Add(DestPin->PinName);
+						TargetPin->CopyPersistentDataFromOldPin(*DestPin);
+						InCompilationContext.GetMessageLog().NotifyIntermediatePinCreation(TargetPin, DestPin);
+					}
+				}
+			}
+		}
+
+		// Remove any unused pins from the assignment node to avoid smashing constant values
+		for (int32 PinIndex = 0; PinIndex < AssignmentNode->ShowPinForProperties.Num(); ++PinIndex)
+		{
+			FOptionalPinFromProperty& TestProperty = AssignmentNode->ShowPinForProperties[PinIndex];
+			TestProperty.bShowPin = PropertiesBeingSet.Contains(TestProperty.PropertyName);
+		}
+
+		AssignmentNode->ReconstructNode();
+	}
+}
+
+void FAnimBlueprintCompilerHandler_Base::FEvaluationHandlerRecord::PatchFunctionNameAndCopyRecordsInto(FExposedValueHandler& Handler) const
+{
+	Handler.CopyRecords.Empty();
+	Handler.ValueHandlerNodeProperty = NodeVariableProperty;
+
+	if (IsFastPath())
+	{
+		for (const TPair<FName, FAnimNodeSinglePropertyHandler>& ServicedPropPair : ServicedProperties)
+		{
+			const FName& PropertyName = ServicedPropPair.Key;
+			const FAnimNodeSinglePropertyHandler& PropertyHandler = ServicedPropPair.Value;
+
+			for (const FPropertyCopyRecord& PropertyCopyRecord : PropertyHandler.CopyRecords)
+			{
+				// Only unbatched copies can be processed on a per-node basis
+				// Skip invalid copy indices as these are usually the result of BP errors/warnings
+				if(PropertyCopyRecord.LibraryCopyIndex != INDEX_NONE && PropertyCopyRecord.LibraryBatchType == EPropertyAccessBatchType::Unbatched)
+				{
+					Handler.CopyRecords.Emplace(PropertyCopyRecord.LibraryCopyIndex, PropertyCopyRecord.Operation);
+				}
+			}
+		}
+	}
+	else
+	{
+		// not all of our pins use copy records so we will need to call our exposed value handler
+		Handler.BoundFunction = HandlerFunctionName;
+	}
+}
+
+static UEdGraphPin* FindFirstInputPin(UEdGraphNode* InNode)
+{
+	const UAnimationGraphSchema* Schema = GetDefault<UAnimationGraphSchema>();
+
+	for(UEdGraphPin* Pin : InNode->Pins)
+	{
+		if(Pin && Pin->Direction == EGPD_Input && !Schema->IsExecPin(*Pin) && !Schema->IsSelfPin(*Pin))
+		{
+			return Pin;
+		}
+	}
+
+	return nullptr;
+}
+
+static bool ForEachInputPin(UEdGraphNode* InNode, TFunctionRef<bool(UEdGraphPin*)> InFunction)
+{
+	const UAnimationGraphSchema* Schema = GetDefault<UAnimationGraphSchema>();
+	bool bResult = false;
+
+	for(UEdGraphPin* Pin : InNode->Pins)
+	{
+		if(Pin && Pin->Direction == EGPD_Input && !Schema->IsExecPin(*Pin) && !Schema->IsSelfPin(*Pin))
+		{
+			bResult |= InFunction(Pin);
+		}
+	}
+
+	return bResult;
+}
+
+static UEdGraphNode* FollowKnots(UEdGraphPin* FromPin, UEdGraphPin*& ToPin)
+{
+	if (FromPin->LinkedTo.Num() == 0)
+	{
+		return nullptr;
+	}
+
+	UEdGraphPin* LinkedPin = FromPin->LinkedTo[0];
+	ToPin = LinkedPin;
+	if(LinkedPin)
+	{
+		UEdGraphNode* LinkedNode = LinkedPin->GetOwningNode();
+		UK2Node_Knot* KnotNode = Cast<UK2Node_Knot>(LinkedNode);
+		while(KnotNode)
+		{
+			if(UEdGraphPin* InputPin = FindFirstInputPin(KnotNode))
+			{
+				if (InputPin->LinkedTo.Num() > 0 && InputPin->LinkedTo[0])
+				{
+					ToPin = InputPin->LinkedTo[0];
+					LinkedNode = InputPin->LinkedTo[0]->GetOwningNode();
+					KnotNode = Cast<UK2Node_Knot>(LinkedNode);
+				}
+				else
+				{
+					KnotNode = nullptr;
+				}
+			}
+		}
+		return LinkedNode;
+	}
+
+	return nullptr;
+}
+
+void FAnimBlueprintCompilerHandler_Base::FEvaluationHandlerRecord::RegisterPin(UEdGraphPin* DestPin, FProperty* AssociatedProperty, int32 AssociatedPropertyArrayIndex)
+{
+	FAnimNodeSinglePropertyHandler& Handler = ServicedProperties.FindOrAdd(AssociatedProperty->GetFName());
+
+	TArray<FString> DestPropertyPath;
+
+	// Prepend the destination property with the node's member property if the property is not on a UClass
+	if(Cast<UClass>(AssociatedProperty->Owner.ToUObject()) == nullptr)
+	{
+		DestPropertyPath.Add(NodeVariableProperty->GetName());
+	}
+
+	if(AssociatedPropertyArrayIndex != INDEX_NONE)
+	{
+		DestPropertyPath.Add(FString::Printf(TEXT("%s[%d]"), *AssociatedProperty->GetName(), AssociatedPropertyArrayIndex));
+	}
+	else
+	{
+		DestPropertyPath.Add(AssociatedProperty->GetName());
+	}
+
+	Handler.CopyRecords.Emplace(DestPin, AssociatedProperty, AssociatedPropertyArrayIndex, MoveTemp(DestPropertyPath));
+}
+
+void FAnimBlueprintCompilerHandler_Base::FEvaluationHandlerRecord::RegisterPropertyBinding(FProperty* InProperty, const FAnimGraphNodePropertyBinding& InBinding)
+{
+	FAnimNodeSinglePropertyHandler& Handler = ServicedProperties.FindOrAdd(InProperty->GetFName());
+
+	TArray<FString> DestPropertyPath;
+
+	// Prepend the destination property with the node's member property if the property is not on a UClass
+	if(Cast<UClass>(InProperty->Owner.ToUObject()) == nullptr)
+	{
+		DestPropertyPath.Add(NodeVariableProperty->GetName());
+	}
+
+	DestPropertyPath.Add(InProperty->GetName());
+
+	Handler.CopyRecords.Emplace(InBinding.PropertyPath, DestPropertyPath);
+}
+
+void FAnimBlueprintCompilerHandler_Base::FEvaluationHandlerRecord::BuildFastPathCopyRecords(FAnimBlueprintCompilerHandler_Base& InHandler)
+{
+	typedef bool (FAnimBlueprintCompilerHandler_Base::FEvaluationHandlerRecord::*GraphCheckerFunc)(FCopyRecordGraphCheckContext&, UEdGraphPin*);
+
+	GraphCheckerFunc GraphCheckerFuncs[] =
+	{
+		&FAnimBlueprintCompilerHandler_Base::FEvaluationHandlerRecord::CheckForMakeStructAccess,
+		&FAnimBlueprintCompilerHandler_Base::FEvaluationHandlerRecord::CheckForVariableGet,
+		&FAnimBlueprintCompilerHandler_Base::FEvaluationHandlerRecord::CheckForLogicalNot,
+		&FAnimBlueprintCompilerHandler_Base::FEvaluationHandlerRecord::CheckForStructMemberAccess,
+		&FAnimBlueprintCompilerHandler_Base::FEvaluationHandlerRecord::CheckForArrayAccess,
+	};
+
+	if (GetDefault<UEngine>()->bOptimizeAnimBlueprintMemberVariableAccess)
+	{
+		for (TPair<FName, FAnimNodeSinglePropertyHandler>& ServicedPropPair : ServicedProperties)
+		{
+			TArray<FPropertyCopyRecord> AllAdditionalCopyRecords;
+
+			for (FPropertyCopyRecord& CopyRecord : ServicedPropPair.Value.CopyRecords)
+			{
+				if(CopyRecord.SourcePropertyPath.Num() == 0)
+				{
+					TArray<FPropertyCopyRecord> AdditionalCopyRecords;
+
+					FCopyRecordGraphCheckContext Context(CopyRecord, AdditionalCopyRecords);
+
+					for (GraphCheckerFunc& CheckFunc : GraphCheckerFuncs)
+					{
+						if ((this->*CheckFunc)(Context, CopyRecord.DestPin))
+						{
+							break;
+						}
+					}
+
+					if(AdditionalCopyRecords.Num() > 0)
+					{
+						for(FPropertyCopyRecord& AdditionalCopyRecord : AdditionalCopyRecords)
+						{
+							CheckForMemberOnlyAccess(AdditionalCopyRecord, AdditionalCopyRecord.DestPin);
+						}
+
+						CopyRecord = AdditionalCopyRecords[0];
+
+						for(int32 AdditionalRecordIndex = 1; AdditionalRecordIndex < AdditionalCopyRecords.Num(); ++AdditionalRecordIndex)
+						{
+							AllAdditionalCopyRecords.Add(AdditionalCopyRecords[AdditionalRecordIndex]);
+						}
+					}
+					else
+					{
+						CheckForMemberOnlyAccess(CopyRecord, CopyRecord.DestPin);
+					}
+				}
+			}
+
+			// Append any additional copy records
+			ServicedPropPair.Value.CopyRecords.Append(AllAdditionalCopyRecords);
+		}
+	}
+}
+
+static void GetFullyQualifiedPathFromPin(const UEdGraphPin* Pin, TArray<FString>& OutPath)
+{
+	FString PinName = Pin->PinName.ToString();
+	while (Pin->ParentPin != nullptr)
+	{
+		PinName[Pin->ParentPin->PinName.GetStringLength()] = TEXT('.');
+		Pin = Pin->ParentPin;
+	}
+
+	UE::String::ParseTokens(PinName, TEXT('.'), [&OutPath](FStringView InStringView)
+	{
+		OutPath.Add(FString(InStringView));
+	});
+}
+
+bool FAnimBlueprintCompilerHandler_Base::FEvaluationHandlerRecord::CheckForVariableGet(FCopyRecordGraphCheckContext& Context, UEdGraphPin* DestPin)
+{
+	if(DestPin)
+	{
+		UEdGraphPin* SourcePin = nullptr;
+		if(UK2Node_VariableGet* VariableGetNode = Cast<UK2Node_VariableGet>(FollowKnots(DestPin, SourcePin)))
+		{
+			if(VariableGetNode && VariableGetNode->IsNodePure() && VariableGetNode->VariableReference.IsSelfContext())
+			{
+				if(SourcePin)
+				{
+					GetFullyQualifiedPathFromPin(SourcePin, Context.CopyRecord->SourcePropertyPath);
+					return true;
+				}
+			}
+		}
+	}
+
+	return false;
+}
+
+bool FAnimBlueprintCompilerHandler_Base::FEvaluationHandlerRecord::CheckForLogicalNot(FCopyRecordGraphCheckContext& Context, UEdGraphPin* DestPin)
+{
+	if(DestPin)
+	{
+		UEdGraphPin* SourcePin = nullptr;
+		UK2Node_CallFunction* CallFunctionNode = Cast<UK2Node_CallFunction>(FollowKnots(DestPin, SourcePin));
+		if(CallFunctionNode && CallFunctionNode->FunctionReference.GetMemberName() == FName(TEXT("Not_PreBool")))
+		{
+			// find and follow input pin
+			if(UEdGraphPin* InputPin = FindFirstInputPin(CallFunctionNode))
+			{
+				check(InputPin->PinType.PinCategory == UEdGraphSchema_K2::PC_Boolean);
+				if(CheckForVariableGet(Context, InputPin) || CheckForStructMemberAccess(Context, InputPin) || CheckForArrayAccess(Context, InputPin))
+				{
+					check(Context.CopyRecord->SourcePropertyPath.Num() > 0);	// this should have been filled in by CheckForVariableGet() or CheckForStructMemberAccess() above
+					Context.CopyRecord->Operation = EPostCopyOperation::LogicalNegateBool;
+					return true;
+				}
+			}
+		}
+	}
+
+	return false;
+}
+
+/** The functions that we can safely native-break */
+static const FName NativeBreakFunctionNameWhitelist[] =
+{
+	FName(TEXT("BreakVector")),
+	FName(TEXT("BreakVector2D")),
+	FName(TEXT("BreakRotator")),
+};
+
+/** Check whether a native break function can be safely used in the fast-path copy system (ie. source and dest data will be the same) */
+static bool IsWhitelistedNativeBreak(const FName& InFunctionName)
+{
+	for(const FName& FunctionName : NativeBreakFunctionNameWhitelist)
+	{
+		if(InFunctionName == FunctionName)
+		{
+			return true;
+		}
+	}
+
+	return false;
+}
+
+/** The functions that we can safely native-make */
+static const FName NativeMakeFunctionNameWhitelist[] =
+{
+	FName(TEXT("MakeVector")),
+	FName(TEXT("MakeVector2D")),
+	FName(TEXT("MakeRotator")),
+};
+
+/** Check whether a native break function can be safely used in the fast-path copy system (ie. source and dest data will be the same) */
+static bool IsWhitelistedNativeMake(const FName& InFunctionName)
+{
+	for(const FName& FunctionName : NativeMakeFunctionNameWhitelist)
+	{
+		if(InFunctionName == FunctionName)
+		{
+			return true;
+		}
+	}
+
+	return false;
+}
+
+bool FAnimBlueprintCompilerHandler_Base::FEvaluationHandlerRecord::CheckForStructMemberAccess(FCopyRecordGraphCheckContext& Context, UEdGraphPin* DestPin)
+{
+	if(DestPin)
+	{
+		UEdGraphPin* SourcePin = nullptr;
+		if(UK2Node_BreakStruct* BreakStructNode = Cast<UK2Node_BreakStruct>(FollowKnots(DestPin, SourcePin)))
+		{
+			if(UEdGraphPin* InputPin = FindFirstInputPin(BreakStructNode))
+			{
+				if(CheckForStructMemberAccess(Context, InputPin) || CheckForVariableGet(Context, InputPin) || CheckForArrayAccess(Context, InputPin))
+				{
+					check(Context.CopyRecord->SourcePropertyPath.Num() > 0);	// this should have been filled in by CheckForVariableGet() above
+					Context.CopyRecord->SourcePropertyPath.Add(SourcePin->PinName.ToString());
+					return true;
+				}
+			}
+		}
+		// could be a native break
+		else if(UK2Node_CallFunction* NativeBreakNode = Cast<UK2Node_CallFunction>(FollowKnots(DestPin, SourcePin)))
+		{
+			UFunction* Function = NativeBreakNode->FunctionReference.ResolveMember<UFunction>(UKismetMathLibrary::StaticClass());
+			if(Function && Function->HasMetaData(TEXT("NativeBreakFunc")) && IsWhitelistedNativeBreak(Function->GetFName()))
+			{
+				if(UEdGraphPin* InputPin = FindFirstInputPin(NativeBreakNode))
+				{
+					if(CheckForStructMemberAccess(Context, InputPin) || CheckForVariableGet(Context, InputPin) || CheckForArrayAccess(Context, InputPin))
+					{
+						check(Context.CopyRecord->SourcePropertyPath.Num() > 0);	// this should have been filled in by CheckForVariableGet() above
+						Context.CopyRecord->SourcePropertyPath.Add(SourcePin->PinName.ToString());
+						return true;
+					}
+				}
+			}
+		}
+	}
+
+	return false;
+}
+
+bool FAnimBlueprintCompilerHandler_Base::FEvaluationHandlerRecord::CheckForMakeStructAccess(FCopyRecordGraphCheckContext& Context, UEdGraphPin* DestPin)
+{
+	if(DestPin)
+	{
+		FPropertyCopyRecord OriginalRecord = *Context.CopyRecord;
+
+		UEdGraphPin* SourcePin = nullptr;
+		if(UK2Node_MakeStruct* MakeStructNode = Cast<UK2Node_MakeStruct>(FollowKnots(DestPin, SourcePin)))
+		{
+			return ForEachInputPin(MakeStructNode, [this, &Context, &OriginalRecord](UEdGraphPin* InputPin)
+			{
+				Context.CopyRecord->SourcePropertyPath = OriginalRecord.SourcePropertyPath;
+				if(CheckForStructMemberAccess(Context, InputPin) || CheckForVariableGet(Context, InputPin) || CheckForArrayAccess(Context, InputPin))
+				{
+					check(Context.CopyRecord->DestPropertyPath.Num() > 0);
+					FPropertyCopyRecord RecordCopy = *Context.CopyRecord;
+					FPropertyCopyRecord& NewRecord = Context.AdditionalCopyRecords.Add_GetRef(MoveTemp(RecordCopy));
+
+					NewRecord.DestPropertyPath = OriginalRecord.DestPropertyPath; 
+					NewRecord.DestPropertyPath.Add(InputPin->PinName.ToString());
+					return true;
+				}
+
+				return false;
+			});
+		}
+		else if(UK2Node_CallFunction* NativeMakeNode = Cast<UK2Node_CallFunction>(FollowKnots(DestPin, SourcePin)))
+		{
+			UFunction* Function = NativeMakeNode->FunctionReference.ResolveMember<UFunction>(UKismetMathLibrary::StaticClass());
+			if(Function && Function->HasMetaData(TEXT("NativeMakeFunc")) && IsWhitelistedNativeMake(Function->GetFName()))
+			{
+				return ForEachInputPin(NativeMakeNode, [this, &Context, &OriginalRecord](UEdGraphPin* InputPin)
+				{
+					Context.CopyRecord->SourcePropertyPath = OriginalRecord.SourcePropertyPath;
+					if(CheckForStructMemberAccess(Context, InputPin) || CheckForVariableGet(Context, InputPin) || CheckForArrayAccess(Context, InputPin))
+					{
+						check(Context.CopyRecord->DestPropertyPath.Num() > 0);
+						FPropertyCopyRecord RecordCopy = *Context.CopyRecord;
+						FPropertyCopyRecord& NewRecord = Context.AdditionalCopyRecords.Add_GetRef(MoveTemp(RecordCopy));
+
+						NewRecord.DestPropertyPath = OriginalRecord.DestPropertyPath;
+						NewRecord.DestPropertyPath.Add(InputPin->PinName.ToString());
+						return true;
+					}
+
+					return false;
+				});
+			}
+		}
+	}
+
+	return false;
+}
+
+bool FAnimBlueprintCompilerHandler_Base::FEvaluationHandlerRecord::CheckForArrayAccess(FCopyRecordGraphCheckContext& Context, UEdGraphPin* DestPin)
+{
+	if(DestPin)
+	{
+		UEdGraphPin* SourcePin = nullptr;
+		if(UK2Node_CallArrayFunction* CallArrayFunctionNode = Cast<UK2Node_CallArrayFunction>(FollowKnots(DestPin, SourcePin)))
+		{
+			if(CallArrayFunctionNode->GetTargetFunction() == UKismetArrayLibrary::StaticClass()->FindFunctionByName(GET_FUNCTION_NAME_CHECKED(UKismetArrayLibrary, Array_Get)))
+			{
+				// Check array index is constant
+				int32 ArrayIndex = INDEX_NONE;
+				if(UEdGraphPin* IndexPin = CallArrayFunctionNode->FindPin(TEXT("Index")))
+				{
+					if(IndexPin->LinkedTo.Num() > 0)
+					{
+						return false;
+					}
+
+					ArrayIndex = FCString::Atoi(*IndexPin->DefaultValue);
+				}
+
+				if(UEdGraphPin* TargetArrayPin = CallArrayFunctionNode->FindPin(TEXT("TargetArray")))
+				{
+					if(CheckForVariableGet(Context, TargetArrayPin) || CheckForStructMemberAccess(Context, TargetArrayPin))
+					{
+						check(Context.CopyRecord->SourcePropertyPath.Num() > 0);	// this should have been filled in by CheckForVariableGet() or CheckForStructMemberAccess() above
+						Context.CopyRecord->SourcePropertyPath.Last().Append(FString::Printf(TEXT("[%d]"), ArrayIndex));
+						return true;
+					}
+				}
+			}
+
+		
+		}
+	}
+
+	return false;
+}
+
+bool FAnimBlueprintCompilerHandler_Base::FEvaluationHandlerRecord::CheckForMemberOnlyAccess(FPropertyCopyRecord& CopyRecord, UEdGraphPin* DestPin)
+{
+	const UAnimationGraphSchema* AnimGraphDefaultSchema = GetDefault<UAnimationGraphSchema>();
+
+	if(DestPin)
+	{
+		// traverse pins to leaf nodes and check for member access/pure only
+		TArray<UEdGraphPin*> PinStack;
+		PinStack.Add(DestPin);
+		while(PinStack.Num() > 0)
+		{
+			UEdGraphPin* CurrentPin = PinStack.Pop(false);
+			for(auto& LinkedPin : CurrentPin->LinkedTo)
+			{
+				UEdGraphNode* LinkedNode = LinkedPin->GetOwningNode();
+				if(LinkedNode)
+				{
+					bool bLeafNode = true;
+					for(auto& Pin : LinkedNode->Pins)
+					{
+						if(Pin != LinkedPin && Pin->Direction == EGPD_Input && !AnimGraphDefaultSchema->IsPosePin(Pin->PinType))
+						{
+							bLeafNode = false;
+							PinStack.Add(Pin);
+						}
+					}
+
+					if(bLeafNode)
+					{
+						if(UK2Node_VariableGet* LinkedVariableGetNode = Cast<UK2Node_VariableGet>(LinkedNode))
+						{
+							if(!LinkedVariableGetNode->IsNodePure() || !LinkedVariableGetNode->VariableReference.IsSelfContext())
+							{
+								// only local variable access is allowed for leaf nodes 
+								CopyRecord.InvalidateFastPath();
+							}
+						}
+						else if(UK2Node_CallFunction* CallFunctionNode = Cast<UK2Node_CallFunction>(LinkedNode))
+						{
+							if(!CallFunctionNode->IsNodePure())
+							{
+								// only allow pure function calls
+								CopyRecord.InvalidateFastPath();
+							}
+						}
+						else if(!LinkedNode->IsA<UK2Node_TransitionRuleGetter>())
+						{
+							CopyRecord.InvalidateFastPath();
+						}
+					}
+				}
+			}
+		}
+	}
+
+	return CopyRecord.IsFastPath();
+}
+
+bool FAnimBlueprintCompilerHandler_Base::FEffectiveConstantRecord::Apply(UObject* Object)
+{
+	uint8* StructPtr = nullptr;
+	uint8* PropertyPtr = nullptr;
+	
+	if(NodeVariableProperty->Struct->IsChildOf(FAnimNode_LinkedAnimGraph::StaticStruct()))
+	{
+		PropertyPtr = ConstantProperty->ContainerPtrToValuePtr<uint8>(Object);
+	}
+	else
+	{
+		StructPtr = NodeVariableProperty->ContainerPtrToValuePtr<uint8>(Object);
+		PropertyPtr = ConstantProperty->ContainerPtrToValuePtr<uint8>(StructPtr);
+	}
+
+	if (ArrayIndex != INDEX_NONE)
+	{
+		FArrayProperty* ArrayProperty = CastFieldChecked<FArrayProperty>(ConstantProperty);
+
+		// Peer inside the array
+		FScriptArrayHelper ArrayHelper(ArrayProperty, PropertyPtr);
+
+		if (ArrayHelper.IsValidIndex(ArrayIndex))
+		{
+			FBlueprintEditorUtils::PropertyValueFromString_Direct(ArrayProperty->Inner, LiteralSourcePin->GetDefaultAsString(), ArrayHelper.GetRawPtr(ArrayIndex));
+		}
+		else
+		{
+			return false;
+		}
+	}
+	else
+	{
+		FBlueprintEditorUtils::PropertyValueFromString_Direct(ConstantProperty, LiteralSourcePin->GetDefaultAsString(), PropertyPtr);
+	}
+
+	return true;
+}
+
+#undef LOCTEXT_NAMESPACE
\ No newline at end of file
diff --git a/Engine/Source/Editor/AnimGraph/Private/AnimBlueprintCompilerHandler_Base.h b/Engine/Source/Editor/AnimGraph/Private/AnimBlueprintCompilerHandler_Base.h
new file mode 100644
--- /dev/null
+++ b/Engine/Source/Editor/AnimGraph/Private/AnimBlueprintCompilerHandler_Base.h
@@ -0,0 +1,273 @@
+// Copyright Epic Games, Inc. All Rights Reserved.
+
+#pragma once
+
+#include "CoreMinimal.h"
+#include "IAnimBlueprintCompilerHandler.h"
+#include "IPropertyAccessCompiler.h"
+#include "Animation/AnimNodeBase.h"
+
+class UAnimGraphNode_Base;
+struct FAnimGraphNodePropertyBinding;
+class UK2Node_CustomEvent;
+class IAnimBlueprintCompilationContext;
+class IAnimBlueprintGeneratedClassCompiledData;
+class IAnimBlueprintCompilerCreationContext;
+class IAnimBlueprintCompilationBracketContext;
+class IAnimBlueprintPostExpansionStepContext;
+class IAnimBlueprintCopyTermDefaultsContext;
+
+class FAnimBlueprintCompilerHandler_Base : public IAnimBlueprintCompilerHandler
+{
+public:
+	FAnimBlueprintCompilerHandler_Base(IAnimBlueprintCompilerCreationContext& InCreationContext);
+
+	// Adds a map of struct eval handlers for the specified node
+	void AddStructEvalHandlers(UAnimGraphNode_Base* InNode, IAnimBlueprintCompilationContext& InCompilationContext, IAnimBlueprintGeneratedClassCompiledData& OutCompiledData);
+
+	// Create an 'expanded' evaluation handler for the specified node, called in the compiler's node expansion step
+	void CreateEvaluationHandlerForNode(IAnimBlueprintCompilationContext& InCompilationContext, UAnimGraphNode_Base* InNode);
+
+private:
+	void StartCompilingClass(const UClass* InClass, IAnimBlueprintCompilationBracketContext& InCompilationContext, IAnimBlueprintGeneratedClassCompiledData& OutCompiledData);
+	void FinishCompilingClass(const UClass* InClass, IAnimBlueprintCompilationBracketContext& InCompilationContext, IAnimBlueprintGeneratedClassCompiledData& OutCompiledData);
+	void PostExpansionStep(const UEdGraph* InGraph, IAnimBlueprintPostExpansionStepContext& InCompilationContext, IAnimBlueprintGeneratedClassCompiledData& OutCompiledData);
+	void CopyTermDefaultsToDefaultObject(UObject* InDefaultObject, IAnimBlueprintCopyTermDefaultsContext& InCompilationContext, IAnimBlueprintGeneratedClassCompiledData& OutCompiledData);
+
+private:
+	/** Record of a single copy operation */
+	struct FPropertyCopyRecord
+	{
+		FPropertyCopyRecord(UEdGraphPin* InDestPin, FProperty* InDestProperty, int32 InDestArrayIndex, TArray<FString>&& InDestPropertyPath)
+			: DestPin(InDestPin)
+			, DestProperty(InDestProperty)
+			, DestArrayIndex(InDestArrayIndex)
+			, DestPropertyPath(MoveTemp(InDestPropertyPath))
+			, LibraryCopyIndex(INDEX_NONE)
+			, LibraryBatchType(EPropertyAccessBatchType::Unbatched)
+			, Operation(EPostCopyOperation::None)
+			, bIsFastPath(true)
+		{}
+
+		FPropertyCopyRecord(const TArray<FString>& InSourcePropertyPath, const TArray<FString>& InDestPropertyPath)
+			: DestPin(nullptr)
+			, DestProperty(nullptr)
+			, DestArrayIndex(INDEX_NONE)
+			, SourcePropertyPath(InSourcePropertyPath)
+			, DestPropertyPath(InDestPropertyPath)
+			, LibraryCopyIndex(INDEX_NONE)
+			, LibraryBatchType(EPropertyAccessBatchType::Unbatched)
+			, Operation(EPostCopyOperation::None)
+			, bIsFastPath(true)
+		{}
+
+		bool IsFastPath() const
+		{
+			return SourcePropertyPath.Num() > 0 && bIsFastPath;
+		}
+
+		void InvalidateFastPath()
+		{
+			bIsFastPath = false;
+		}
+
+		/** The destination pin we are copying to */
+		UEdGraphPin* DestPin;
+
+		/** The destination property we are copying to (on an animation node) */
+		FProperty* DestProperty;
+
+		/** The array index we use if the destination property is an array */
+		int32 DestArrayIndex;
+
+		/** The property path relative to the class */
+		TArray<FString> SourcePropertyPath;
+
+		/** The property path relative to the class */
+		TArray<FString> DestPropertyPath;
+
+		/** The index of the copy in the property access library */
+		int32 LibraryCopyIndex;
+
+		/** the batch type within the property access library */
+		EPropertyAccessBatchType LibraryBatchType;
+
+		/** Any operation we want to perform post-copy on the destination data */
+		EPostCopyOperation Operation;
+
+		/** Fast-path flag */
+		bool bIsFastPath;
+	};
+
+	// Context used to build fast-path copy records
+	struct FCopyRecordGraphCheckContext
+	{
+		FCopyRecordGraphCheckContext(FPropertyCopyRecord& InCopyRecord, TArray<FPropertyCopyRecord>& InAdditionalCopyRecords)
+			: CopyRecord(&InCopyRecord)
+			, AdditionalCopyRecords(InAdditionalCopyRecords)
+		{}
+
+		// Copy record we are operating on
+		FPropertyCopyRecord* CopyRecord;
+
+		// Things like split input pins can add additional copy records
+		TArray<FPropertyCopyRecord>& AdditionalCopyRecords;
+	};
+
+	// Wireup record for a single anim node property (which might be an array)
+	struct FAnimNodeSinglePropertyHandler
+	{
+		/** Copy records */
+		TArray<FPropertyCopyRecord> CopyRecords;
+
+		// If the anim instance is the container target instead of the node.
+		bool bInstanceIsTarget;
+
+		FAnimNodeSinglePropertyHandler()
+			: bInstanceIsTarget(false)
+		{
+		}
+	};
+
+	// Record for a property that was exposed as a pin, but wasn't wired up (just a literal)
+	struct FEffectiveConstantRecord
+	{
+	public:
+		// The node variable that the handler is in
+		class FStructProperty* NodeVariableProperty;
+
+		// The property within the struct to set
+		class FProperty* ConstantProperty;
+
+		// The array index if ConstantProperty is an array property, or INDEX_NONE otherwise
+		int32 ArrayIndex;
+
+		// The pin to pull the DefaultValue/DefaultObject from
+		UEdGraphPin* LiteralSourcePin;
+
+		FEffectiveConstantRecord()
+			: NodeVariableProperty(NULL)
+			, ConstantProperty(NULL)
+			, ArrayIndex(INDEX_NONE)
+			, LiteralSourcePin(NULL)
+		{
+		}
+
+		FEffectiveConstantRecord(FStructProperty* ContainingNodeProperty, UEdGraphPin* SourcePin, FProperty* SourcePinProperty, int32 SourceArrayIndex)
+			: NodeVariableProperty(ContainingNodeProperty)
+			, ConstantProperty(SourcePinProperty)
+			, ArrayIndex(SourceArrayIndex)
+			, LiteralSourcePin(SourcePin)
+		{
+		}
+
+		bool Apply(UObject* Object);
+	};
+
+	/** BP execution handler for Anim node */
+	struct FEvaluationHandlerRecord
+	{
+	public:
+		// The Node this record came from
+		UAnimGraphNode_Base* AnimGraphNode;
+
+		// The node variable that the handler is in
+		FStructProperty* NodeVariableProperty;
+
+		// The specific evaluation handler inside the specified node
+		int32 EvaluationHandlerIdx;
+
+		// Whether or not our serviced properties are actually on the anim node 
+		bool bServicesNodeProperties;
+
+		// Whether or not our serviced properties are actually on the instance instead of the node
+		bool bServicesInstanceProperties;
+
+		// Set of properties serviced by this handler (Map from property name to the record for that property)
+		TMap<FName, FAnimNodeSinglePropertyHandler> ServicedProperties;
+
+		// The generated custom event node
+		TArray<UEdGraphNode*> CustomEventNodes;
+
+		// The resulting function name
+		FName HandlerFunctionName;
+
+	public:
+
+		FEvaluationHandlerRecord()
+			: AnimGraphNode(nullptr)
+			, NodeVariableProperty(nullptr)
+			, EvaluationHandlerIdx(INDEX_NONE)
+			, bServicesNodeProperties(false)
+			, bServicesInstanceProperties(false)
+			, HandlerFunctionName(NAME_None)
+		{}
+
+		bool IsFastPath() const
+		{
+			for(TMap<FName, FAnimNodeSinglePropertyHandler>::TConstIterator It(ServicedProperties); It; ++It)
+			{
+				const FAnimNodeSinglePropertyHandler& AnimNodeSinglePropertyHandler = It.Value();
+				for (const FPropertyCopyRecord& CopyRecord : AnimNodeSinglePropertyHandler.CopyRecords)
+				{
+					if (!CopyRecord.IsFastPath())
+					{
+						return false;
+					}
+				}
+			}
+
+			return true;
+		}
+
+		bool IsValid() const
+		{
+			return NodeVariableProperty != nullptr;
+		}
+
+		void PatchFunctionNameAndCopyRecordsInto(FExposedValueHandler& Handler) const;
+
+		void RegisterPin(UEdGraphPin* DestPin, FProperty* AssociatedProperty, int32 AssociatedPropertyArrayIndex);
+
+		void RegisterPropertyBinding(FProperty* InProperty, const FAnimGraphNodePropertyBinding& InBinding);
+
+		FStructProperty* GetHandlerNodeProperty() const { return NodeVariableProperty; }
+
+		void BuildFastPathCopyRecords(FAnimBlueprintCompilerHandler_Base& InHandler);
+
+	private:
+
+		bool CheckForVariableGet(FCopyRecordGraphCheckContext& Context, UEdGraphPin* DestPin);
+
+		bool CheckForLogicalNot(FCopyRecordGraphCheckContext& Context, UEdGraphPin* DestPin);
+
+		bool CheckForStructMemberAccess(FCopyRecordGraphCheckContext& Context, UEdGraphPin* DestPin);
+
+		bool CheckForMemberOnlyAccess(FPropertyCopyRecord& Context, UEdGraphPin* DestPin);
+
+		bool CheckForMakeStructAccess(FCopyRecordGraphCheckContext& Context, UEdGraphPin* DestPin);
+
+		bool CheckForArrayAccess(FCopyRecordGraphCheckContext& Context, UEdGraphPin* DestPin);
+	};
+
+	// Create an evaluation handler for the specified node/record
+	void CreateEvaluationHandler(IAnimBlueprintCompilationContext& InCompilationContext, UAnimGraphNode_Base* InNode, FEvaluationHandlerRecord& Record);
+
+private:
+	// Records of pose pins for later patchup with an associated evaluation handler
+	TMap<UAnimGraphNode_Base*, FEvaluationHandlerRecord> PerNodeStructEvalHandlers;
+
+	// List of successfully created evaluation handlers
+	TArray<FEvaluationHandlerRecord> ValidEvaluationHandlerList;
+	TMap<UAnimGraphNode_Base*, int32> ValidEvaluationHandlerMap;
+
+	// List of animation node literals (values exposed as pins but never wired up) that need to be pushed into the CDO
+	TArray<FEffectiveConstantRecord> ValidAnimNodePinConstants;
+
+	// Set of used handler function names
+	TSet<FName> HandlerFunctionNames;
+
+	// Delegate handle for registering against library pre/post-compilation
+	FDelegateHandle PreLibraryCompiledDelegateHandle;
+	FDelegateHandle PostLibraryCompiledDelegateHandle;
+};
\ No newline at end of file
diff --git a/Engine/Source/Editor/AnimGraph/Private/AnimBlueprintCompilerHandler_CachedPose.cpp b/Engine/Source/Editor/AnimGraph/Private/AnimBlueprintCompilerHandler_CachedPose.cpp
new file mode 100644
--- /dev/null
+++ b/Engine/Source/Editor/AnimGraph/Private/AnimBlueprintCompilerHandler_CachedPose.cpp
@@ -0,0 +1,178 @@
+// Copyright Epic Games, Inc. All Rights Reserved.
+
+#include "AnimBlueprintCompilerHandler_CachedPose.h"
+#include "AnimGraphNode_SaveCachedPose.h"
+#include "AnimGraphNode_Root.h"
+#include "AnimGraphNode_UseCachedPose.h"
+#include "Kismet2/CompilerResultsLog.h"
+#include "AnimGraphNode_StateMachine.h"
+#include "AnimGraphNode_StateResult.h"
+#include "AnimationStateMachineGraph.h"
+#include "IAnimBlueprintGeneratedClassCompiledData.h"
+#include "IAnimBlueprintCompilerCreationContext.h"
+#include "IAnimBlueprintCompilationContext.h"
+
+FAnimBlueprintCompilerHandler_CachedPose::FAnimBlueprintCompilerHandler_CachedPose(IAnimBlueprintCompilerCreationContext& InCreationContext)
+{
+	InCreationContext.OnPreProcessAnimationNodes().AddRaw(this, &FAnimBlueprintCompilerHandler_CachedPose::PreProcessAnimationNodes);
+	InCreationContext.OnPostProcessAnimationNodes().AddRaw(this, &FAnimBlueprintCompilerHandler_CachedPose::PostProcessAnimationNodes);
+}
+
+void FAnimBlueprintCompilerHandler_CachedPose::PreProcessAnimationNodes(TArrayView<UAnimGraphNode_Base*> InAnimNodes, IAnimBlueprintCompilationContext& InCompilationContext, IAnimBlueprintGeneratedClassCompiledData& OutCompiledData)
+{
+	for(UAnimGraphNode_Base* Node : InAnimNodes)
+	{
+		if(UAnimGraphNode_SaveCachedPose* SavePoseRoot = Cast<UAnimGraphNode_SaveCachedPose>(Node))
+		{
+			//@TODO: Ideally we only add these if there is a UseCachedPose node referencing them, but those can be anywhere and are hard to grab
+			SaveCachedPoseNodes.Add(SavePoseRoot->CacheName, SavePoseRoot);
+		}
+	}
+}
+
+void FAnimBlueprintCompilerHandler_CachedPose::PostProcessAnimationNodes(TArrayView<UAnimGraphNode_Base*> InAnimNodes, IAnimBlueprintCompilationContext& InCompilationContext, IAnimBlueprintGeneratedClassCompiledData& OutCompiledData)
+{
+	// Build cached pose map
+	BuildCachedPoseNodeUpdateOrder(InCompilationContext, OutCompiledData);
+}
+
+TAutoConsoleVariable<int32> CVarAnimDebugCachePoseNodeUpdateOrder(TEXT("a.Compiler.CachePoseNodeUpdateOrderDebug.Enable"), 0, TEXT("Toggle debugging for CacheNodeUpdateOrder debug during AnimBP compilation"));
+
+void FAnimBlueprintCompilerHandler_CachedPose::BuildCachedPoseNodeUpdateOrder(IAnimBlueprintCompilationContext& InCompilationContext, IAnimBlueprintGeneratedClassCompiledData& OutCompiledData)
+{
+	TArray<UAnimGraphNode_Root*> RootNodes;
+	InCompilationContext.GetConsolidatedEventGraph()->GetNodesOfClass<UAnimGraphNode_Root>(RootNodes);
+
+	// State results are also "root" nodes, need to find the true roots
+	RootNodes.RemoveAll([](UAnimGraphNode_Root* InPossibleRootNode)
+	{
+		return InPossibleRootNode->GetClass() != UAnimGraphNode_Root::StaticClass();
+	});
+
+	const bool bEnableDebug = (CVarAnimDebugCachePoseNodeUpdateOrder.GetValueOnAnyThread() == 1);
+
+	for(UAnimGraphNode_Root* RootNode : RootNodes)
+	{
+		TArray<UAnimGraphNode_SaveCachedPose*> OrderedSavePoseNodes;
+
+		TArray<UAnimGraphNode_Base*> VisitedRootNodes;
+	
+		UE_CLOG(bEnableDebug, LogAnimation, Display, TEXT("CachePoseNodeOrdering BEGIN")); 
+
+		CachePoseNodeOrdering_StartNewTraversal(InCompilationContext, RootNode, OrderedSavePoseNodes, VisitedRootNodes);
+
+		UE_CLOG(bEnableDebug, LogAnimation, Display, TEXT("CachePoseNodeOrdering END"));
+
+		if (bEnableDebug)
+		{
+			UE_LOG(LogAnimation, Display, TEXT("Ordered Save Pose Node List:"));
+			for (UAnimGraphNode_SaveCachedPose* SavedPoseNode : OrderedSavePoseNodes)
+			{
+				UE_LOG(LogAnimation, Display, TEXT("\t%s"), *SavedPoseNode->Node.CachePoseName.ToString())
+			}
+			UE_LOG(LogAnimation, Display, TEXT("End List"));
+		}
+
+		FCachedPoseIndices& OrderedSavedPoseIndices = OutCompiledData.GetOrderedSavedPoseIndicesMap().FindOrAdd(RootNode->Node.Name);
+
+		for(UAnimGraphNode_SaveCachedPose* PoseNode : OrderedSavePoseNodes)
+		{
+			if(const int32* NodeIndex = InCompilationContext.GetAllocatedAnimNodeIndices().Find(PoseNode))
+			{
+				OrderedSavedPoseIndices.OrderedSavedPoseNodeIndices.Add(*NodeIndex);
+			}
+			else
+			{
+				InCompilationContext.GetMessageLog().Error(TEXT("Failed to find index for a saved pose node while building ordered pose list."));
+			}
+		}
+	}
+}
+
+void FAnimBlueprintCompilerHandler_CachedPose::CachePoseNodeOrdering_StartNewTraversal(IAnimBlueprintCompilationContext& InCompilationContext, UAnimGraphNode_Base* InRootNode, TArray<UAnimGraphNode_SaveCachedPose*> &OrderedSavePoseNodes, TArray<UAnimGraphNode_Base*> VisitedRootNodes)
+{
+	check(InRootNode);
+	UAnimGraphNode_SaveCachedPose* RootCacheNode = Cast<UAnimGraphNode_SaveCachedPose>(InRootNode);
+	FString RootName = RootCacheNode ? RootCacheNode->CacheName : InRootNode->GetName();
+
+	const bool bEnableDebug = (CVarAnimDebugCachePoseNodeUpdateOrder.GetValueOnAnyThread() == 1);
+
+	UE_CLOG(bEnableDebug, LogAnimation, Display, TEXT("StartNewTraversal %s"), *RootName);
+
+	// Track which root nodes we've visited to prevent infinite recursion.
+	VisitedRootNodes.Add(InRootNode);
+
+	// Need a list of only what we find here to recurse, we can't do that with the total list
+	TArray<UAnimGraphNode_SaveCachedPose*> InternalOrderedNodes;
+
+	// Traverse whole graph from root collecting SaveCachePose nodes we've touched.
+	CachePoseNodeOrdering_TraverseInternal(InCompilationContext, InRootNode, InternalOrderedNodes);
+
+	// Process nodes that we've touched 
+	UE_CLOG(bEnableDebug, LogAnimation, Display, TEXT("Process Queue for %s"), *RootName);
+	for (UAnimGraphNode_SaveCachedPose* QueuedCacheNode : InternalOrderedNodes)
+	{
+		if (VisitedRootNodes.Contains(QueuedCacheNode))
+		{
+			UE_CLOG(bEnableDebug, LogAnimation, Display, TEXT("Process Queue SaveCachePose %s. ALREADY VISITED, INFINITE RECURSION DETECTED! SKIPPING"), *QueuedCacheNode->CacheName);
+			InCompilationContext.GetMessageLog().Error(*FString::Printf(TEXT("Infinite recursion detected with SaveCachePose %s and %s"), *RootName, *QueuedCacheNode->CacheName));
+			continue;
+		}
+		else
+		{
+			OrderedSavePoseNodes.Remove(QueuedCacheNode);
+			OrderedSavePoseNodes.Add(QueuedCacheNode);
+
+			CachePoseNodeOrdering_StartNewTraversal(InCompilationContext, QueuedCacheNode, OrderedSavePoseNodes, VisitedRootNodes);
+		}
+	}
+
+	UE_CLOG(bEnableDebug, LogAnimation, Display, TEXT("EndNewTraversal %s"), *RootName);
+}
+
+void FAnimBlueprintCompilerHandler_CachedPose::CachePoseNodeOrdering_TraverseInternal(IAnimBlueprintCompilationContext& InCompilationContext, UAnimGraphNode_Base* InAnimGraphNode, TArray<UAnimGraphNode_SaveCachedPose*> &OrderedSavePoseNodes)
+{
+	TArray<UAnimGraphNode_Base*> LinkedAnimNodes;
+	InCompilationContext.GetLinkedAnimNodes(InAnimGraphNode, LinkedAnimNodes);
+
+	const bool bEnableDebug = (CVarAnimDebugCachePoseNodeUpdateOrder.GetValueOnAnyThread() == 1);
+
+	for (UAnimGraphNode_Base* LinkedNode : LinkedAnimNodes)
+	{
+		UE_CLOG(bEnableDebug, LogAnimation, Display, TEXT("\t Processing %s"), *LinkedNode->GetName());
+		if (UAnimGraphNode_UseCachedPose* UsePoseNode = Cast<UAnimGraphNode_UseCachedPose>(LinkedNode))
+		{
+			if (UAnimGraphNode_SaveCachedPose* SaveNode = UsePoseNode->SaveCachedPoseNode.Get())
+			{
+				UE_CLOG(bEnableDebug, LogAnimation, Display, TEXT("\t Queueing SaveCachePose %s"), *SaveNode->CacheName);
+
+				// Requeue the node we found
+				OrderedSavePoseNodes.Remove(SaveNode);
+				OrderedSavePoseNodes.Add(SaveNode);
+			}
+		}
+		else if (UAnimGraphNode_StateMachine* StateMachineNode = Cast<UAnimGraphNode_StateMachine>(LinkedNode))
+		{
+			for (UEdGraph* StateGraph : StateMachineNode->EditorStateMachineGraph->SubGraphs)
+			{
+				TArray<UAnimGraphNode_StateResult*> ResultNodes;
+				StateGraph->GetNodesOfClass(ResultNodes);
+
+				// We should only get one here but doesn't hurt to loop here in case that changes
+				for (UAnimGraphNode_StateResult* ResultNode : ResultNodes)
+				{
+					CachePoseNodeOrdering_TraverseInternal(InCompilationContext, ResultNode, OrderedSavePoseNodes);
+				}
+			}
+		}
+		else
+		{
+			CachePoseNodeOrdering_TraverseInternal(InCompilationContext, LinkedNode, OrderedSavePoseNodes);
+		}
+	}
+}
+
+const TMap<FString, UAnimGraphNode_SaveCachedPose*>& FAnimBlueprintCompilerHandler_CachedPose::GetSaveCachedPoseNodes() const
+{
+	return SaveCachedPoseNodes;
+}
diff --git a/Engine/Source/Editor/AnimGraph/Private/AnimBlueprintCompilerHandler_CachedPose.h b/Engine/Source/Editor/AnimGraph/Private/AnimBlueprintCompilerHandler_CachedPose.h
new file mode 100644
--- /dev/null
+++ b/Engine/Source/Editor/AnimGraph/Private/AnimBlueprintCompilerHandler_CachedPose.h
@@ -0,0 +1,37 @@
+// Copyright Epic Games, Inc. All Rights Reserved.
+
+#pragma once
+
+#include "CoreMinimal.h"
+#include "IAnimBlueprintCompilerHandler.h"
+
+class UAnimGraphNode_SaveCachedPose;
+class IAnimBlueprintCompilerCreationContext;
+class IAnimBlueprintCompilationContext;
+class IAnimBlueprintGeneratedClassCompiledData;
+
+class FAnimBlueprintCompilerHandler_CachedPose : public IAnimBlueprintCompilerHandler
+{
+public:
+	FAnimBlueprintCompilerHandler_CachedPose(IAnimBlueprintCompilerCreationContext& InCreationContext);
+
+	// Get the map of cache name to encountered save cached pose nodes
+	const TMap<FString, UAnimGraphNode_SaveCachedPose*>& GetSaveCachedPoseNodes() const;
+
+private:
+	void PreProcessAnimationNodes(TArrayView<UAnimGraphNode_Base*> InAnimNodes, IAnimBlueprintCompilationContext& InCompilationContext, IAnimBlueprintGeneratedClassCompiledData& OutCompiledData);
+	void PostProcessAnimationNodes(TArrayView<UAnimGraphNode_Base*> InAnimNodes, IAnimBlueprintCompilationContext& InCompilationContext, IAnimBlueprintGeneratedClassCompiledData& OutCompiledData);
+
+	// Builds the update order list for saved pose nodes in this blueprint
+	void BuildCachedPoseNodeUpdateOrder(IAnimBlueprintCompilationContext& InCompilationContext, IAnimBlueprintGeneratedClassCompiledData& OutCompiledData);
+
+	// Traverses a graph to collect save pose nodes starting at InRootNode, then processes each node
+	void CachePoseNodeOrdering_StartNewTraversal(IAnimBlueprintCompilationContext& InCompilationContext, UAnimGraphNode_Base* InRootNode, TArray<UAnimGraphNode_SaveCachedPose*> &OrderedSavePoseNodes, TArray<UAnimGraphNode_Base*> VisitedRootNodes);
+
+	// Traverses a graph to collect save pose nodes starting at InAnimGraphNode, does NOT process saved pose nodes afterwards
+	void CachePoseNodeOrdering_TraverseInternal(IAnimBlueprintCompilationContext& InCompilationContext, UAnimGraphNode_Base* InAnimGraphNode, TArray<UAnimGraphNode_SaveCachedPose*> &OrderedSavePoseNodes);
+
+private:
+	// Map of cache name to encountered save cached pose nodes
+	TMap<FString, UAnimGraphNode_SaveCachedPose*> SaveCachedPoseNodes;
+};
\ No newline at end of file
diff --git a/Engine/Source/Editor/AnimGraph/Private/AnimBlueprintCompilerHandler_LinkedAnimGraph.cpp b/Engine/Source/Editor/AnimGraph/Private/AnimBlueprintCompilerHandler_LinkedAnimGraph.cpp
new file mode 100644
--- /dev/null
+++ b/Engine/Source/Editor/AnimGraph/Private/AnimBlueprintCompilerHandler_LinkedAnimGraph.cpp
@@ -0,0 +1,23 @@
+// Copyright Epic Games, Inc. All Rights Reserved.
+
+#include "AnimBlueprintCompilerHandler_LinkedAnimGraph.h"
+#include "AnimGraphNode_LinkedAnimGraphBase.h"
+#include "IAnimBlueprintCompilerCreationContext.h"
+#include "IAnimBlueprintCompilationContext.h"
+#include "IAnimBlueprintGeneratedClassCompiledData.h"
+
+FAnimBlueprintCompilerHandler_LinkedAnimGraph::FAnimBlueprintCompilerHandler_LinkedAnimGraph(IAnimBlueprintCompilerCreationContext& InCreationContext)
+{
+	InCreationContext.OnPreProcessAnimationNodes().AddRaw(this, &FAnimBlueprintCompilerHandler_LinkedAnimGraph::PreProcessAnimationNodes);
+}
+
+void FAnimBlueprintCompilerHandler_LinkedAnimGraph::PreProcessAnimationNodes(TArrayView<UAnimGraphNode_Base*> InAnimNodes, IAnimBlueprintCompilationContext& InCompilationContext, IAnimBlueprintGeneratedClassCompiledData& OutCompiledData)
+{
+	for(UAnimGraphNode_Base* AnimNode : InAnimNodes)
+	{
+		if(UAnimGraphNode_LinkedAnimGraphBase* LinkedAnimGraphBase = Cast<UAnimGraphNode_LinkedAnimGraphBase>(AnimNode))
+		{
+			LinkedAnimGraphBase->AllocatePoseLinks();
+		}
+	}
+}
\ No newline at end of file
diff --git a/Engine/Source/Editor/AnimGraph/Private/AnimBlueprintCompilerHandler_LinkedAnimGraph.h b/Engine/Source/Editor/AnimGraph/Private/AnimBlueprintCompilerHandler_LinkedAnimGraph.h
new file mode 100644
--- /dev/null
+++ b/Engine/Source/Editor/AnimGraph/Private/AnimBlueprintCompilerHandler_LinkedAnimGraph.h
@@ -0,0 +1,19 @@
+// Copyright Epic Games, Inc. All Rights Reserved.
+
+#pragma once
+
+#include "CoreMinimal.h"
+#include "IAnimBlueprintCompilerHandler.h"
+
+class IAnimBlueprintCompilerCreationContext;
+class IAnimBlueprintCompilationContext;
+class IAnimBlueprintGeneratedClassCompiledData;
+
+class FAnimBlueprintCompilerHandler_LinkedAnimGraph : public IAnimBlueprintCompilerHandler
+{
+public:
+	FAnimBlueprintCompilerHandler_LinkedAnimGraph(IAnimBlueprintCompilerCreationContext& InCreationContext);
+
+private:
+	void PreProcessAnimationNodes(TArrayView<UAnimGraphNode_Base*> InAnimNodes, IAnimBlueprintCompilationContext& InCompilationContext, IAnimBlueprintGeneratedClassCompiledData& OutCompiledData);
+};
\ No newline at end of file
diff --git a/Engine/Source/Editor/AnimGraph/Private/AnimBlueprintCompilerHandler_StateMachine.cpp b/Engine/Source/Editor/AnimGraph/Private/AnimBlueprintCompilerHandler_StateMachine.cpp
new file mode 100644
--- /dev/null
+++ b/Engine/Source/Editor/AnimGraph/Private/AnimBlueprintCompilerHandler_StateMachine.cpp
@@ -0,0 +1,553 @@
+// Copyright Epic Games, Inc. All Rights Reserved.
+
+#include "AnimBlueprintCompilerHandler_StateMachine.h"
+#include "Animation/AnimBlueprintGeneratedClass.h"
+#include "K2Node_AnimGetter.h"
+#include "K2Node_TransitionRuleGetter.h"
+#include "Kismet2/CompilerResultsLog.h"
+#include "AnimGraphNode_Base.h"
+#include "K2Node_StructMemberGet.h"
+#include "AnimStateNode.h"
+#include "AnimStateTransitionNode.h"
+#include "AnimGraphNode_StateMachineBase.h"
+#include "EdGraphUtilities.h"
+#include "AnimationStateMachineSchema.h"
+#include "Algo/Transform.h"
+#include "Animation/AnimTypes.h"
+#include "IAnimBlueprintGeneratedClassCompiledData.h"
+#include "IAnimBlueprintCompilerCreationContext.h"
+#include "IAnimBlueprintCompilationContext.h"
+
+#define LOCTEXT_NAMESPACE "StateMachineHandler"
+
+FAnimBlueprintCompilerHandler_StateMachine::FAnimBlueprintCompilerHandler_StateMachine(IAnimBlueprintCompilerCreationContext& InCreationContext)
+{
+	InCreationContext.OnPreProcessAnimationNodes().AddRaw(this, &FAnimBlueprintCompilerHandler_StateMachine::PreProcessAnimationNodes);
+	InCreationContext.OnPostProcessAnimationNodes().AddRaw(this, &FAnimBlueprintCompilerHandler_StateMachine::PostProcessAnimationNodes);
+	InCreationContext.RegisterKnownGraphSchema(UAnimationStateMachineSchema::StaticClass());
+}
+
+void FAnimBlueprintCompilerHandler_StateMachine::PreProcessAnimationNodes(TArrayView<UAnimGraphNode_Base*> InAnimNodes, IAnimBlueprintCompilationContext& InCompilationContext, IAnimBlueprintGeneratedClassCompiledData& OutCompiledData)
+{
+	InCompilationContext.GetConsolidatedEventGraph()->GetNodesOfClass<UK2Node_TransitionRuleGetter>(RootTransitionGetters);
+
+	// Get anim getters from the root anim graph (processing the nodes below will collect them in nested graphs)
+	InCompilationContext.GetConsolidatedEventGraph()->GetNodesOfClass<UK2Node_AnimGetter>(RootGraphAnimGetters);
+}
+
+void FAnimBlueprintCompilerHandler_StateMachine::PostProcessAnimationNodes(TArrayView<UAnimGraphNode_Base*> InAnimNodes, IAnimBlueprintCompilationContext& InCompilationContext, IAnimBlueprintGeneratedClassCompiledData& OutCompiledData)
+{
+	// Process the getter nodes in the graph if there were any
+	for (auto GetterIt = RootTransitionGetters.CreateIterator(); GetterIt; ++GetterIt)
+	{
+		ProcessTransitionGetter(*GetterIt, nullptr, InCompilationContext, OutCompiledData); // transition nodes should not appear at top-level
+	}
+
+	// Wire root getters
+	for(UK2Node_AnimGetter* RootGraphGetter : RootGraphAnimGetters)
+	{
+		AutoWireAnimGetter(RootGraphGetter, nullptr, InCompilationContext, OutCompiledData);
+	}
+
+	// Wire nested getters
+	for(UK2Node_AnimGetter* Getter : FoundGetterNodes)
+	{
+		AutoWireAnimGetter(Getter, nullptr, InCompilationContext, OutCompiledData);
+	}
+}
+
+UK2Node_CallFunction* FAnimBlueprintCompilerHandler_StateMachine::SpawnCallAnimInstanceFunction(IAnimBlueprintCompilationContext& InCompilationContext, UEdGraphNode* SourceNode, FName FunctionName)
+{
+	//@TODO: SKELETON: This is a call on a parent function (UAnimInstance::StaticClass() specifically), should we treat it as self or not?
+	UK2Node_CallFunction* FunctionCall = InCompilationContext.SpawnIntermediateNode<UK2Node_CallFunction>(SourceNode);
+	FunctionCall->FunctionReference.SetSelfMember(FunctionName);
+	FunctionCall->AllocateDefaultPins();
+
+	return FunctionCall;
+}
+
+void FAnimBlueprintCompilerHandler_StateMachine::ProcessTransitionGetter(UK2Node_TransitionRuleGetter* Getter, UAnimStateTransitionNode* TransitionNode, IAnimBlueprintCompilationContext& InCompilationContext, IAnimBlueprintGeneratedClassCompiledData& OutCompiledData)
+{
+	// Get common elements for multiple getters
+	UEdGraphPin* OutputPin = Getter->GetOutputPin();
+
+	UEdGraphPin* SourceTimePin = NULL;
+	UAnimationAsset* AnimAsset= NULL;
+	int32 PlayerNodeIndex = INDEX_NONE;
+
+	if (UAnimGraphNode_Base* SourcePlayerNode = Getter->AssociatedAnimAssetPlayerNode)
+	{
+		// This check should never fail as the source state is always processed first before handling it's rules
+		UAnimGraphNode_Base* TrueSourceNode = InCompilationContext.GetMessageLog().FindSourceObjectTypeChecked<UAnimGraphNode_Base>(SourcePlayerNode);
+		UAnimGraphNode_Base* UndertypedPlayerNode = InCompilationContext.GetSourceNodeToProcessedNodeMap().FindRef(TrueSourceNode);
+
+		if (UndertypedPlayerNode == NULL)
+		{
+			InCompilationContext.GetMessageLog().Error(TEXT("ICE: Player node @@ was not processed prior to handling a transition getter @@ that used it"), SourcePlayerNode, Getter);
+			return;
+		}
+
+		// Make sure the node is still relevant
+		UEdGraph* PlayerGraph = UndertypedPlayerNode->GetGraph();
+		if (!PlayerGraph->Nodes.Contains(UndertypedPlayerNode))
+		{
+			InCompilationContext.GetMessageLog().Error(TEXT("@@ is not associated with a node in @@; please delete and recreate it"), Getter, PlayerGraph);
+		}
+
+		// Make sure the referenced AnimAsset player has been allocated
+		PlayerNodeIndex = InCompilationContext.GetAllocationIndexOfNode(UndertypedPlayerNode);
+		if (PlayerNodeIndex == INDEX_NONE)
+		{
+			InCompilationContext.GetMessageLog().Error(*LOCTEXT("BadAnimAssetNodeUsedInGetter", "@@ doesn't have a valid associated AnimAsset node.  Delete and recreate it").ToString(), Getter);
+		}
+
+		// Grab the AnimAsset, and time pin if needed
+		UScriptStruct* TimePropertyInStructType = NULL;
+		const TCHAR* TimePropertyName = NULL;
+		if (UndertypedPlayerNode->DoesSupportTimeForTransitionGetter())
+		{
+			AnimAsset = UndertypedPlayerNode->GetAnimationAsset();
+			TimePropertyInStructType = UndertypedPlayerNode->GetTimePropertyStruct();
+			TimePropertyName = UndertypedPlayerNode->GetTimePropertyName();
+		}
+		else
+		{
+			InCompilationContext.GetMessageLog().Error(TEXT("@@ is associated with @@, which is an unexpected type"), Getter, UndertypedPlayerNode);
+		}
+
+		bool bNeedTimePin = false;
+
+		// Determine if we need to read the current time variable from the specified sequence player
+		switch (Getter->GetterType)
+		{
+		case ETransitionGetter::AnimationAsset_GetCurrentTime:
+		case ETransitionGetter::AnimationAsset_GetCurrentTimeFraction:
+		case ETransitionGetter::AnimationAsset_GetTimeFromEnd:
+		case ETransitionGetter::AnimationAsset_GetTimeFromEndFraction:
+			bNeedTimePin = true;
+			break;
+		default:
+			bNeedTimePin = false;
+			break;
+		}
+
+		if (bNeedTimePin && (PlayerNodeIndex != INDEX_NONE) && (TimePropertyName != NULL) && (TimePropertyInStructType != NULL))
+		{
+			const FProperty* NodeProperty = InCompilationContext.GetAllocatedPropertiesByIndex().FindChecked(PlayerNodeIndex);
+
+			// Create a struct member read node to grab the current position of the sequence player node
+			UK2Node_StructMemberGet* TimeReadNode = InCompilationContext.SpawnIntermediateNode<UK2Node_StructMemberGet>(Getter, InCompilationContext.GetConsolidatedEventGraph());
+			TimeReadNode->VariableReference.SetSelfMember(NodeProperty->GetFName());
+			TimeReadNode->StructType = TimePropertyInStructType;
+
+			TimeReadNode->AllocatePinsForSingleMemberGet(TimePropertyName);
+			SourceTimePin = TimeReadNode->FindPinChecked(TimePropertyName);
+		}
+	}
+
+	// Expand it out
+	UK2Node_CallFunction* GetterHelper = NULL;
+	switch (Getter->GetterType)
+	{
+	case ETransitionGetter::AnimationAsset_GetCurrentTime:
+		if ((AnimAsset != NULL) && (SourceTimePin != NULL))
+		{
+			GetterHelper = SpawnCallAnimInstanceFunction(InCompilationContext, Getter, TEXT("GetInstanceAssetPlayerTime"));
+			GetterHelper->FindPinChecked(TEXT("AssetPlayerIndex"))->DefaultValue = FString::FromInt(PlayerNodeIndex);
+		}
+		else
+		{
+			if (Getter->AssociatedAnimAssetPlayerNode)
+			{
+				InCompilationContext.GetMessageLog().Error(TEXT("Please replace @@ with Get Relevant Anim Time. @@ has no animation asset"), Getter, Getter->AssociatedAnimAssetPlayerNode);
+			}
+			else
+			{
+				InCompilationContext.GetMessageLog().Error(TEXT("@@ is not asscociated with an asset player"), Getter);
+			}
+		}
+		break;
+	case ETransitionGetter::AnimationAsset_GetLength:
+		if (AnimAsset != NULL)
+		{
+			GetterHelper = SpawnCallAnimInstanceFunction(InCompilationContext, Getter, TEXT("GetInstanceAssetPlayerLength"));
+			GetterHelper->FindPinChecked(TEXT("AssetPlayerIndex"))->DefaultValue = FString::FromInt(PlayerNodeIndex);
+		}
+		else
+		{
+			if (Getter->AssociatedAnimAssetPlayerNode)
+			{
+				InCompilationContext.GetMessageLog().Error(TEXT("Please replace @@ with Get Relevant Anim Length. @@ has no animation asset"), Getter, Getter->AssociatedAnimAssetPlayerNode);
+			}
+			else
+			{
+				InCompilationContext.GetMessageLog().Error(TEXT("@@ is not asscociated with an asset player"), Getter);
+			}
+		}
+		break;
+	case ETransitionGetter::AnimationAsset_GetCurrentTimeFraction:
+		if ((AnimAsset != NULL) && (SourceTimePin != NULL))
+		{
+			GetterHelper = SpawnCallAnimInstanceFunction(InCompilationContext, Getter, TEXT("GetInstanceAssetPlayerTimeFraction"));
+			GetterHelper->FindPinChecked(TEXT("AssetPlayerIndex"))->DefaultValue = FString::FromInt(PlayerNodeIndex);
+		}
+		else
+		{
+			if (Getter->AssociatedAnimAssetPlayerNode)
+			{
+				InCompilationContext.GetMessageLog().Error(TEXT("Please replace @@ with Get Relevant Anim Time Fraction. @@ has no animation asset"), Getter, Getter->AssociatedAnimAssetPlayerNode);
+			}
+			else
+			{
+				InCompilationContext.GetMessageLog().Error(TEXT("@@ is not asscociated with an asset player"), Getter);
+			}
+		}
+		break;
+	case ETransitionGetter::AnimationAsset_GetTimeFromEnd:
+		if ((AnimAsset != NULL) && (SourceTimePin != NULL))
+		{
+			GetterHelper = SpawnCallAnimInstanceFunction(InCompilationContext, Getter, TEXT("GetInstanceAssetPlayerTimeFromEnd"));
+			GetterHelper->FindPinChecked(TEXT("AssetPlayerIndex"))->DefaultValue = FString::FromInt(PlayerNodeIndex);
+		}
+		else
+		{
+			if (Getter->AssociatedAnimAssetPlayerNode)
+			{
+				InCompilationContext.GetMessageLog().Error(TEXT("Please replace @@ with Get Relevant Anim Time Remaining. @@ has no animation asset"), Getter, Getter->AssociatedAnimAssetPlayerNode);
+			}
+			else
+			{
+				InCompilationContext.GetMessageLog().Error(TEXT("@@ is not asscociated with an asset player"), Getter);
+			}
+		}
+		break;
+	case ETransitionGetter::AnimationAsset_GetTimeFromEndFraction:
+		if ((AnimAsset != NULL) && (SourceTimePin != NULL))
+		{
+			GetterHelper = SpawnCallAnimInstanceFunction(InCompilationContext, Getter, TEXT("GetInstanceAssetPlayerTimeFromEndFraction"));
+			GetterHelper->FindPinChecked(TEXT("AssetPlayerIndex"))->DefaultValue = FString::FromInt(PlayerNodeIndex);
+		}
+		else
+		{
+			if (Getter->AssociatedAnimAssetPlayerNode)
+			{
+				InCompilationContext.GetMessageLog().Error(TEXT("Please replace @@ with Get Relevant Anim Time Remaining Fraction. @@ has no animation asset"), Getter, Getter->AssociatedAnimAssetPlayerNode);
+			}
+			else
+			{
+				InCompilationContext.GetMessageLog().Error(TEXT("@@ is not asscociated with an asset player"), Getter);
+			}
+		}
+		break;
+
+	case ETransitionGetter::CurrentTransitionDuration:
+		{
+			check(TransitionNode);
+			if(UAnimStateNode* SourceStateNode = InCompilationContext.GetMessageLog().FindSourceObjectTypeChecked<UAnimStateNode>(TransitionNode->GetPreviousState()))
+			{
+				if(UObject* SourceTransitionNode = InCompilationContext.GetMessageLog().FindSourceObject(TransitionNode))
+				{
+					if(FStateMachineDebugData* DebugData = OutCompiledData.GetAnimBlueprintDebugData().StateMachineDebugData.Find(SourceStateNode->GetGraph()))
+					{
+						if(int32* pStateIndex = DebugData->NodeToStateIndex.Find(SourceStateNode))
+						{
+							const int32 StateIndex = *pStateIndex;
+							
+							// This check should never fail as all animation nodes should be processed before getters are
+							UAnimGraphNode_Base* CompiledMachineInstanceNode = InCompilationContext.GetSourceNodeToProcessedNodeMap().FindChecked(DebugData->MachineInstanceNode.Get());
+							const int32 MachinePropertyIndex = InCompilationContext.GetAllocatedAnimNodeIndices().FindChecked(CompiledMachineInstanceNode);
+							int32 TransitionPropertyIndex = INDEX_NONE;
+
+							for(TMap<TWeakObjectPtr<UEdGraphNode>, int32>::TIterator TransIt(DebugData->NodeToTransitionIndex); TransIt; ++TransIt)
+							{
+								UEdGraphNode* CurrTransNode = TransIt.Key().Get();
+								
+								if(CurrTransNode == SourceTransitionNode)
+								{
+									TransitionPropertyIndex = TransIt.Value();
+									break;
+								}
+							}
+
+							if(TransitionPropertyIndex != INDEX_NONE)
+							{
+								GetterHelper = SpawnCallAnimInstanceFunction(InCompilationContext, Getter, TEXT("GetInstanceTransitionCrossfadeDuration"));
+								GetterHelper->FindPinChecked(TEXT("MachineIndex"))->DefaultValue = FString::FromInt(MachinePropertyIndex);
+								GetterHelper->FindPinChecked(TEXT("TransitionIndex"))->DefaultValue = FString::FromInt(TransitionPropertyIndex);
+							}
+						}
+					}
+				}
+			}
+		}
+		break;
+
+	case ETransitionGetter::ArbitraryState_GetBlendWeight:
+		{
+			if (Getter->AssociatedStateNode)
+			{
+				if (UAnimStateNode* SourceStateNode = InCompilationContext.GetMessageLog().FindSourceObjectTypeChecked<UAnimStateNode>(Getter->AssociatedStateNode))
+				{
+					if (FStateMachineDebugData* DebugData = OutCompiledData.GetAnimBlueprintDebugData().StateMachineDebugData.Find(SourceStateNode->GetGraph()))
+					{
+						if (int32* pStateIndex = DebugData->NodeToStateIndex.Find(SourceStateNode))
+						{
+							const int32 StateIndex = *pStateIndex;
+							//const int32 MachineIndex = DebugData->MachineIndex;
+
+							// This check should never fail as all animation nodes should be processed before getters are
+							UAnimGraphNode_Base* CompiledMachineInstanceNode = InCompilationContext.GetSourceNodeToProcessedNodeMap().FindChecked(DebugData->MachineInstanceNode.Get());
+							const int32 MachinePropertyIndex = InCompilationContext.GetAllocatedAnimNodeIndices().FindChecked(CompiledMachineInstanceNode);
+
+							GetterHelper = SpawnCallAnimInstanceFunction(InCompilationContext, Getter, TEXT("GetInstanceStateWeight"));
+							GetterHelper->FindPinChecked(TEXT("MachineIndex"))->DefaultValue = FString::FromInt(MachinePropertyIndex);
+							GetterHelper->FindPinChecked(TEXT("StateIndex"))->DefaultValue = FString::FromInt(StateIndex);
+						}
+					}
+				}
+			}
+
+			if (GetterHelper == NULL)
+			{
+				InCompilationContext.GetMessageLog().Error(TEXT("@@ is not associated with a valid state"), Getter);
+			}
+		}
+		break;
+
+	case ETransitionGetter::CurrentState_ElapsedTime:
+		{
+			check(TransitionNode);
+			if (UAnimStateNode* SourceStateNode = InCompilationContext.GetMessageLog().FindSourceObjectTypeChecked<UAnimStateNode>(TransitionNode->GetPreviousState()))
+			{
+				if (FStateMachineDebugData* DebugData = OutCompiledData.GetAnimBlueprintDebugData().StateMachineDebugData.Find(SourceStateNode->GetGraph()))
+				{
+					// This check should never fail as all animation nodes should be processed before getters are
+					UAnimGraphNode_Base* CompiledMachineInstanceNode = InCompilationContext.GetSourceNodeToProcessedNodeMap().FindChecked(DebugData->MachineInstanceNode.Get());
+					const int32 MachinePropertyIndex = InCompilationContext.GetAllocatedAnimNodeIndices().FindChecked(CompiledMachineInstanceNode);
+
+					GetterHelper = SpawnCallAnimInstanceFunction(InCompilationContext, Getter, TEXT("GetInstanceCurrentStateElapsedTime"));
+					GetterHelper->FindPinChecked(TEXT("MachineIndex"))->DefaultValue = FString::FromInt(MachinePropertyIndex);
+				}
+			}
+			if (GetterHelper == NULL)
+			{
+				InCompilationContext.GetMessageLog().Error(TEXT("@@ is not associated with a valid state"), Getter);
+			}
+		}
+		break;
+
+	case ETransitionGetter::CurrentState_GetBlendWeight:
+		{
+			check(TransitionNode);
+			if (UAnimStateNode* SourceStateNode = InCompilationContext.GetMessageLog().FindSourceObjectTypeChecked<UAnimStateNode>(TransitionNode->GetPreviousState()))
+			{
+				{
+					if (FStateMachineDebugData* DebugData = OutCompiledData.GetAnimBlueprintDebugData().StateMachineDebugData.Find(SourceStateNode->GetGraph()))
+					{
+						if (int32* pStateIndex = DebugData->NodeToStateIndex.Find(SourceStateNode))
+						{
+							const int32 StateIndex = *pStateIndex;
+							//const int32 MachineIndex = DebugData->MachineIndex;
+
+							// This check should never fail as all animation nodes should be processed before getters are
+							UAnimGraphNode_Base* CompiledMachineInstanceNode = InCompilationContext.GetSourceNodeToProcessedNodeMap().FindChecked(DebugData->MachineInstanceNode.Get());
+							const int32 MachinePropertyIndex = InCompilationContext.GetAllocatedAnimNodeIndices().FindChecked(CompiledMachineInstanceNode);
+
+							GetterHelper = SpawnCallAnimInstanceFunction(InCompilationContext, Getter, TEXT("GetInstanceStateWeight"));
+							GetterHelper->FindPinChecked(TEXT("MachineIndex"))->DefaultValue = FString::FromInt(MachinePropertyIndex);
+							GetterHelper->FindPinChecked(TEXT("StateIndex"))->DefaultValue = FString::FromInt(StateIndex);
+						}
+					}
+				}
+			}
+			if (GetterHelper == NULL)
+			{
+				InCompilationContext.GetMessageLog().Error(TEXT("@@ is not associated with a valid state"), Getter);
+			}
+		}
+		break;
+
+	default:
+		InCompilationContext.GetMessageLog().Error(TEXT("Unrecognized getter type on @@"), Getter);
+		break;
+	}
+
+	// Finish wiring up a call function if needed
+	if (GetterHelper != NULL)
+	{
+		check(GetterHelper->IsNodePure());
+
+		UEdGraphPin* NewReturnPin = GetterHelper->FindPinChecked(TEXT("ReturnValue"));
+		InCompilationContext.GetMessageLog().NotifyIntermediatePinCreation(NewReturnPin, OutputPin);
+
+		NewReturnPin->CopyPersistentDataFromOldPin(*OutputPin);
+	}
+
+	// Remove the getter from the equation
+	Getter->BreakAllNodeLinks();
+}
+
+void FAnimBlueprintCompilerHandler_StateMachine::AutoWireAnimGetter(class UK2Node_AnimGetter* Getter, UAnimStateTransitionNode* InTransitionNode, IAnimBlueprintCompilationContext& InCompilationContext, IAnimBlueprintGeneratedClassCompiledData& OutCompiledData)
+{
+	UEdGraphPin* ReferencedNodeTimePin = nullptr;
+	int32 ReferencedNodeIndex = INDEX_NONE;
+	int32 SubNodeIndex = INDEX_NONE;
+	
+	UAnimGraphNode_Base* ProcessedNodeCheck = NULL;
+
+	if(UAnimGraphNode_Base* SourceNode = Getter->SourceNode)
+	{
+		UAnimGraphNode_Base* ActualSourceNode = InCompilationContext.GetMessageLog().FindSourceObjectTypeChecked<UAnimGraphNode_Base>(SourceNode);
+		
+		if(UAnimGraphNode_Base* ProcessedSourceNode = InCompilationContext.GetSourceNodeToProcessedNodeMap().FindRef(ActualSourceNode))
+		{
+			ProcessedNodeCheck = ProcessedSourceNode;
+
+			ReferencedNodeIndex = InCompilationContext.GetAllocationIndexOfNode(ProcessedSourceNode);
+
+			if(ProcessedSourceNode->DoesSupportTimeForTransitionGetter())
+			{
+				UScriptStruct* TimePropertyInStructType = ProcessedSourceNode->GetTimePropertyStruct();
+				const TCHAR* TimePropertyName = ProcessedSourceNode->GetTimePropertyName();
+
+				if(ReferencedNodeIndex != INDEX_NONE && TimePropertyName && TimePropertyInStructType)
+				{
+					FProperty* NodeProperty = InCompilationContext.GetAllocatedPropertiesByIndex().FindChecked(ReferencedNodeIndex);
+
+					UK2Node_StructMemberGet* ReaderNode = InCompilationContext.SpawnIntermediateNode<UK2Node_StructMemberGet>(Getter, InCompilationContext.GetConsolidatedEventGraph());
+					ReaderNode->VariableReference.SetSelfMember(NodeProperty->GetFName());
+					ReaderNode->StructType = TimePropertyInStructType;
+					ReaderNode->AllocatePinsForSingleMemberGet(TimePropertyName);
+
+					ReferencedNodeTimePin = ReaderNode->FindPinChecked(TimePropertyName);
+				}
+			}
+		}
+	}
+	
+	if(Getter->SourceStateNode)
+	{
+		UObject* SourceObject = InCompilationContext.GetMessageLog().FindSourceObject(Getter->SourceStateNode);
+		if(UAnimStateNode* SourceStateNode = Cast<UAnimStateNode>(SourceObject))
+		{
+			if(FStateMachineDebugData* DebugData = OutCompiledData.GetAnimBlueprintDebugData().StateMachineDebugData.Find(SourceStateNode->GetGraph()))
+			{
+				if(int32* StateIndexPtr = DebugData->NodeToStateIndex.Find(SourceStateNode))
+				{
+					SubNodeIndex = *StateIndexPtr;
+				}
+			}
+		}
+		else if(UAnimStateTransitionNode* TransitionNode = Cast<UAnimStateTransitionNode>(SourceObject))
+		{
+			if(FStateMachineDebugData* DebugData = OutCompiledData.GetAnimBlueprintDebugData().StateMachineDebugData.Find(TransitionNode->GetGraph()))
+			{
+				if(int32* TransitionIndexPtr = DebugData->NodeToTransitionIndex.Find(TransitionNode))
+				{
+					SubNodeIndex = *TransitionIndexPtr;
+				}
+			}
+		}
+	}
+
+	check(Getter->IsNodePure());
+
+	for(UEdGraphPin* Pin : Getter->Pins)
+	{
+		// Hook up autowired parameters / pins
+		if(Pin->PinName == TEXT("CurrentTime"))
+		{
+			Pin->MakeLinkTo(ReferencedNodeTimePin);
+		}
+		else if(Pin->PinName == TEXT("AssetPlayerIndex") || Pin->PinName == TEXT("MachineIndex"))
+		{
+			Pin->DefaultValue = FString::FromInt(ReferencedNodeIndex);
+		}
+		else if(Pin->PinName == TEXT("StateIndex") || Pin->PinName == TEXT("TransitionIndex"))
+		{
+			Pin->DefaultValue = FString::FromInt(SubNodeIndex);
+		}
+	}
+}
+
+int32 FAnimBlueprintCompilerHandler_StateMachine::ExpandGraphAndProcessNodes(UEdGraph* SourceGraph, UAnimGraphNode_Base* SourceRootNode, IAnimBlueprintCompilationContext& InCompilationContext, IAnimBlueprintGeneratedClassCompiledData& OutCompiledData, UAnimStateTransitionNode* TransitionNode, TArray<UEdGraphNode*>* ClonedNodes)
+{
+	// Clone the nodes from the source graph
+	// Note that we outer this graph to the ConsolidatedEventGraph to allow ExpansionStep to 
+	// correctly retrieve the context for any expanded function calls (custom make/break structs etc.)
+	UEdGraph* ClonedGraph = FEdGraphUtilities::CloneGraph(SourceGraph, InCompilationContext.GetConsolidatedEventGraph(), &InCompilationContext.GetMessageLog(), true);
+
+	// Grab all the animation nodes and find the corresponding root node in the cloned set
+	UAnimGraphNode_Base* TargetRootNode = nullptr;
+	TArray<UAnimGraphNode_Base*> AnimNodeList;
+	TArray<UK2Node_TransitionRuleGetter*> Getters;
+	TArray<UK2Node_AnimGetter*> AnimGetterNodes;
+
+	for (auto NodeIt = ClonedGraph->Nodes.CreateIterator(); NodeIt; ++NodeIt)
+	{
+		UEdGraphNode* Node = *NodeIt;
+
+		if (UK2Node_TransitionRuleGetter* GetterNode = Cast<UK2Node_TransitionRuleGetter>(Node))
+		{
+			Getters.Add(GetterNode);
+		}
+		else if(UK2Node_AnimGetter* NewGetterNode = Cast<UK2Node_AnimGetter>(Node))
+		{
+			AnimGetterNodes.Add(NewGetterNode);
+		}
+		else if (UAnimGraphNode_Base* TestNode = Cast<UAnimGraphNode_Base>(Node))
+		{
+			AnimNodeList.Add(TestNode);
+
+			//@TODO: There ought to be a better way to determine this
+			if (InCompilationContext.GetMessageLog().FindSourceObject(TestNode) == InCompilationContext.GetMessageLog().FindSourceObject(SourceRootNode))
+			{
+				TargetRootNode = TestNode;
+			}
+		}
+
+		if (ClonedNodes != NULL)
+		{
+			ClonedNodes->Add(Node);
+		}
+	}
+	check(TargetRootNode);
+
+	// Run another expansion pass to catch the graph we just added (this is slightly wasteful 
+	InCompilationContext.ExpansionStep(ClonedGraph, false);
+
+	// Validate graph now we have expanded/pruned
+	InCompilationContext.ValidateGraphIsWellFormed(ClonedGraph);
+
+	// Move the cloned nodes into the consolidated event graph
+	const bool bIsLoading = InCompilationContext.GetBlueprint()->bIsRegeneratingOnLoad || IsAsyncLoading();
+	const bool bIsCompiling = InCompilationContext.GetBlueprint()->bBeingCompiled;
+	ClonedGraph->MoveNodesToAnotherGraph(InCompilationContext.GetConsolidatedEventGraph(), bIsLoading, bIsCompiling);
+
+	// Process any animation nodes
+	{
+		TArray<UAnimGraphNode_Base*> RootSet;
+		RootSet.Add(TargetRootNode);
+
+		InCompilationContext.PruneIsolatedAnimationNodes(RootSet, AnimNodeList);
+
+		InCompilationContext.ProcessAnimationNodes(AnimNodeList);
+	}
+
+	// Process the getter nodes in the graph if there were any
+	for (auto GetterIt = Getters.CreateIterator(); GetterIt; ++GetterIt)
+	{
+		ProcessTransitionGetter(*GetterIt, TransitionNode, InCompilationContext, OutCompiledData);
+	}
+
+	// Wire anim getter nodes
+	for(UK2Node_AnimGetter* GetterNode : AnimGetterNodes)
+	{
+		FoundGetterNodes.Add(GetterNode);
+	}
+
+	// Returns the index of the processed cloned version of SourceRootNode
+	return InCompilationContext.GetAllocationIndexOfNode(TargetRootNode);	
+}
+
+#undef LOCTEXT_NAMESPACE
\ No newline at end of file
diff --git a/Engine/Source/Editor/AnimGraph/Private/AnimBlueprintCompilerHandler_StateMachine.h b/Engine/Source/Editor/AnimGraph/Private/AnimBlueprintCompilerHandler_StateMachine.h
new file mode 100644
--- /dev/null
+++ b/Engine/Source/Editor/AnimGraph/Private/AnimBlueprintCompilerHandler_StateMachine.h
@@ -0,0 +1,53 @@
+// Copyright Epic Games, Inc. All Rights Reserved.
+
+#pragma once
+
+#include "CoreMinimal.h"
+#include "IAnimBlueprintCompilerHandler.h"
+
+class UK2Node_CallFunction;
+class UK2Node_AnimGetter;
+class UK2Node_TransitionRuleGetter;
+class UAnimGraphNode_Base;
+class UAnimStateTransitionNode;
+class UEdGraph;
+class UEdGraphNode;
+struct FAnimNotifyEvent;
+class IAnimBlueprintCompilerCreationContext;
+class IAnimBlueprintGeneratedClassCompiledData;
+class IAnimBlueprintCompilationContext;
+
+class FAnimBlueprintCompilerHandler_StateMachine : public IAnimBlueprintCompilerHandler
+{
+public:
+	FAnimBlueprintCompilerHandler_StateMachine(IAnimBlueprintCompilerCreationContext& InCreationContext);
+
+	// This function does the following steps:
+	//   Clones the nodes in the specified source graph
+	//   Merges them into the ConsolidatedEventGraph
+	//   Processes any animation nodes
+	//   Returns the index of the processed cloned version of SourceRootNode
+	//	 If supplied, will also return an array of all cloned nodes
+	int32 ExpandGraphAndProcessNodes(UEdGraph* SourceGraph, UAnimGraphNode_Base* SourceRootNode, IAnimBlueprintCompilationContext& InCompilationContext, IAnimBlueprintGeneratedClassCompiledData& OutCompiledData, UAnimStateTransitionNode* TransitionNode = nullptr, TArray<UEdGraphNode*>* ClonedNodes = nullptr);
+
+private:
+	void PreProcessAnimationNodes(TArrayView<UAnimGraphNode_Base*> InAnimNodes, IAnimBlueprintCompilationContext& InCompilationContext, IAnimBlueprintGeneratedClassCompiledData& OutCompiledData);
+	void PostProcessAnimationNodes(TArrayView<UAnimGraphNode_Base*> InAnimNodes, IAnimBlueprintCompilationContext& InCompilationContext, IAnimBlueprintGeneratedClassCompiledData& OutCompiledData);
+
+	// Spawns a function call node, calling a function on the anim instance
+	UK2Node_CallFunction* SpawnCallAnimInstanceFunction(IAnimBlueprintCompilationContext& InCompilationContext, UEdGraphNode* SourceNode, FName FunctionName);
+
+	// Convert transition getters into a function call/etc...
+	void ProcessTransitionGetter(UK2Node_TransitionRuleGetter* Getter, UAnimStateTransitionNode* TransitionNode, IAnimBlueprintCompilationContext& InCompilationContext, IAnimBlueprintGeneratedClassCompiledData& OutCompiledData);
+
+	// Automatically fill in parameters for the specified Getter node
+	void AutoWireAnimGetter(UK2Node_AnimGetter* Getter, UAnimStateTransitionNode* InTransitionNode, IAnimBlueprintCompilationContext& InCompilationContext, IAnimBlueprintGeneratedClassCompiledData& OutCompiledData);
+	
+private:
+	// List of getter node's we've found so the auto-wire can be deferred till after state machine compilation
+	TArray<UK2Node_AnimGetter*> FoundGetterNodes;
+
+	// Preprocessed lists of getters from the root of the ubergraph
+	TArray<UK2Node_TransitionRuleGetter*> RootTransitionGetters;
+	TArray<UK2Node_AnimGetter*> RootGraphAnimGetters;
+};
\ No newline at end of file
diff --git a/Engine/Source/Editor/AnimGraph/Private/AnimGraphModule.cpp b/Engine/Source/Editor/AnimGraph/Private/AnimGraphModule.cpp
--- a/Engine/Source/Editor/AnimGraph/Private/AnimGraphModule.cpp
+++ b/Engine/Source/Editor/AnimGraph/Private/AnimGraphModule.cpp
@@ -24,6 +24,10 @@
 #include "AnimGraphDetails.h"
 #include "AnimationGraphSchema.h"
 #include "AnimBlueprintCompiler.h"
+#include "AnimBlueprintCompilerHandler_Base.h"
+#include "AnimBlueprintCompilerHandler_CachedPose.h"
+#include "AnimBlueprintCompilerHandler_LinkedAnimGraph.h"
+#include "AnimBlueprintCompilerHandler_StateMachine.h"
 
 IMPLEMENT_MODULE(FAnimGraphModule, AnimGraph);
 
@@ -38,6 +42,27 @@ void FAnimGraphModule::StartupModule()
 		return MakeShared<FAnimBlueprintCompilerContext>(CastChecked<UAnimBlueprint>(InBlueprint), InMessageLog, InCompileOptions);
 	});
 
+	// Register node compilation handlers
+	IAnimBlueprintCompilerHandlerCollection::RegisterHandler("AnimBlueprintCompilerHandler_Base", [](IAnimBlueprintCompilerCreationContext& InCreationContext)
+	{
+		return MakeUnique<FAnimBlueprintCompilerHandler_Base>(InCreationContext);
+	});
+
+	IAnimBlueprintCompilerHandlerCollection::RegisterHandler("AnimBlueprintCompilerHandler_CachedPose", [](IAnimBlueprintCompilerCreationContext& InCreationContext)
+	{
+		return MakeUnique<FAnimBlueprintCompilerHandler_CachedPose>(InCreationContext);
+	});
+
+	IAnimBlueprintCompilerHandlerCollection::RegisterHandler("AnimBlueprintCompilerHandler_LinkedAnimGraph", [](IAnimBlueprintCompilerCreationContext& InCreationContext)
+	{
+		return MakeUnique<FAnimBlueprintCompilerHandler_LinkedAnimGraph>(InCreationContext);
+	});
+
+	IAnimBlueprintCompilerHandlerCollection::RegisterHandler("AnimBlueprintCompilerHandler_StateMachine", [](IAnimBlueprintCompilerCreationContext& InCreationContext)
+	{
+		return MakeUnique<FAnimBlueprintCompilerHandler_StateMachine>(InCreationContext);
+	});
+
 	// Register the editor modes
 	FEditorModeRegistry::Get().RegisterMode<FAnimNodeEditMode>(AnimNodeEditModes::AnimNode, LOCTEXT("AnimNodeEditMode", "Anim Node"), FSlateIcon(), false);
 	FEditorModeRegistry::Get().RegisterMode<FTwoBoneIKEditMode>(AnimNodeEditModes::TwoBoneIK, LOCTEXT("TwoBoneIKEditMode", "2-Bone IK"), FSlateIcon(), false);
@@ -68,6 +93,11 @@ void FAnimGraphModule::StartupModule()
 
 void FAnimGraphModule::ShutdownModule()
 {
+	IAnimBlueprintCompilerHandlerCollection::UnregisterHandler("AnimBlueprintCompilerHandler_Base");
+	IAnimBlueprintCompilerHandlerCollection::UnregisterHandler("AnimBlueprintCompilerHandler_CachedPose");
+	IAnimBlueprintCompilerHandlerCollection::UnregisterHandler("AnimBlueprintCompilerHandler_LinkedAnimGraph");
+	IAnimBlueprintCompilerHandlerCollection::UnregisterHandler("AnimBlueprintCompilerHandler_StateMachine");
+
 	// Unregister the editor modes
 	FEditorModeRegistry::Get().UnregisterMode(AnimNodeEditModes::CCDIK);
 	FEditorModeRegistry::Get().UnregisterMode(AnimNodeEditModes::SplineIK);
