Restricted anim BP compiler subsystem APIs

diff --git a/Engine/Source/Editor/AnimGraph/Private/AnimGraphNode_AssetPlayerBase.cpp b/Engine/Source/Editor/AnimGraph/Private/AnimGraphNode_AssetPlayerBase.cpp
--- a/Engine/Source/Editor/AnimGraph/Private/AnimGraphNode_AssetPlayerBase.cpp
+++ b/Engine/Source/Editor/AnimGraph/Private/AnimGraphNode_AssetPlayerBase.cpp
@@ -16,8 +16,11 @@
 #include "AnimGraphNode_PoseByName.h"
 #include "AnimGraphNode_PoseDriver.h"
 #include "UObject/UObjectIterator.h"
-#include "AnimBlueprintCompiler.h"
 #include "Animation/AnimLayerInterface.h"
+#include "IAnimBlueprintCompilerHandlerCollection.h"
+#include "AnimBlueprintCompilerHandler_Base.h"
+#include "IAnimBlueprintGeneratedClassCompiledData.h"
+#include "IAnimBlueprintCompilationContext.h"
 
 void UAnimGraphNode_AssetPlayerBase::PinConnectionListChanged(UEdGraphPin* Pin)
 {
@@ -47,13 +50,12 @@ void UAnimGraphNode_AssetPlayerBase::PinDefaultValueChanged(UEdGraphPin* Pin)
 	}
 }
 
-void UAnimGraphNode_AssetPlayerBase::OnProcessDuringCompilation(FAnimBlueprintCompilerContext& InCompilerContext)
+void UAnimGraphNode_AssetPlayerBase::OnProcessDuringCompilation(IAnimBlueprintCompilationContext& InCompilationContext, IAnimBlueprintGeneratedClassCompiledData& OutCompiledData)
 {
 	UBlueprint* Blueprint = GetBlueprint();
-	UAnimBlueprintGeneratedClass* NewAnimBlueprintClass = CastChecked<UAnimBlueprintGeneratedClass>(InCompilerContext.NewClass);
 
 	// Process Asset Player nodes to, if necessary cache off their node index for retrieval at runtime (used for evaluating Automatic Rule Transitions when using Layer nodes)
-	auto ProcessGraph = [this, NewAnimBlueprintClass](UEdGraph* Graph)
+	auto ProcessGraph = [this, &OutCompiledData](UEdGraph* Graph)
 	{
 		// Make sure we do not process the default AnimGraph
 		static const FName DefaultAnimGraphName("AnimGraph");
@@ -63,11 +65,11 @@ void UAnimGraphNode_AssetPlayerBase::OnProcessDuringCompilation(FAnimBlueprintCo
 			// Also make sure we do not process any empty stub graphs
 			if (!GraphName.Contains(ANIM_FUNC_DECORATOR))
 			{
-				if (Graph->Nodes.ContainsByPredicate([this, NewAnimBlueprintClass](UEdGraphNode* Node) { return Node->NodeGuid == NodeGuid; }))
+				if (Graph->Nodes.ContainsByPredicate([this, &OutCompiledData](UEdGraphNode* Node) { return Node->NodeGuid == NodeGuid; }))
 				{
-					if (int32* IndexPtr = NewAnimBlueprintClass->AnimBlueprintDebugData.NodeGuidToIndexMap.Find(NodeGuid))
+					if (int32* IndexPtr = OutCompiledData.GetAnimBlueprintDebugData().NodeGuidToIndexMap.Find(NodeGuid))
 					{
-						FGraphAssetPlayerInformation& Info = NewAnimBlueprintClass->GraphAssetPlayerInformation.FindOrAdd(FName(*GraphName));
+						FGraphAssetPlayerInformation& Info = OutCompiledData.GetGraphAssetPlayerInformation().FindOrAdd(FName(*GraphName));
 						Info.PlayerNodeIndices.AddUnique(*IndexPtr);
 					}
 				}
diff --git a/Engine/Source/Editor/AnimGraph/Private/AnimGraphNode_Base.cpp b/Engine/Source/Editor/AnimGraph/Private/AnimGraphNode_Base.cpp
--- a/Engine/Source/Editor/AnimGraph/Private/AnimGraphNode_Base.cpp
+++ b/Engine/Source/Editor/AnimGraph/Private/AnimGraphNode_Base.cpp
@@ -18,7 +18,9 @@
 #include "Kismet2/CompilerResultsLog.h"
 #include "Subsystems/AssetEditorSubsystem.h"
 #include "AnimBlueprintCompiler.h"
-#include "AnimBlueprintCompilerSubsystem_Base.h"
+#include "AnimBlueprintCompilerHandler_Base.h"
+#include "IAnimBlueprintCompilationContext.h"
+#include "AnimBlueprintCompilationContext.h"
 #include "FindInBlueprintManager.h"
 
 #define LOCTEXT_NAMESPACE "UAnimGraphNode_Base"
@@ -33,8 +35,9 @@ UAnimGraphNode_Base::UAnimGraphNode_Base(const FObjectInitializer& ObjectInitial
 
 void UAnimGraphNode_Base::ExpandNode(FKismetCompilerContext& CompilerContext, UEdGraph* SourceGraph)
 {
-	UAnimBlueprintCompilerSubsystem_Base* Subsystem = static_cast<FAnimBlueprintCompilerContext*>(&CompilerContext)->GetSubsystem<UAnimBlueprintCompilerSubsystem_Base>();
-	Subsystem->CreateEvaluationHandlerForNode(this);
+	TUniquePtr<IAnimBlueprintCompilationContext> CompilationContext = IAnimBlueprintCompilationContext::Get(CompilerContext);
+	FAnimBlueprintCompilerHandler_Base* Handler = CompilationContext->GetHandler<FAnimBlueprintCompilerHandler_Base>("AnimBlueprintCompilerHandler_Base");
+	Handler->CreateEvaluationHandlerForNode(*CompilationContext.Get(), this);
 }
 
 void UAnimGraphNode_Base::PreEditChange(FProperty* PropertyThatWillChange)
@@ -338,15 +341,15 @@ void UAnimGraphNode_Base::GetPinHoverText(const UEdGraphPin& Pin, FString& Hover
 	}
 }
 
-void UAnimGraphNode_Base::ProcessDuringCompilation(FAnimBlueprintCompilerContext& InCompilerContext)
+void UAnimGraphNode_Base::ProcessDuringCompilation(IAnimBlueprintCompilationContext& InCompilationContext, IAnimBlueprintGeneratedClassCompiledData& OutCompiledData)
 {
-	UAnimBlueprintCompilerSubsystem_Base* SubsystemBase = InCompilerContext.GetSubsystem<UAnimBlueprintCompilerSubsystem_Base>();
+	FAnimBlueprintCompilerHandler_Base* HandlerBase = InCompilationContext.GetHandler<FAnimBlueprintCompilerHandler_Base>("AnimBlueprintCompilerHandler_Base");
 
 	// Record pose pins for later patchup and gather pins that have an associated evaluation handler
-	SubsystemBase->AddStructEvalHandlers(this);
+	HandlerBase->AddStructEvalHandlers(this, InCompilationContext, OutCompiledData);
 
 	// Call the override point
-	OnProcessDuringCompilation(InCompilerContext);
+	OnProcessDuringCompilation(InCompilationContext, OutCompiledData);
 }
 
 void UAnimGraphNode_Base::HandleAnimReferenceCollection(UAnimationAsset* AnimAsset, TArray<UAnimationAsset*>& AnimationAssets) const
diff --git a/Engine/Source/Editor/AnimGraph/Private/AnimGraphNode_CustomProperty.cpp b/Engine/Source/Editor/AnimGraph/Private/AnimGraphNode_CustomProperty.cpp
--- a/Engine/Source/Editor/AnimGraph/Private/AnimGraphNode_CustomProperty.cpp
+++ b/Engine/Source/Editor/AnimGraph/Private/AnimGraphNode_CustomProperty.cpp
@@ -15,11 +15,15 @@
 #include "DetailWidgetRow.h"
 #include "PropertyCustomizationHelpers.h"
 #include "KismetCompilerMisc.h"
+#include "AnimBlueprintCompilerHandler_Base.h"
+#include "IAnimBlueprintCompilerHandlerCollection.h"
+#include "IAnimBlueprintCompilerHandler.h"
 #include "KismetCompiler.h"
+#include "IAnimBlueprintCompilationContext.h"
 
 #define LOCTEXT_NAMESPACE "CustomPropNode"
 
-void UAnimGraphNode_CustomProperty::CreateClassVariablesFromBlueprint(FKismetCompilerContext& InCompilerContext)
+void UAnimGraphNode_CustomProperty::CreateClassVariablesFromBlueprint(IAnimBlueprintVariableCreationContext& InCreationContext)
 {
 	for (UEdGraphPin* Pin : Pins)
 	{
@@ -39,16 +43,12 @@ void UAnimGraphNode_CustomProperty::CreateClassVariablesFromBlueprint(FKismetCom
 			FString PrefixedName = GetPinTargetVariableName(Pin);
 
 			// Create a property on the new class to hold the pin data
-			FProperty* NewProperty = FKismetCompilerUtilities::CreatePropertyOnScope(InCompilerContext.NewClass, FName(*PrefixedName), Pin->PinType, InCompilerContext.NewClass, CPF_None, CastChecked<UEdGraphSchema_K2>(GetSchema()), InCompilerContext.MessageLog);
-			if (NewProperty)
-			{
-				FKismetCompilerUtilities::LinkAddedProperty(InCompilerContext.NewClass, NewProperty);
-			}
+			InCreationContext.CreateVariable(FName(*PrefixedName), Pin->PinType);
 		}
 	}
 }
 
-void UAnimGraphNode_CustomProperty::OnProcessDuringCompilation(FAnimBlueprintCompilerContext& InCompilerContext)
+void UAnimGraphNode_CustomProperty::OnProcessDuringCompilation(IAnimBlueprintCompilationContext& InCompilationContext, IAnimBlueprintGeneratedClassCompiledData& OutCompiledData)
 {
 	for (UEdGraphPin* Pin : Pins)
 	{
@@ -162,12 +162,12 @@ void UAnimGraphNode_CustomProperty::ReallocatePinsDuringReconstruction(TArray<UE
 	}
 }
 
-void UAnimGraphNode_CustomProperty::GetInstancePinProperty(const UClass* InOwnerInstanceClass, UEdGraphPin* InInputPin, FProperty*& OutProperty)
+void UAnimGraphNode_CustomProperty::GetInstancePinProperty(const IAnimBlueprintCompilationContext& InCompilationContext, UEdGraphPin* InInputPin, FProperty*& OutProperty)
 {
 	// The actual name of the instance property
 	FString FullName = GetPinTargetVariableName(InInputPin);
 
-	if(FProperty* Property = FindFProperty<FProperty>(InOwnerInstanceClass, *FullName))
+	if(FProperty* Property = InCompilationContext.FindClassFProperty<FProperty>(*FullName))
 	{
 		OutProperty = Property;
 	}
diff --git a/Engine/Source/Editor/AnimGraph/Private/AnimGraphNode_LinkedInputPose.cpp b/Engine/Source/Editor/AnimGraph/Private/AnimGraphNode_LinkedInputPose.cpp
--- a/Engine/Source/Editor/AnimGraph/Private/AnimGraphNode_LinkedInputPose.cpp
+++ b/Engine/Source/Editor/AnimGraph/Private/AnimGraphNode_LinkedInputPose.cpp
@@ -35,9 +35,9 @@ UAnimGraphNode_LinkedInputPose::UAnimGraphNode_LinkedInputPose()
 {
 }
 
-void UAnimGraphNode_LinkedInputPose::CreateClassVariablesFromBlueprint(FKismetCompilerContext& InCompilerContext)
+void UAnimGraphNode_LinkedInputPose::CreateClassVariablesFromBlueprint(IAnimBlueprintVariableCreationContext& InCreationContext)
 {
-	IterateFunctionParameters([this, &InCompilerContext](const FName& InName, FEdGraphPinType InPinType)
+	IterateFunctionParameters([this, &InCreationContext](const FName& InName, FEdGraphPinType InPinType)
 	{
 		if(!UAnimationGraphSchema::IsPosePin(InPinType))
 		{
@@ -45,7 +45,7 @@ void UAnimGraphNode_LinkedInputPose::CreateClassVariablesFromBlueprint(FKismetCo
 			if(Pin && Pin->LinkedTo.Num() > 0)
 			{
 				// create properties for 'local' linked input pose pins
-				FProperty* NewLinkedInputPoseProperty = static_cast<FAnimBlueprintCompilerContext*>(&InCompilerContext)->CreateVariable(InName, InPinType);
+				FProperty* NewLinkedInputPoseProperty = InCreationContext.CreateVariable(InName, InPinType);
 				if(NewLinkedInputPoseProperty)
 				{
 					NewLinkedInputPoseProperty->SetPropertyFlags(CPF_BlueprintReadOnly);
diff --git a/Engine/Source/Editor/AnimGraph/Private/AnimGraphNode_Root.cpp b/Engine/Source/Editor/AnimGraph/Private/AnimGraphNode_Root.cpp
--- a/Engine/Source/Editor/AnimGraph/Private/AnimGraphNode_Root.cpp
+++ b/Engine/Source/Editor/AnimGraph/Private/AnimGraphNode_Root.cpp
@@ -3,6 +3,8 @@
 #include "AnimGraphNode_Root.h"
 #include "GraphEditorSettings.h"
 #include "AnimBlueprintCompiler.h"
+#include "AnimBlueprintCompilerHandler_Base.h"
+#include "IAnimBlueprintCompilationContext.h"
 
 /////////////////////////////////////////////////////
 // FPoseLinkMappingRecord
@@ -85,9 +87,9 @@ FString UAnimGraphNode_Root::GetDocumentationLink() const
 	return TEXT("Shared/GraphNodes/AnimationStateMachine");
 }
 
-void UAnimGraphNode_Root::OnProcessDuringCompilation(FAnimBlueprintCompilerContext& InCompilerContext)
+void UAnimGraphNode_Root::OnProcessDuringCompilation(IAnimBlueprintCompilationContext& InCompilationContext, IAnimBlueprintGeneratedClassCompiledData& OutCompiledData)
 {
-	UAnimGraphNode_Root* TrueNode = InCompilerContext.MessageLog.FindSourceObjectTypeChecked<UAnimGraphNode_Root>(this);
+	UAnimGraphNode_Root* TrueNode = InCompilationContext.GetMessageLog().FindSourceObjectTypeChecked<UAnimGraphNode_Root>(this);
 
 	Node.Name = TrueNode->GetGraph()->GetFName();
 }
diff --git a/Engine/Source/Editor/AnimGraph/Private/AnimGraphNode_StateMachineBase.cpp b/Engine/Source/Editor/AnimGraph/Private/AnimGraphNode_StateMachineBase.cpp
--- a/Engine/Source/Editor/AnimGraph/Private/AnimGraphNode_StateMachineBase.cpp
+++ b/Engine/Source/Editor/AnimGraph/Private/AnimGraphNode_StateMachineBase.cpp
@@ -25,7 +25,9 @@
 #include "AnimGraphNode_LinkedAnimLayer.h"
 #include "AnimGraphNode_TransitionPoseEvaluator.h"
 #include "AnimGraphNode_CustomTransitionResult.h"
-#include "AnimBlueprintCompilerSubsystem_StateMachine.h"
+#include "AnimBlueprintCompilerHandler_StateMachine.h"
+#include "IAnimBlueprintGeneratedClassCompiledData.h"
+#include "IAnimBlueprintCompilationContext.h"
 
 /////////////////////////////////////////////////////
 // FAnimStateMachineNodeNameValidator
@@ -210,30 +212,30 @@ void UAnimGraphNode_StateMachineBase::OnRenameNode(const FString& NewName)
 	FBlueprintEditorUtils::RenameGraph(EditorStateMachineGraph, NewName);
 }
 
-void UAnimGraphNode_StateMachineBase::OnProcessDuringCompilation(FAnimBlueprintCompilerContext& InCompilerContext)
+void UAnimGraphNode_StateMachineBase::OnProcessDuringCompilation(IAnimBlueprintCompilationContext& InCompilationContext, IAnimBlueprintGeneratedClassCompiledData& OutCompiledData)
 {
-	UAnimBlueprintCompilerSubsystem_StateMachine* CompilerSubsystem = InCompilerContext.GetSubsystem<UAnimBlueprintCompilerSubsystem_StateMachine>();
-	check(CompilerSubsystem);
-
 	struct FMachineCreator
 	{
 	public:
 		int32 MachineIndex;
 		TMap<UAnimStateNodeBase*, int32> StateIndexTable;
 		TMap<UAnimStateTransitionNode*, int32> TransitionIndexTable;
-		UAnimBlueprintGeneratedClass* AnimBlueprintClass;
 		UAnimGraphNode_StateMachineBase* StateMachineInstance;
-		FCompilerResultsLog& MessageLog;
+		TArray<FBakedAnimationStateMachine>& BakedMachines;
+		IAnimBlueprintGeneratedClassCompiledData& CompiledData;
+		IAnimBlueprintCompilationContext& CompilationContext;
+
 	public:
-		FMachineCreator(FCompilerResultsLog& InMessageLog, UAnimGraphNode_StateMachineBase* InStateMachineInstance, int32 InMachineIndex, UAnimBlueprintGeneratedClass* InNewClass)
+		FMachineCreator(UAnimGraphNode_StateMachineBase* InStateMachineInstance, int32 InMachineIndex, TArray<FBakedAnimationStateMachine>& InBakedMachines, IAnimBlueprintCompilationContext& InCompilationContext, IAnimBlueprintGeneratedClassCompiledData& InCompiledData)
 			: MachineIndex(InMachineIndex)
-			, AnimBlueprintClass(InNewClass)
 			, StateMachineInstance(InStateMachineInstance)
-			, MessageLog(InMessageLog)
+			, BakedMachines(InBakedMachines)
+			, CompiledData(InCompiledData)
+			, CompilationContext(InCompilationContext)
 		{
 			FStateMachineDebugData& MachineInfo = GetMachineSpecificDebugData();
 			MachineInfo.MachineIndex = MachineIndex;
-			MachineInfo.MachineInstanceNode = MessageLog.FindSourceObjectTypeChecked<UAnimGraphNode_StateMachineBase>(InStateMachineInstance);
+			MachineInfo.MachineInstanceNode = CompilationContext.GetMessageLog().FindSourceObjectTypeChecked<UAnimGraphNode_StateMachineBase>(InStateMachineInstance);
 
 			StateMachineInstance->GetNode().StateMachineIndexInClass = MachineIndex;
 
@@ -244,13 +246,13 @@ void UAnimGraphNode_StateMachineBase::OnProcessDuringCompilation(FAnimBlueprintC
 
 		FBakedAnimationStateMachine& GetMachine()
 		{
-			return AnimBlueprintClass->BakedStateMachines[MachineIndex];
+			return BakedMachines[MachineIndex];
 		}
 
 		FStateMachineDebugData& GetMachineSpecificDebugData()
 		{
-			UAnimationStateMachineGraph* SourceGraph = MessageLog.FindSourceObjectTypeChecked<UAnimationStateMachineGraph>(StateMachineInstance->EditorStateMachineGraph);
-			return AnimBlueprintClass->GetAnimBlueprintDebugData().StateMachineDebugData.FindOrAdd(SourceGraph);
+			UAnimationStateMachineGraph* SourceGraph = CompilationContext.GetMessageLog().FindSourceObjectTypeChecked<UAnimationStateMachineGraph>(StateMachineInstance->EditorStateMachineGraph);
+			return CompiledData.GetAnimBlueprintDebugData().StateMachineDebugData.FindOrAdd(SourceGraph);
 		}
 
 		int32 FindOrAddState(UAnimStateNodeBase* StateNode)
@@ -267,11 +269,11 @@ void UAnimGraphNode_StateMachineBase::OnProcessDuringCompilation(FAnimBlueprintC
 				StateIndexTable.Add(StateNode, StateIndex);
 				new (BakedMachine.States) FBakedAnimationState();
 
-				UAnimStateNodeBase* SourceNode = MessageLog.FindSourceObjectTypeChecked<UAnimStateNodeBase>(StateNode);
+				UAnimStateNodeBase* SourceNode = CompilationContext.GetMessageLog().FindSourceObjectTypeChecked<UAnimStateNodeBase>(StateNode);
 				GetMachineSpecificDebugData().NodeToStateIndex.Add(SourceNode, StateIndex);
 				if (UAnimStateNode* SourceStateNode = Cast<UAnimStateNode>(SourceNode))
 				{
-					AnimBlueprintClass->GetAnimBlueprintDebugData().StateGraphToNodeMap.Add(SourceStateNode->BoundGraph, SourceStateNode);
+					CompiledData.GetAnimBlueprintDebugData().StateGraphToNodeMap.Add(SourceStateNode->BoundGraph, SourceStateNode);
 				}
 
 				return StateIndex;
@@ -292,13 +294,13 @@ void UAnimGraphNode_StateMachineBase::OnProcessDuringCompilation(FAnimBlueprintC
 				TransitionIndexTable.Add(TransitionNode, TransitionIndex);
 				new (BakedMachine.Transitions) FAnimationTransitionBetweenStates();
 
-				UAnimStateTransitionNode* SourceTransitionNode = MessageLog.FindSourceObjectTypeChecked<UAnimStateTransitionNode>(TransitionNode);
+				UAnimStateTransitionNode* SourceTransitionNode = CompilationContext.GetMessageLog().FindSourceObjectTypeChecked<UAnimStateTransitionNode>(TransitionNode);
 				GetMachineSpecificDebugData().NodeToTransitionIndex.Add(SourceTransitionNode, TransitionIndex);
-				AnimBlueprintClass->GetAnimBlueprintDebugData().TransitionGraphToNodeMap.Add(SourceTransitionNode->BoundGraph, SourceTransitionNode);
+				CompiledData.GetAnimBlueprintDebugData().TransitionGraphToNodeMap.Add(SourceTransitionNode->BoundGraph, SourceTransitionNode);
 
 				if (SourceTransitionNode->CustomTransitionGraph != NULL)
 				{
-					AnimBlueprintClass->GetAnimBlueprintDebugData().TransitionBlendGraphToNodeMap.Add(SourceTransitionNode->CustomTransitionGraph, SourceTransitionNode);
+					CompiledData.GetAnimBlueprintDebugData().TransitionBlendGraphToNodeMap.Add(SourceTransitionNode->CustomTransitionGraph, SourceTransitionNode);
 				}
 
 				return TransitionIndex;
@@ -312,7 +314,7 @@ void UAnimGraphNode_StateMachineBase::OnProcessDuringCompilation(FAnimBlueprintC
 			// Make sure there is a valid entry point
 			if (BakedMachine.InitialState == INDEX_NONE)
 			{
-				MessageLog.Warning(*LOCTEXT("NoEntryNode", "There was no entry state connection in @@").ToString(), StateMachineInstance);
+				CompilationContext.GetMessageLog().Warning(*LOCTEXT("NoEntryNode", "There was no entry state connection in @@").ToString(), StateMachineInstance);
 				BakedMachine.InitialState = 0;
 			}
 			else
@@ -321,23 +323,26 @@ void UAnimGraphNode_StateMachineBase::OnProcessDuringCompilation(FAnimBlueprintC
 				if (BakedMachine.States[BakedMachine.InitialState].bIsAConduit)
 				{
 					UEdGraphNode* StateNode = GetMachineSpecificDebugData().FindNodeFromStateIndex(BakedMachine.InitialState);
-					MessageLog.Error(*LOCTEXT("BadStateEntryNode", "A conduit (@@) cannot be used as the entry node for a state machine").ToString(), StateNode);
+					CompilationContext.GetMessageLog().Error(*LOCTEXT("BadStateEntryNode", "A conduit (@@) cannot be used as the entry node for a state machine").ToString(), StateNode);
 				}
 			}
 		}
 	};
 	
+	FAnimBlueprintCompilerHandler_StateMachine* CompilerHandler = InCompilationContext.GetHandler<FAnimBlueprintCompilerHandler_StateMachine>("AnimBlueprintCompilerHandler_StateMachine");
+	check(CompilerHandler);
+
 	if (EditorStateMachineGraph == NULL)
 	{
-		CompilerSubsystem->GetMessageLog().Error(*LOCTEXT("BadStateMachineNoGraph", "@@ does not have a corresponding graph").ToString(), this);
+		InCompilationContext.GetMessageLog().Error(*LOCTEXT("BadStateMachineNoGraph", "@@ does not have a corresponding graph").ToString(), this);
 		return;
 	}
 
 	TMap<UAnimGraphNode_TransitionResult*, int32> AlreadyMergedTransitionList;
 
-	const int32 MachineIndex = CompilerSubsystem->GetNewAnimBlueprintClass()->GetBakedStateMachines().Num();
-	new (CompilerSubsystem->GetNewAnimBlueprintClass()->BakedStateMachines) FBakedAnimationStateMachine();
-	FMachineCreator Oven(CompilerSubsystem->GetMessageLog(), this, MachineIndex, CompilerSubsystem->GetNewAnimBlueprintClass());
+	TArray<FBakedAnimationStateMachine>& BakedStateMachines = OutCompiledData.GetBakedStateMachines();
+	const int32 MachineIndex = BakedStateMachines.AddDefaulted();
+	FMachineCreator Oven(this, MachineIndex, BakedStateMachines, InCompilationContext, OutCompiledData);
 
 	// Map of states that contain a single player node (from state root node index to associated sequence player)
 	TMap<int32, UObject*> SimplePlayerStatesMap;
@@ -353,7 +358,7 @@ void UAnimGraphNode_StateMachineBase::OnProcessDuringCompilation(FAnimBlueprintC
 			FBakedAnimationStateMachine& BakedMachine = Oven.GetMachine();
 			if (BakedMachine.InitialState != INDEX_NONE)
 			{
-				CompilerSubsystem->GetMessageLog().Error(*LOCTEXT("TooManyStateMachineEntryNodes", "Found an extra entry node @@").ToString(), EntryNode);
+				InCompilationContext.GetMessageLog().Error(*LOCTEXT("TooManyStateMachineEntryNodes", "Found an extra entry node @@").ToString(), EntryNode);
 			}
 			else if (UAnimStateNodeBase* StartState = Cast<UAnimStateNodeBase>(EntryNode->GetOutputNode()))
 			{
@@ -361,20 +366,20 @@ void UAnimGraphNode_StateMachineBase::OnProcessDuringCompilation(FAnimBlueprintC
 			}
 			else
 			{
-				CompilerSubsystem->GetMessageLog().Warning(*LOCTEXT("NoConnection", "Entry node @@ is not connected to state").ToString(), EntryNode);
+				InCompilationContext.GetMessageLog().Warning(*LOCTEXT("NoConnection", "Entry node @@ is not connected to state").ToString(), EntryNode);
 			}
 		}
 		else if (UAnimStateTransitionNode* TransitionNode = Cast<UAnimStateTransitionNode>(Node))
 		{
-			TransitionNode->ValidateNodeDuringCompilation(CompilerSubsystem->GetMessageLog());
+			TransitionNode->ValidateNodeDuringCompilation(InCompilationContext.GetMessageLog());
 
 			const int32 TransitionIndex = Oven.FindOrAddTransition(TransitionNode);
 			FAnimationTransitionBetweenStates& BakedTransition = Oven.GetMachine().Transitions[TransitionIndex];
 
 			BakedTransition.CrossfadeDuration = TransitionNode->CrossfadeDuration;
-			BakedTransition.StartNotify = CompilerSubsystem->FindOrAddNotify(TransitionNode->TransitionStart);
-			BakedTransition.EndNotify = CompilerSubsystem->FindOrAddNotify(TransitionNode->TransitionEnd);
-			BakedTransition.InterruptNotify = CompilerSubsystem->FindOrAddNotify(TransitionNode->TransitionInterrupt);
+			BakedTransition.StartNotify = OutCompiledData.FindOrAddNotify(TransitionNode->TransitionStart);
+			BakedTransition.EndNotify = OutCompiledData.FindOrAddNotify(TransitionNode->TransitionEnd);
+			BakedTransition.InterruptNotify = OutCompiledData.FindOrAddNotify(TransitionNode->TransitionInterrupt);
 			BakedTransition.BlendMode = TransitionNode->BlendMode;
 			BakedTransition.CustomCurve = TransitionNode->CustomBlendCurve;
 			BakedTransition.BlendProfile = TransitionNode->BlendProfile;
@@ -390,7 +395,7 @@ void UAnimGraphNode_StateMachineBase::OnProcessDuringCompilation(FAnimBlueprintC
 
 				if (TransitionNode->Bidirectional)
 				{
-					CompilerSubsystem->GetMessageLog().Warning(*LOCTEXT("BidirectionalTransWarning", "Bidirectional transitions aren't supported yet @@").ToString(), TransitionNode);
+					InCompilationContext.GetMessageLog().Warning(*LOCTEXT("BidirectionalTransWarning", "Bidirectional transitions aren't supported yet @@").ToString(), TransitionNode);
 				}
 
 				BakedTransition.PreviousState = PreviousStateIndex;
@@ -398,14 +403,14 @@ void UAnimGraphNode_StateMachineBase::OnProcessDuringCompilation(FAnimBlueprintC
 			}
 			else
 			{
-				CompilerSubsystem->GetMessageLog().Warning(*LOCTEXT("BogusTransition", "@@ is incomplete, without a previous and next state").ToString(), TransitionNode);
+				InCompilationContext.GetMessageLog().Warning(*LOCTEXT("BogusTransition", "@@ is incomplete, without a previous and next state").ToString(), TransitionNode);
 				BakedTransition.PreviousState = INDEX_NONE;
 				BakedTransition.NextState = INDEX_NONE;
 			}
 		}
 		else if (UAnimStateNode* StateNode = Cast<UAnimStateNode>(Node))
 		{
-			StateNode->ValidateNodeDuringCompilation(CompilerSubsystem->GetMessageLog());
+			StateNode->ValidateNodeDuringCompilation(InCompilationContext.GetMessageLog());
 
 			const int32 StateIndex = Oven.FindOrAddState(StateNode);
 			FBakedAnimationState& BakedState = Oven.GetMachine().States[StateIndex];
@@ -413,18 +418,18 @@ void UAnimGraphNode_StateMachineBase::OnProcessDuringCompilation(FAnimBlueprintC
 			if (StateNode->BoundGraph != NULL)
 			{
 				BakedState.StateName = StateNode->BoundGraph->GetFName();
-				BakedState.StartNotify = CompilerSubsystem->FindOrAddNotify(StateNode->StateEntered);
-				BakedState.EndNotify = CompilerSubsystem->FindOrAddNotify(StateNode->StateLeft);
-				BakedState.FullyBlendedNotify = CompilerSubsystem->FindOrAddNotify(StateNode->StateFullyBlended);
+				BakedState.StartNotify = OutCompiledData.FindOrAddNotify(StateNode->StateEntered);
+				BakedState.EndNotify = OutCompiledData.FindOrAddNotify(StateNode->StateLeft);
+				BakedState.FullyBlendedNotify = OutCompiledData.FindOrAddNotify(StateNode->StateFullyBlended);
 				BakedState.bIsAConduit = false;
 				BakedState.bAlwaysResetOnEntry = StateNode->bAlwaysResetOnEntry;
 
 				// Process the inner graph of this state
 				if (UAnimGraphNode_StateResult* AnimGraphResultNode = CastChecked<UAnimationStateGraph>(StateNode->BoundGraph)->GetResultNode())
 				{
-					CompilerSubsystem->ValidateGraphIsWellFormed(StateNode->BoundGraph);
+					InCompilationContext.ValidateGraphIsWellFormed(StateNode->BoundGraph);
 
-					BakedState.StateRootNodeIndex = CompilerSubsystem->ExpandGraphAndProcessNodes(StateNode->BoundGraph, AnimGraphResultNode);
+					BakedState.StateRootNodeIndex = CompilerHandler->ExpandGraphAndProcessNodes(StateNode->BoundGraph, AnimGraphResultNode, InCompilationContext, OutCompiledData);
 
 					// See if the state consists of a single sequence player node, and remember the index if so
 					for (UEdGraphPin* TestPin : AnimGraphResultNode->Pins)
@@ -433,7 +438,7 @@ void UAnimGraphNode_StateMachineBase::OnProcessDuringCompilation(FAnimBlueprintC
 						{
 							if (UAnimGraphNode_SequencePlayer* SequencePlayer = Cast<UAnimGraphNode_SequencePlayer>(TestPin->LinkedTo[0]->GetOwningNode()))
 							{
-								SimplePlayerStatesMap.Add(BakedState.StateRootNodeIndex, CompilerSubsystem->GetMessageLog().FindSourceObject(SequencePlayer));
+								SimplePlayerStatesMap.Add(BakedState.StateRootNodeIndex, InCompilationContext.GetMessageLog().FindSourceObject(SequencePlayer));
 							}
 						}
 					}
@@ -441,13 +446,13 @@ void UAnimGraphNode_StateMachineBase::OnProcessDuringCompilation(FAnimBlueprintC
 				else
 				{
 					BakedState.StateRootNodeIndex = INDEX_NONE;
-					CompilerSubsystem->GetMessageLog().Error(*LOCTEXT("StateWithNoResult", "@@ has no result node").ToString(), StateNode);
+					InCompilationContext.GetMessageLog().Error(*LOCTEXT("StateWithNoResult", "@@ has no result node").ToString(), StateNode);
 				}
 			}
 			else
 			{
 				BakedState.StateName = NAME_None;
-				CompilerSubsystem->GetMessageLog().Error(*LOCTEXT("StateWithBadGraph", "@@ has no bound graph").ToString(), StateNode);
+				InCompilationContext.GetMessageLog().Error(*LOCTEXT("StateWithBadGraph", "@@ has no bound graph").ToString(), StateNode);
 			}
 
 			// If this check fires, then something in the machine has changed causing the states array to not
@@ -456,7 +461,7 @@ void UAnimGraphNode_StateMachineBase::OnProcessDuringCompilation(FAnimBlueprintC
 		}
 		else if (UAnimStateConduitNode* ConduitNode = Cast<UAnimStateConduitNode>(Node))
 		{
-			ConduitNode->ValidateNodeDuringCompilation(CompilerSubsystem->GetMessageLog());
+			ConduitNode->ValidateNodeDuringCompilation(InCompilationContext.GetMessageLog());
 
 			const int32 StateIndex = Oven.FindOrAddState(ConduitNode);
 			FBakedAnimationState& BakedState = Oven.GetMachine().States[StateIndex];
@@ -468,7 +473,7 @@ void UAnimGraphNode_StateMachineBase::OnProcessDuringCompilation(FAnimBlueprintC
 			{
 				if (UAnimGraphNode_TransitionResult* EntryRuleResultNode = CastChecked<UAnimationTransitionGraph>(ConduitNode->BoundGraph)->GetResultNode())
 				{
-					BakedState.EntryRuleNodeIndex = CompilerSubsystem->ExpandGraphAndProcessNodes(ConduitNode->BoundGraph, EntryRuleResultNode);
+					BakedState.EntryRuleNodeIndex = CompilerHandler->ExpandGraphAndProcessNodes(ConduitNode->BoundGraph, EntryRuleResultNode, InCompilationContext, OutCompiledData);
 				}
 			}
 
@@ -502,7 +507,7 @@ void UAnimGraphNode_StateMachineBase::OnProcessDuringCompilation(FAnimBlueprintC
 
 		for (UAnimGraphNode_AssetPlayerBase* Node : AssetPlayerNodes)
 		{
-			if (int32* IndexPtr = CompilerSubsystem->GetNewAnimBlueprintClass()->AnimBlueprintDebugData.NodeGuidToIndexMap.Find(Node->NodeGuid))
+			if (int32* IndexPtr = OutCompiledData.GetAnimBlueprintDebugData().NodeGuidToIndexMap.Find(Node->NodeGuid))
 			{
 				BakedState.PlayerNodeIndices.Add(*IndexPtr);
 			}
@@ -510,7 +515,7 @@ void UAnimGraphNode_StateMachineBase::OnProcessDuringCompilation(FAnimBlueprintC
 
 		for (UAnimGraphNode_LinkedAnimLayer* Node : LinkedAnimLayerNodes)
 		{
-			if (int32* IndexPtr = CompilerSubsystem->GetNewAnimBlueprintClass()->AnimBlueprintDebugData.NodeGuidToIndexMap.Find(Node->NodeGuid))
+			if (int32* IndexPtr = OutCompiledData.GetAnimBlueprintDebugData().NodeGuidToIndexMap.Find(Node->NodeGuid))
 			{
 				BakedState.LayerNodeIndices.Add(*IndexPtr);
 			}
@@ -540,14 +545,14 @@ void UAnimGraphNode_StateMachineBase::OnProcessDuringCompilation(FAnimBlueprintC
 				}
 				else
 				{
-					Rule.CanTakeDelegateIndex = CompilerSubsystem->ExpandGraphAndProcessNodes(TransitionNode->BoundGraph, TransitionResultNode, TransitionNode);
+					Rule.CanTakeDelegateIndex = CompilerHandler->ExpandGraphAndProcessNodes(TransitionNode->BoundGraph, TransitionResultNode, InCompilationContext, OutCompiledData, TransitionNode);
 					AlreadyMergedTransitionList.Add(TransitionResultNode, Rule.CanTakeDelegateIndex);
 				}
 			}
 			else
 			{
 				Rule.CanTakeDelegateIndex = INDEX_NONE;
-				CompilerSubsystem->GetMessageLog().Error(*LOCTEXT("TransitionWithNoResult", "@@ has no result node").ToString(), TransitionNode);
+				InCompilationContext.GetMessageLog().Error(*LOCTEXT("TransitionWithNoResult", "@@ has no result node").ToString(), TransitionNode);
 			}
 
 			// Handle automatic time remaining rules
@@ -560,7 +565,7 @@ void UAnimGraphNode_StateMachineBase::OnProcessDuringCompilation(FAnimBlueprintC
 				TArray<UEdGraphNode*> ClonedNodes;
 				if (CustomTransitionGraph->GetResultNode())
 				{
-					Rule.CustomResultNodeIndex = CompilerSubsystem->ExpandGraphAndProcessNodes(TransitionNode->CustomTransitionGraph, CustomTransitionGraph->GetResultNode(), NULL, &ClonedNodes);
+					Rule.CustomResultNodeIndex = CompilerHandler->ExpandGraphAndProcessNodes(TransitionNode->CustomTransitionGraph, CustomTransitionGraph->GetResultNode(), InCompilationContext, OutCompiledData, nullptr, &ClonedNodes);
 				}
 
 				// Find all the pose evaluators used in this transition, save handles to them because we need to populate some pose data before executing
@@ -579,7 +584,7 @@ void UAnimGraphNode_StateMachineBase::OnProcessDuringCompilation(FAnimBlueprintC
 				for (auto TransitionPoseListIt = TransitionPoseList.CreateIterator(); TransitionPoseListIt; ++TransitionPoseListIt)
 				{
 					UAnimGraphNode_TransitionPoseEvaluator* TransitionPoseNode = *TransitionPoseListIt;
-					Rule.PoseEvaluatorLinks.Add( CompilerSubsystem->GetAllocationIndexOfNode(TransitionPoseNode) );
+					Rule.PoseEvaluatorLinks.Add( InCompilationContext.GetAllocationIndexOfNode(TransitionPoseNode) );
 				}
 			}
 		}
diff --git a/Engine/Source/Editor/AnimGraph/Private/AnimGraphNode_StateResult.cpp b/Engine/Source/Editor/AnimGraph/Private/AnimGraphNode_StateResult.cpp
--- a/Engine/Source/Editor/AnimGraph/Private/AnimGraphNode_StateResult.cpp
+++ b/Engine/Source/Editor/AnimGraph/Private/AnimGraphNode_StateResult.cpp
@@ -3,6 +3,8 @@
 #include "AnimGraphNode_StateResult.h"
 #include "GraphEditorSettings.h"
 #include "AnimBlueprintCompiler.h"
+#include "AnimBlueprintCompilerHandler_Base.h"
+#include "IAnimBlueprintCompilationContext.h"
 
 #define LOCTEXT_NAMESPACE "A3Nodes"
 
@@ -44,9 +46,9 @@ FString UAnimGraphNode_StateResult::GetDocumentationLink() const
 	return TEXT("Shared/GraphNodes/AnimationStateMachine");
 }
 
-void UAnimGraphNode_StateResult::OnProcessDuringCompilation(FAnimBlueprintCompilerContext& InCompilerContext)
+void UAnimGraphNode_StateResult::OnProcessDuringCompilation(IAnimBlueprintCompilationContext& InCompilationContext, IAnimBlueprintGeneratedClassCompiledData& OutCompiledData)
 {
-	UAnimGraphNode_StateResult* TrueNode = InCompilerContext.MessageLog.FindSourceObjectTypeChecked<UAnimGraphNode_StateResult>(this);
+	UAnimGraphNode_StateResult* TrueNode = InCompilationContext.GetMessageLog().FindSourceObjectTypeChecked<UAnimGraphNode_StateResult>(this);
 
 	Node.Name = TrueNode->GetGraph()->GetFName();
 }
diff --git a/Engine/Source/Editor/AnimGraph/Private/AnimGraphNode_UseCachedPose.cpp b/Engine/Source/Editor/AnimGraph/Private/AnimGraphNode_UseCachedPose.cpp
--- a/Engine/Source/Editor/AnimGraph/Private/AnimGraphNode_UseCachedPose.cpp
+++ b/Engine/Source/Editor/AnimGraph/Private/AnimGraphNode_UseCachedPose.cpp
@@ -8,8 +8,9 @@
 #include "BlueprintActionFilter.h"
 #include "BlueprintNodeSpawner.h"
 #include "BlueprintActionDatabaseRegistrar.h"
-#include "AnimBlueprintCompiler.h"
-#include "AnimBlueprintCompilerSubsystem_CachedPose.h"
+#include "AnimBlueprintCompilerHandler_CachedPose.h"
+#include "IAnimBlueprintCompilerHandlerCollection.h"
+#include "IAnimBlueprintCompilationContext.h"
 
 /////////////////////////////////////////////////////
 // UAnimGraphNode_UseCachedPose
@@ -178,17 +179,17 @@ bool UAnimGraphNode_UseCachedPose::IsActionFilteredOut(class FBlueprintActionFil
 	return bIsFilteredOut;
 }
 
-void UAnimGraphNode_UseCachedPose::OnProcessDuringCompilation(FAnimBlueprintCompilerContext& InCompilerContext)
+void UAnimGraphNode_UseCachedPose::OnProcessDuringCompilation(IAnimBlueprintCompilationContext& InCompilationContext, IAnimBlueprintGeneratedClassCompiledData& OutCompiledData)
 {	
-	UAnimBlueprintCompilerSubsystem_CachedPose* CompilerSubsystem = InCompilerContext.GetSubsystem<UAnimBlueprintCompilerSubsystem_CachedPose>();
-	check(CompilerSubsystem);
+	FAnimBlueprintCompilerHandler_CachedPose* CompilerHandler = InCompilationContext.GetHandler<FAnimBlueprintCompilerHandler_CachedPose>("AnimBlueprintCompilerHandler_CachedPose");
+	check(CompilerHandler);
 
 	bool bSuccessful = false;
 
 	// Link to the saved cached pose
 	if(SaveCachedPoseNode.IsValid())
 	{
-		if (UAnimGraphNode_SaveCachedPose* AssociatedSaveNode = CompilerSubsystem->GetSaveCachedPoseNodes().FindRef(SaveCachedPoseNode->CacheName))
+		if (UAnimGraphNode_SaveCachedPose* AssociatedSaveNode = CompilerHandler->GetSaveCachedPoseNodes().FindRef(SaveCachedPoseNode->CacheName))
 		{
 			FStructProperty* LinkProperty = FindFProperty<FStructProperty>(FAnimNode_UseCachedPose::StaticStruct(), TEXT("LinkToCachingNode"));
 			check(LinkProperty);
@@ -196,7 +197,7 @@ void UAnimGraphNode_UseCachedPose::OnProcessDuringCompilation(FAnimBlueprintComp
 			FPoseLinkMappingRecord LinkRecord = FPoseLinkMappingRecord::MakeFromMember(this, AssociatedSaveNode, LinkProperty);
 			if (LinkRecord.IsValid())
 			{
-				CompilerSubsystem->AddPoseLinkMappingRecord(LinkRecord);
+				InCompilationContext.AddPoseLinkMappingRecord(LinkRecord);
 			}
 			bSuccessful = true;
 
@@ -209,7 +210,7 @@ void UAnimGraphNode_UseCachedPose::OnProcessDuringCompilation(FAnimBlueprintComp
 	
 	if(!bSuccessful)
 	{
-		CompilerSubsystem->GetMessageLog().Error(*LOCTEXT("NoAssociatedSaveNode", "@@ does not have an associated Save Cached Pose node").ToString(), this);
+		InCompilationContext.GetMessageLog().Error(*LOCTEXT("NoAssociatedSaveNode", "@@ does not have an associated Save Cached Pose node").ToString(), this);
 	}
 }
 
diff --git a/Engine/Source/Editor/AnimGraph/Public/IAnimBlueprintCompilationBracketContext.h b/Engine/Source/Editor/AnimGraph/Public/IAnimBlueprintCompilationBracketContext.h
new file mode 100644
--- /dev/null
+++ b/Engine/Source/Editor/AnimGraph/Public/IAnimBlueprintCompilationBracketContext.h
@@ -0,0 +1,33 @@
+// Copyright Epic Games, Inc. All Rights Reserved.
+
+#pragma once
+
+#include "CoreMinimal.h"
+
+class FCompilerResultsLog;
+class IAnimBlueprintCompilerHandler;
+
+// Interface passed to start/end compilation delegates
+class ANIMGRAPH_API IAnimBlueprintCompilationBracketContext
+{
+public:	
+	virtual ~IAnimBlueprintCompilationBracketContext() {}
+
+	// Get a handler of the specified type and name (i.e. via simple name-based RTTI)
+	// Handlers are registered via IAnimBlueprintCompilerHandlerCollection::RegisterHandler
+	template <typename THandlerClass>
+	THandlerClass* GetHandler(FName InName) const
+	{
+		return static_cast<THandlerClass*>(GetHandlerInternal(InName));
+	}
+
+	// Get the message log for the current compilation
+	FCompilerResultsLog& GetMessageLog() const { return GetMessageLogImpl(); }
+
+protected:	
+	// GetHandler helper function
+	virtual IAnimBlueprintCompilerHandler* GetHandlerInternal(FName InName) const = 0;
+
+	// Get the message log for the current compilation
+	virtual FCompilerResultsLog& GetMessageLogImpl() const = 0;
+};
diff --git a/Engine/Source/Editor/AnimGraph/Public/IAnimBlueprintCompilerCreationContext.h b/Engine/Source/Editor/AnimGraph/Public/IAnimBlueprintCompilerCreationContext.h
new file mode 100644
--- /dev/null
+++ b/Engine/Source/Editor/AnimGraph/Public/IAnimBlueprintCompilerCreationContext.h
@@ -0,0 +1,98 @@
+// Copyright Epic Games, Inc. All Rights Reserved.
+
+#pragma once
+
+#include "CoreMinimal.h"
+#include "Containers/ArrayView.h"
+
+class UObject;
+class UClass;
+class IAnimBlueprintCompilationContext;
+class IAnimBlueprintCompilationBracketContext;
+class IAnimBlueprintPostExpansionStepContext;
+class IAnimBlueprintCopyTermDefaultsContext;
+class IAnimBlueprintGeneratedClassCompiledData;
+class UAnimGraphNode_Base;
+class UEdGraph;
+class UEdGraphSchema;
+
+/** Begin ordered delegate calls - these functions are called in the order presented here */
+
+/** 
+ * Delegate fired when the class starts compiling. The class may be new or recycled.
+ * @param	InClass			The class that is being compiled. This could be a newly created class or a recycled class
+ * @param	InCompilerContext	The compiler context for the current compilation
+ * @param	OutCompiledData	The compiled data that this handler can write to
+ */
+DECLARE_MULTICAST_DELEGATE_ThreeParams(FOnStartCompilingClass, const UClass* /*InClass*/, IAnimBlueprintCompilationBracketContext& /*InCompilerContext*/, IAnimBlueprintGeneratedClassCompiledData& /*OutCompiledData*/);
+
+/** 
+ * Delegate fired before all animation nodes are processed
+ * @param	InAnimNodes		The set of anim nodes that should be processed. Note that these nodes have not yet been pruned for connectivity
+ * @param	InCompilerContext	The compiler context for the current compilation
+ * @param	OutCompiledData	The compiled data that this handler can write to
+ */
+DECLARE_MULTICAST_DELEGATE_ThreeParams(FOnPreProcessAnimationNodes, TArrayView<UAnimGraphNode_Base*> /*InAnimNodes*/, IAnimBlueprintCompilationContext& /*InCompilerContext*/, IAnimBlueprintGeneratedClassCompiledData& /*OutCompiledData*/);
+
+/** 
+ * Delegate fired after all animation nodes are processed
+ * @param	InAnimNodes		The set of anim nodes that were processed. Note that these nodes were not pruned for connectivity (they will be the same set passed to PreProcessAnimationNodes)
+ * @param	InCompilerContext	The compiler context for the current compilation
+ * @param	OutCompiledData	The compiled data that this handler can write to
+ */
+DECLARE_MULTICAST_DELEGATE_ThreeParams(FOnPostProcessAnimationNodes, TArrayView<UAnimGraphNode_Base*> /*InAnimNodes*/, IAnimBlueprintCompilationContext& /*InCompilerContext*/, IAnimBlueprintGeneratedClassCompiledData& /*OutCompiledData*/);
+
+/** 
+ * Delegate fired post- graph expansion
+ * @param	InGraph		The graph that was just expanded
+ * @param	InCompilerContext	The compiler context for the current compilation
+ * @param	OutCompiledData	The compiled data that this handler can write to
+ */
+DECLARE_MULTICAST_DELEGATE_ThreeParams(FOnPostExpansionStep, const UEdGraph* /*InGraph*/, IAnimBlueprintPostExpansionStepContext& /*InCompilerContext*/, IAnimBlueprintGeneratedClassCompiledData& /*OutCompiledData*/);
+
+/** 
+ * Delegate fired when the class has finished compiling
+ * @param	InClass			The class that was compiled
+ * @param	InCompilerContext	The compiler context for the current compilation
+ * @param	OutCompiledData	The compiled data that this handler can write to
+ */
+DECLARE_MULTICAST_DELEGATE_ThreeParams(FOnFinishCompilingClass, const UClass* /*InClass*/, IAnimBlueprintCompilationBracketContext& /*InCompilerContext*/, IAnimBlueprintGeneratedClassCompiledData& /*OutCompiledData*/);
+
+/** 
+ * Delegate fired when data is being copied to the CDO
+ * @param	InDefaultObject		The CDO for the just-compiled class
+ * @param	InCompilerContext	The compiler context for the current compilation
+ * @param	OutCompiledData	The compiled data that this handler can write to
+ */
+DECLARE_MULTICAST_DELEGATE_ThreeParams(FOnCopyTermDefaultsToDefaultObject, UObject* /*InDefaultObject*/, IAnimBlueprintCopyTermDefaultsContext& /*InCompilerContext*/, IAnimBlueprintGeneratedClassCompiledData& /*OutCompiledData*/);
+
+/** End ordered calls */
+
+/** Interface to the creation of an anim BP compiler context */
+class ANIMGRAPH_API IAnimBlueprintCompilerCreationContext
+{
+public:
+	virtual ~IAnimBlueprintCompilerCreationContext() {}
+
+	/** Delegate fired when the class starts compiling. The class may be new or recycled. */
+	virtual FOnStartCompilingClass& OnStartCompilingClass() = 0;
+
+	/** Delegate fired before all animation nodes are processed */
+	virtual FOnPreProcessAnimationNodes& OnPreProcessAnimationNodes() = 0;
+
+	/** Delegate fired after all animation nodes are processed */
+	virtual FOnPostProcessAnimationNodes& OnPostProcessAnimationNodes() = 0;
+
+	/** Delegate fired post- graph expansion */
+	virtual FOnPostExpansionStep& OnPostExpansionStep() = 0;
+
+	/** Delegate fired when the class has finished compiling */
+	virtual FOnFinishCompilingClass& OnFinishCompilingClass() = 0;
+
+	/** Delegate fired when data is being copied to the CDO */
+	virtual FOnCopyTermDefaultsToDefaultObject& OnCopyTermDefaultsToDefaultObject() = 0;
+
+	/** Registers a graphs schema class to the anim BP compiler so that default function processing is not performed on it */
+	virtual void RegisterKnownGraphSchema(TSubclassOf<UEdGraphSchema> InGraphSchemaClass) = 0;
+};
+
diff --git a/Engine/Source/Editor/AnimGraph/Public/IAnimBlueprintCompilerHandler.h b/Engine/Source/Editor/AnimGraph/Public/IAnimBlueprintCompilerHandler.h
new file mode 100644
--- /dev/null
+++ b/Engine/Source/Editor/AnimGraph/Public/IAnimBlueprintCompilerHandler.h
@@ -0,0 +1,12 @@
+// Copyright Epic Games, Inc. All Rights Reserved.
+
+#pragma once
+
+#include "CoreMinimal.h"
+
+/** Empty base class for handlers that are maintained for the lifetime of the compiler context */
+class ANIMGRAPH_API IAnimBlueprintCompilerHandler
+{
+public:
+	virtual ~IAnimBlueprintCompilerHandler() {}
+};
diff --git a/Engine/Source/Editor/AnimGraph/Public/IAnimBlueprintCompilerHandlerCollection.h b/Engine/Source/Editor/AnimGraph/Public/IAnimBlueprintCompilerHandlerCollection.h
new file mode 100644
--- /dev/null
+++ b/Engine/Source/Editor/AnimGraph/Public/IAnimBlueprintCompilerHandlerCollection.h
@@ -0,0 +1,32 @@
+// Copyright Epic Games, Inc. All Rights Reserved.
+
+#pragma once
+
+#include "CoreMinimal.h"
+#include "IAnimBlueprintCompilerHandler.h"
+
+class FAnimBlueprintCompilerContext;
+class IAnimBlueprintCompilerCreationContext;
+
+class ANIMGRAPH_API IAnimBlueprintCompilerHandlerCollection
+{
+public:
+	virtual ~IAnimBlueprintCompilerHandlerCollection() {}
+
+	/** Register a named handler */
+	static void RegisterHandler(FName InName, TFunction<TUniquePtr<IAnimBlueprintCompilerHandler>(IAnimBlueprintCompilerCreationContext&)> InFunction);
+
+	/** Register a named handler */
+	static void UnregisterHandler(FName InName);
+
+	// Get a GetHandler of specified type with the specified name
+	template<typename THandlerClass>
+	THandlerClass* GetHandler(FName InName) const
+	{
+		return static_cast<THandlerClass*>(GetHandlerByName(InName));
+	}
+
+protected:
+	// Get a GetHandler with the specified name
+	virtual IAnimBlueprintCompilerHandler* GetHandlerByName(FName InName) const = 0;
+};
\ No newline at end of file
diff --git a/Engine/Source/Editor/AnimGraph/Public/IAnimBlueprintPostExpansionStepContext.h b/Engine/Source/Editor/AnimGraph/Public/IAnimBlueprintPostExpansionStepContext.h
new file mode 100644
--- /dev/null
+++ b/Engine/Source/Editor/AnimGraph/Public/IAnimBlueprintPostExpansionStepContext.h
@@ -0,0 +1,47 @@
+// Copyright Epic Games, Inc. All Rights Reserved.
+
+#pragma once
+
+#include "CoreMinimal.h"
+#include "IAnimBlueprintCompilerHandler.h"
+
+class FCompilerResultsLog;
+class UEdGraph;
+struct FKismetCompilerOptions;
+
+// Interface passed to PostExpansionStep delegate
+class IAnimBlueprintPostExpansionStepContext
+{
+public:
+	virtual ~IAnimBlueprintPostExpansionStepContext() {}
+
+	// Get a handler of the specified type and name (i.e. via simple name-based RTTI)
+	// Handlers are registered via IAnimBlueprintCompilerHandlerCollection::RegisterHandler
+	template <typename THandlerClass>
+	THandlerClass* GetHandler(FName InName) const
+	{
+		return static_cast<THandlerClass*>(GetHandlerInternal(InName));
+	}
+
+	// Get the message log for the current compilation
+	FCompilerResultsLog& GetMessageLog() const { return GetMessageLogImpl(); }
+
+	// Get the consolidated uber graph during compilation
+	UEdGraph* GetConsolidatedEventGraph() const { return GetConsolidatedEventGraphImpl(); }
+
+	// Get the compiler options we are currently using
+	const FKismetCompilerOptions& GetCompileOptions() const { return GetCompileOptionsImpl(); }
+
+protected:
+	// Get the message log for the current compilation
+	virtual FCompilerResultsLog& GetMessageLogImpl() const = 0;
+
+	// Get the consolidated uber graph during compilation
+	virtual UEdGraph* GetConsolidatedEventGraphImpl() const = 0;
+
+	// Get the compiler options we are currently using
+	virtual const FKismetCompilerOptions& GetCompileOptionsImpl() const = 0;
+
+	// GetHandler helper function
+	virtual IAnimBlueprintCompilerHandler* GetHandlerInternal(FName InName) const = 0;
+};
diff --git a/Engine/Source/Editor/AnimGraph/Public/IClassVariableCreator.h b/Engine/Source/Editor/AnimGraph/Public/IClassVariableCreator.h
--- a/Engine/Source/Editor/AnimGraph/Public/IClassVariableCreator.h
+++ b/Engine/Source/Editor/AnimGraph/Public/IClassVariableCreator.h
@@ -8,6 +8,16 @@
 #include "IClassVariableCreator.generated.h"
 
 class FKismetCompilerContext;
+class FProperty;
+struct FEdGraphPinType;
+
+/** Context passed to IClassVariableCreator::CreateClassVariablesFromBlueprint */
+class IAnimBlueprintVariableCreationContext
+{
+public:
+	/** Create a class variable in the current class */
+	virtual FProperty* CreateVariable(const FName Name, const FEdGraphPinType& Type) = 0;
+};
 
 UINTERFACE(MinimalAPI)
 class UClassVariableCreator : public UInterface
@@ -23,7 +33,7 @@ public:
 	/** 
 	 * Implement this in a graph node and the anim BP compiler will call this expecting to generate
 	 * class variables.
-	 * @param	InCompilerContext	The compiler context for the current BP compilation
+	 * @param	InVariableCreator	The variable creation context for the current BP compilation
 	 */
-	virtual void CreateClassVariablesFromBlueprint(FKismetCompilerContext& InCompilerContext) = 0;
+	virtual void CreateClassVariablesFromBlueprint(IAnimBlueprintVariableCreationContext& InCreationContext) = 0;
 };
\ No newline at end of file
diff --git a/Engine/Source/Editor/AnimGraph/Public/PropertyAccessCompilerHandler.h b/Engine/Source/Editor/AnimGraph/Public/PropertyAccessCompilerHandler.h
new file mode 100644
--- /dev/null
+++ b/Engine/Source/Editor/AnimGraph/Public/PropertyAccessCompilerHandler.h
@@ -0,0 +1,32 @@
+// Copyright Epic Games, Inc. All Rights Reserved.
+
+#pragma once
+
+#include "CoreMinimal.h"
+#include "UObject/Interface.h"
+#include "Containers/ArrayView.h"
+#include "Delegates/Delegate.h"
+#include "IAnimBlueprintCompilerHandler.h"
+
+enum class EPropertyAccessBatchType : uint8;
+class IAnimBlueprintGeneratedClassCompiledData;
+
+// Delegate called when the library is compiled (whether successfully or not)
+DECLARE_MULTICAST_DELEGATE_OneParam(FOnPostLibraryCompiled, IAnimBlueprintGeneratedClassCompiledData& /*OutCompiledData*/)
+
+class FPropertyAccessCompilerHandler : public IAnimBlueprintCompilerHandler
+{
+public:
+	// Add a copy to the property access library we are compiling
+	// @return an integer handle to the pending copy. This can be resolved to a true copy index by calling MapCopyIndex
+	virtual int32 AddCopy(TArrayView<FString> InSourcePath, TArrayView<FString> InDestPath, EPropertyAccessBatchType InBatchType, UObject* InObject = nullptr) = 0;
+
+	// Delegate called when the library is compiled (whether successfully or not)
+	virtual FSimpleMulticastDelegate& OnPreLibraryCompiled() = 0;
+
+	// Delegate called when the library is compiled (whether successfully or not)
+	virtual FOnPostLibraryCompiled& OnPostLibraryCompiled() = 0;
+
+	// Maps the initial integer copy handle to a true handle, post compilation
+	virtual int32 MapCopyIndex(int32 InIndex) const = 0;
+};
\ No newline at end of file
diff --git a/Engine/Source/Editor/KismetCompiler/Public/KismetCompiler.h b/Engine/Source/Editor/KismetCompiler/Public/KismetCompiler.h
--- a/Engine/Source/Editor/KismetCompiler/Public/KismetCompiler.h
+++ b/Engine/Source/Editor/KismetCompiler/Public/KismetCompiler.h
@@ -388,7 +388,7 @@ protected:
 	virtual void PostCompile() { OnPostCompile.Broadcast(); }
 
 	// Gives derived classes a chance to process post-node expansion
-	virtual void PostExpansionStep(UEdGraph* Graph) {}
+	virtual void PostExpansionStep(const UEdGraph* Graph) {}
 
 	/** Determines if a node is pure */
 	virtual bool IsNodePure(const UEdGraphNode* Node) const;
diff --git a/Engine/Source/Editor/UnrealEd/Public/IPropertyAccessCompiler.h b/Engine/Source/Editor/UnrealEd/Public/IPropertyAccessCompiler.h
--- a/Engine/Source/Editor/UnrealEd/Public/IPropertyAccessCompiler.h
+++ b/Engine/Source/Editor/UnrealEd/Public/IPropertyAccessCompiler.h
@@ -26,7 +26,7 @@ public:
 	virtual ~IPropertyAccessLibraryCompiler() {}
 
 	// Begin compilation - reset the library to its default state
-	virtual void BeginCompilation(UClass* InClass) = 0;
+	virtual void BeginCompilation(const UClass* InClass) = 0;
 
 	// Add a copy to the property access library we are compiling
 	// @return an integer handle to the pending copy. This can be resolved to a true copy index by calling MapCopyIndex
diff --git a/Engine/Source/Editor/UnrealEd/Public/IPropertyAccessEditor.h b/Engine/Source/Editor/UnrealEd/Public/IPropertyAccessEditor.h
--- a/Engine/Source/Editor/UnrealEd/Public/IPropertyAccessEditor.h
+++ b/Engine/Source/Editor/UnrealEd/Public/IPropertyAccessEditor.h
@@ -194,7 +194,7 @@ public:
 	virtual TSharedRef<SWidget> MakePropertyBindingWidget(UBlueprint* InBlueprint, const FPropertyBindingWidgetArgs& InArgs = FPropertyBindingWidgetArgs()) const = 0;
 
 	/** Resolve a property path to a structure, returning the leaf property and array index if any. @return true if resolution succeeded */
-	virtual EPropertyAccessResolveResult ResolveLeafProperty(UStruct* InStruct, TArrayView<FString> InPath, FProperty*& OutProperty, int32& OutArrayIndex) const = 0;
+	virtual EPropertyAccessResolveResult ResolveLeafProperty(const UStruct* InStruct, TArrayView<FString> InPath, FProperty*& OutProperty, int32& OutArrayIndex) const = 0;
 
 	// Get the compatibility of the two supplied properties. Ordering matters for promotion (A->B).
 	virtual EPropertyAccessCompatibility GetPropertyCompatibility(const FProperty* InPropertyA, const FProperty* InPropertyB) const = 0;
diff --git a/Engine/Source/Runtime/Engine/Classes/Animation/AnimBlueprintGeneratedClass.h b/Engine/Source/Runtime/Engine/Classes/Animation/AnimBlueprintGeneratedClass.h
--- a/Engine/Source/Runtime/Engine/Classes/Animation/AnimBlueprintGeneratedClass.h
+++ b/Engine/Source/Runtime/Engine/Classes/Animation/AnimBlueprintGeneratedClass.h
@@ -12,7 +12,7 @@
 #include "Animation/AnimClassInterface.h"
 #include "Animation/AnimNodeBase.h"
 #include "Animation/BlendSpaceBase.h"
-#include "AnimBlueprintClassSubsystem.h"
+#include "PropertyAccess.h"
 
 #include "AnimBlueprintGeneratedClass.generated.h"
 
@@ -284,8 +284,8 @@ class ENGINE_API UAnimBlueprintGeneratedClass : public UBlueprintGeneratedClass,
 {
 	GENERATED_UCLASS_BODY()
 
-	friend class UAnimBlueprintClassSubsystem;
 	friend class FAnimBlueprintCompilerContext;
+	friend class FAnimBlueprintGeneratedClassCompiledData;
 
 	// List of state machines present in this blueprint class
 	UPROPERTY()
@@ -331,22 +331,16 @@ class ENGINE_API UAnimBlueprintGeneratedClass : public UBlueprintGeneratedClass,
 	UPROPERTY()
 	TMap<FName, FAnimGraphBlendOptions> GraphBlendOptions;
 
-	/** Data for each subsystem */
+private:
+	// The library holding the property access data
 	UPROPERTY()
-	TArray<UAnimBlueprintClassSubsystem*> Subsystems;
-
-	/** Map of class->subsystem */
-	TMap<TSubclassOf<UAnimBlueprintClassSubsystem>, UAnimBlueprintClassSubsystem*> SubsystemMap;
-	TMap<TSubclassOf<UInterface>, UAnimBlueprintClassSubsystem*> SubsystemInterfaceMap;
-
-	/** Subsystem properties */
-	TArray<FStructProperty*> SubsystemProperties;
+	FPropertyAccessLibrary PropertyAccessLibrary;
 
 public:
 	// IAnimClassInterface interface
-	virtual const TArray<FBakedAnimationStateMachine>& GetBakedStateMachines() const override { return GetRootClass()->BakedStateMachines; }
+	virtual const TArray<FBakedAnimationStateMachine>& GetBakedStateMachines() const override { return GetRootClass()->GetBakedStateMachines_Direct(); }
 	virtual USkeleton* GetTargetSkeleton() const override { return TargetSkeleton; }
-	virtual const TArray<FAnimNotifyEvent>& GetAnimNotifies() const override { return GetRootClass()->AnimNotifies; }
+	virtual const TArray<FAnimNotifyEvent>& GetAnimNotifies() const override { return GetRootClass()->GetAnimNotifies_Direct(); }
 	virtual const TArray<FStructProperty*>& GetAnimNodeProperties() const override { return AnimNodeProperties; }
 	virtual const TArray<FStructProperty*>& GetLinkedAnimGraphNodeProperties() const override { return LinkedAnimGraphNodeProperties; }
 	virtual const TArray<FStructProperty*>& GetLinkedAnimLayerNodeProperties() const override { return LinkedAnimLayerNodeProperties; }
@@ -354,27 +348,25 @@ public:
 	virtual const TArray<FStructProperty*>& GetDynamicResetNodeProperties() const override { return DynamicResetNodeProperties; }
 	virtual const TArray<FStructProperty*>& GetStateMachineNodeProperties() const override { return StateMachineNodeProperties; }
 	virtual const TArray<FStructProperty*>& GetInitializationNodeProperties() const override { return InitializationNodeProperties; }
-	virtual const TArray<FName>& GetSyncGroupNames() const override { return GetRootClass()->SyncGroupNames; }
-	virtual const TMap<FName, FCachedPoseIndices>& GetOrderedSavedPoseNodeIndicesMap() const override { return GetRootClass()->OrderedSavedPoseIndicesMap; }
+	virtual const TArray<FName>& GetSyncGroupNames() const override { return GetRootClass()->GetSyncGroupNames_Direct(); }
+	virtual const TMap<FName, FCachedPoseIndices>& GetOrderedSavedPoseNodeIndicesMap() const override { return GetRootClass()->GetOrderedSavedPoseNodeIndicesMap_Direct(); }
 	virtual int32 GetSyncGroupIndex(FName SyncGroupName) const override { return GetSyncGroupNames().IndexOfByKey(SyncGroupName); }
 	virtual const TArray<FExposedValueHandler>& GetExposedValueHandlers() const { return EvaluateGraphExposedInputs; }
 	virtual const TArray<FAnimBlueprintFunction>& GetAnimBlueprintFunctions() const override { return AnimBlueprintFunctions; }
-	virtual const TMap<FName, FGraphAssetPlayerInformation>& GetGraphAssetPlayerInformation() const override { return GetRootClass()->GraphAssetPlayerInformation; }
-	virtual const TMap<FName, FAnimGraphBlendOptions>& GetGraphBlendOptions() const override { return GetRootClass()->GraphBlendOptions; }
-	virtual const TArray<UAnimBlueprintClassSubsystem*>& GetSubsystems() const override { return GetRootClass()->Subsystems; }
-	virtual UAnimBlueprintClassSubsystem* GetSubsystem(TSubclassOf<UAnimBlueprintClassSubsystem> InClass) const override;
-	virtual UAnimBlueprintClassSubsystem* FindSubsystemWithInterface(TSubclassOf<UInterface> InClassInterface) const override;
-	virtual const TArray<FStructProperty*>& GetSubsystemProperties() const override { return SubsystemProperties; }
+	virtual const TMap<FName, FGraphAssetPlayerInformation>& GetGraphAssetPlayerInformation() const override { return GetRootClass()->GetGraphAssetPlayerInformation_Direct(); }
+	virtual const TMap<FName, FAnimGraphBlendOptions>& GetGraphBlendOptions() const override { return GetRootClass()->GetGraphBlendOptions_Direct(); }
+	virtual const FPropertyAccessLibrary& GetPropertyAccessLibrary() const override { return GetRootClass()->GetPropertyAccessLibrary_Direct(); }
+
+private:
+	virtual const TArray<FBakedAnimationStateMachine>& GetBakedStateMachines_Direct() const override { return BakedStateMachines; }
+	virtual const TArray<FAnimNotifyEvent>& GetAnimNotifies_Direct() const override { return AnimNotifies; }
+	virtual const TArray<FName>& GetSyncGroupNames_Direct() const override { return SyncGroupNames; }
+	virtual const TMap<FName, FCachedPoseIndices>& GetOrderedSavedPoseNodeIndicesMap_Direct() const override { return OrderedSavedPoseIndicesMap; }
+	virtual const TMap<FName, FGraphAssetPlayerInformation>& GetGraphAssetPlayerInformation_Direct() const override { return GraphAssetPlayerInformation; }
+	virtual const TMap<FName, FAnimGraphBlendOptions>& GetGraphBlendOptions_Direct() const override { return GraphBlendOptions; }
+	virtual const FPropertyAccessLibrary& GetPropertyAccessLibrary_Direct() const override { return PropertyAccessLibrary; }
 
 public:
-	// Get the root anim BP class (i.e. if this is a derived class).
-	// Some properties that are derived from the compiled anim graph are routed to the 'Root' class
-	// as child classes don't get fully compiled. Instead they just override various asset players leaving the
-	// full compilation up to the base class. 
-	// Previously we copied over all the parent class data in Link(), but as Link() can be called on the async 
-	// loading thread we cant do any object-duplication operations (e.g. with subsystems).
-	const UAnimBlueprintGeneratedClass* GetRootClass() const;
-
 #if WITH_EDITORONLY_DATA
 	FAnimBlueprintDebugData AnimBlueprintDebugData;
 
@@ -490,9 +482,6 @@ public:
 	// Populates AnimBlueprintFunctions according to the UFunction(s) on this class
 	void GenerateAnimationBlueprintFunctions();
 
-	// Rebuild subsystem & subsystem interface maps from Subsystems array
-	void RebuildSubsystemMaps();
-
 	// UObject interface
 	virtual void Serialize(FArchive& Ar) override;
 	// End of UObject interface
diff --git a/Engine/Source/Runtime/Engine/Classes/Animation/AnimClassData.h b/Engine/Source/Runtime/Engine/Classes/Animation/AnimClassData.h
--- a/Engine/Source/Runtime/Engine/Classes/Animation/AnimClassData.h
+++ b/Engine/Source/Runtime/Engine/Classes/Animation/AnimClassData.h
@@ -11,6 +11,7 @@
 #include "Animation/AnimClassInterface.h"
 #include "Algo/Transform.h"
 #include "AnimBlueprintGeneratedClass.h"
+#include "PropertyAccess.h"
 
 #include "AnimClassData.generated.h"
 
@@ -114,19 +115,9 @@ public:
 	UPROPERTY()
 	TMap<FName, FAnimGraphBlendOptions> GraphBlendOptions;
 
-private:
-	/** Data for each subsystem */
-	UPROPERTY()
-	TArray<UAnimBlueprintClassSubsystem*> Subsystems;
-
-	/** Map of class->subsystem */
-	TMap<TSubclassOf<UAnimBlueprintClassSubsystem>, UAnimBlueprintClassSubsystem*> SubsystemMap;
-	TMap<TSubclassOf<UInterface>, UAnimBlueprintClassSubsystem*> SubsystemInterfaceMap;
-
-	/** Subsystem properties */
+	// Property access library
 	UPROPERTY()
-	TArray< TFieldPath<FStructProperty> > SubsystemProperties;
-	TArray<FStructProperty*> ResolvedSubsystemProperties;
+	FPropertyAccessLibrary PropertyAccessLibrary;
 
 public:
 	// IAnimClassInterface interface
@@ -147,12 +138,19 @@ public:
 	virtual const TArray<FExposedValueHandler>& GetExposedValueHandlers() const override { return EvaluateGraphExposedInputs; }
 	virtual const TMap<FName, FGraphAssetPlayerInformation>& GetGraphAssetPlayerInformation() const override { return GraphNameAssetPlayers; }
 	virtual const TMap<FName, FAnimGraphBlendOptions>& GetGraphBlendOptions() const override { return GraphBlendOptions; }
-	virtual const TArray<UAnimBlueprintClassSubsystem*>& GetSubsystems() const override { return Subsystems; }
-	virtual UAnimBlueprintClassSubsystem* GetSubsystem(TSubclassOf<UAnimBlueprintClassSubsystem> InClass) const override;
-	virtual UAnimBlueprintClassSubsystem* FindSubsystemWithInterface(TSubclassOf<UInterface> InClassInterface) const override;
-	virtual const TArray<FStructProperty*>& GetSubsystemProperties() const override { return ResolvedSubsystemProperties; }
+	virtual const FPropertyAccessLibrary& GetPropertyAccessLibrary() const override { return PropertyAccessLibrary; }
 
-	// Resolve TFieldPaths to FStructPropertys, setup subsystems, init value handlers
+private:
+	virtual const TArray<FBakedAnimationStateMachine>& GetBakedStateMachines_Direct() const override { return BakedStateMachines; }
+	virtual const TArray<FAnimNotifyEvent>& GetAnimNotifies_Direct() const override { return AnimNotifies; }
+	virtual const TArray<FName>& GetSyncGroupNames_Direct() const override { return SyncGroupNames; }
+	virtual const TMap<FName, FCachedPoseIndices>& GetOrderedSavedPoseNodeIndicesMap_Direct() const override { return OrderedSavedPoseIndicesMap; }
+	virtual const TMap<FName, FGraphAssetPlayerInformation>& GetGraphAssetPlayerInformation_Direct() const override { return GraphNameAssetPlayers; }
+	virtual const TMap<FName, FAnimGraphBlendOptions>& GetGraphBlendOptions_Direct() const override { return GraphBlendOptions; }
+	virtual const FPropertyAccessLibrary& GetPropertyAccessLibrary_Direct() const override { return PropertyAccessLibrary; }
+
+public:
+	// Resolve TFieldPaths to FStructPropertys, init value handlers
 	void DynamicClassInitialization(UDynamicClass* InDynamicClass);
 
 #if WITH_EDITOR
diff --git a/Engine/Source/Runtime/Engine/Classes/Animation/AnimClassInterface.h b/Engine/Source/Runtime/Engine/Classes/Animation/AnimClassInterface.h
--- a/Engine/Source/Runtime/Engine/Classes/Animation/AnimClassInterface.h
+++ b/Engine/Source/Runtime/Engine/Classes/Animation/AnimClassInterface.h
@@ -11,12 +11,12 @@
 #include "Animation/AnimTypes.h"
 #include "Animation/AnimStateMachineTypes.h"
 #include "UObject/FieldPath.h"
-#include "AnimBlueprintClassSubsystem.h"
 
 #include "AnimClassInterface.generated.h"
 
 class USkeleton;
 struct FExposedValueHandler;
+struct FPropertyAccessLibrary;
 
 /** Describes the input and output of an anim blueprint 'function' */
 USTRUCT()
@@ -164,10 +164,46 @@ public:
 	virtual const TMap<FName, FAnimGraphBlendOptions>& GetGraphBlendOptions() const = 0;
 	virtual USkeleton* GetTargetSkeleton() const = 0;
 	virtual int32 GetSyncGroupIndex(FName SyncGroupName) const = 0;
-	virtual const TArray<UAnimBlueprintClassSubsystem*>& GetSubsystems() const = 0;
-	virtual UAnimBlueprintClassSubsystem* GetSubsystem(TSubclassOf<UAnimBlueprintClassSubsystem> InClass) const = 0;
-	virtual UAnimBlueprintClassSubsystem* FindSubsystemWithInterface(TSubclassOf<UInterface> InClassInterface) const = 0;
-	virtual const TArray<FStructProperty*>& GetSubsystemProperties() const = 0;
+	virtual const FPropertyAccessLibrary& GetPropertyAccessLibrary() const = 0;
+	
+protected:
+	friend class UAnimBlueprintGeneratedClass;
+
+	// These direct accessors are here to allow internal access that doesnt redirect to the root class
+	virtual const TArray<FBakedAnimationStateMachine>& GetBakedStateMachines_Direct() const = 0;
+	virtual const TArray<FAnimNotifyEvent>& GetAnimNotifies_Direct() const = 0;
+	virtual const TArray<FName>& GetSyncGroupNames_Direct() const = 0;
+	virtual const TMap<FName, FCachedPoseIndices>& GetOrderedSavedPoseNodeIndicesMap_Direct() const = 0;
+	virtual const TMap<FName, FGraphAssetPlayerInformation>& GetGraphAssetPlayerInformation_Direct() const = 0;
+	virtual const TMap<FName, FAnimGraphBlendOptions>& GetGraphBlendOptions_Direct() const = 0;
+	virtual const FPropertyAccessLibrary& GetPropertyAccessLibrary_Direct() const = 0;
+
+public:
+
+	// Get the root anim class interface (i.e. if this is a derived class).
+	// Some properties that are derived from the compiled anim graph are routed to the 'Root' class
+	// as child classes don't get fully compiled. Instead they just override various asset players leaving the
+	// full compilation up to the base class. 
+	const IAnimClassInterface* GetRootClass() const
+	{
+		auto GetSuperClassInterface = [](const IAnimClassInterface* InClass) -> const IAnimClassInterface*
+		{
+			if(const UClass* ActualClass = GetActualAnimClass(InClass))
+			{
+				return GetFromClass(ActualClass->GetSuperClass());
+			}
+
+			return nullptr;
+		};
+
+		const IAnimClassInterface* RootClass = this;
+		while(const IAnimClassInterface* NextClass = GetSuperClassInterface(RootClass))
+		{
+			RootClass = NextClass;
+		}
+
+		return RootClass;
+	}
 
 	static IAnimClassInterface* GetFromClass(UClass* InClass)
 	{
@@ -183,6 +219,20 @@ public:
 		return nullptr;
 	}
 
+	static const IAnimClassInterface* GetFromClass(const UClass* InClass)
+	{
+		if (auto AnimClassInterface = Cast<const IAnimClassInterface>(InClass))
+		{
+			return AnimClassInterface;
+		}
+		if (auto DynamicClass = Cast<const UDynamicClass>(InClass))
+		{
+			DynamicClass->GetDefaultObject(true);
+			return CastChecked<const IAnimClassInterface>(DynamicClass->AnimClassImplementation, ECastCheckedType::NullAllowed);
+		}
+		return nullptr;
+	}
+
 	static UClass* GetActualAnimClass(IAnimClassInterface* AnimClassInterface)
 	{
 		if (UClass* ActualAnimClass = Cast<UClass>(AnimClassInterface))
@@ -196,6 +246,19 @@ public:
 		return nullptr;
 	}
 
+	static const UClass* GetActualAnimClass(const IAnimClassInterface* AnimClassInterface)
+	{
+		if (const UClass* ActualAnimClass = Cast<const UClass>(AnimClassInterface))
+		{
+			return ActualAnimClass;
+		}
+		if (const UObject* AsObject = Cast<const UObject>(AnimClassInterface))
+		{
+			return Cast<const UClass>(AsObject->GetOuter());
+		}
+		return nullptr;
+	}
+
 	static const FAnimBlueprintFunction* FindAnimBlueprintFunction(IAnimClassInterface* AnimClassInterface, const FName& InFunctionName)
 	{
 		for(const FAnimBlueprintFunction& Function : AnimClassInterface->GetAnimBlueprintFunctions())
@@ -230,23 +293,6 @@ public:
 		return false;
 	}
 
-	/** Get a subsystem */
-	template <typename TSubsystemClass>
-	static TSubsystemClass* GetSubsystem(IAnimClassInterface* InAnimClassInterface)
-	{
-		return Cast<TSubsystemClass>(InAnimClassInterface->GetSubsystem(TSubsystemClass::StaticClass()));
-	}
-
-	/** Find the first subsystem with the specified interface */
-	template <typename TInterfaceClass>
-	static TInterfaceClass* FindSubsystemWithInterface(IAnimClassInterface* InAnimClassInterface)
-	{
-		return Cast<TInterfaceClass>(InAnimClassInterface->FindSubsystemWithInterface(TInterfaceClass::UClassType::StaticClass()));
-	}
-
-	/** Run a function on each subsystem's instance data */
-	static void ForEachAnimInstanceSubsystemData(UAnimInstance* InAnimInstance, TFunctionRef<void(UAnimBlueprintClassSubsystem*, FAnimInstanceSubsystemData&)> InFunction);
-
 	UE_DEPRECATED(4.23, "Please use GetAnimBlueprintFunctions()")
 	virtual int32 GetRootAnimNodeIndex() const { return INDEX_NONE; }
 
diff --git a/Engine/Source/Runtime/Engine/Classes/Animation/AnimNodeBase.h b/Engine/Source/Runtime/Engine/Classes/Animation/AnimNodeBase.h
--- a/Engine/Source/Runtime/Engine/Classes/Animation/AnimNodeBase.h
+++ b/Engine/Source/Runtime/Engine/Classes/Animation/AnimNodeBase.h
@@ -31,7 +31,7 @@ class UAnimInstance;
 struct FAnimInstanceProxy;
 struct FAnimNode_Base;
 class UProperty;
-class IPropertyAccess;
+struct FPropertyAccessLibrary;
 
 /**
  * Utility container for tracking a stack of ancestor nodes by node type during graph traversal
@@ -742,7 +742,7 @@ struct ENGINE_API FExposedValueHandler
 		: BoundFunction(NAME_None)
 		, Function(nullptr)
 		, ValueHandlerNodeProperty(nullptr)
-		, PropertyAccess(nullptr)
+		, PropertyAccessLibrary(nullptr)
 		, bInitialized(false)
 	{
 	}
@@ -765,8 +765,8 @@ struct ENGINE_API FExposedValueHandler
 	UPROPERTY()
 	TFieldPath<FStructProperty> ValueHandlerNodeProperty;
 
-	// Cached property access system ptr
-	IPropertyAccess* PropertyAccess;
+	// Cached property access library ptr
+	const FPropertyAccessLibrary* PropertyAccessLibrary;
 
 	// Prevent multiple initialization
 	bool bInitialized;
@@ -780,7 +780,7 @@ struct ENGINE_API FExposedValueHandler
 	static void ClassInitialization(TArray<FExposedValueHandler>& Handlers, UObject* ClassDefaultObject);
 
 	// Bind copy records and cache UFunction if necessary
-	void Initialize(UClass* InClass, IPropertyAccess* InPropertyAccess);
+	void Initialize(UClass* InClass, const FPropertyAccessLibrary& InPropertyAccessLibrary);
 
 	// Execute the function and copy records
 	void Execute(const FAnimationBaseContext& Context) const;
diff --git a/Engine/Source/Runtime/Engine/Engine.Build.cs b/Engine/Source/Runtime/Engine/Engine.Build.cs
--- a/Engine/Source/Runtime/Engine/Engine.Build.cs
+++ b/Engine/Source/Runtime/Engine/Engine.Build.cs
@@ -94,6 +94,7 @@ public class Engine : ModuleRules
 				"SignalProcessing",
 				"AudioExtensions",
 				"DeveloperSettings",
+				"PropertyAccess",
 			}
 		);
 
diff --git a/Engine/Source/Runtime/Engine/Private/Animation/AnimBlueprintGeneratedClass.cpp b/Engine/Source/Runtime/Engine/Private/Animation/AnimBlueprintGeneratedClass.cpp
--- a/Engine/Source/Runtime/Engine/Private/Animation/AnimBlueprintGeneratedClass.cpp
+++ b/Engine/Source/Runtime/Engine/Private/Animation/AnimBlueprintGeneratedClass.cpp
@@ -20,7 +20,6 @@
 #include "Animation/AnimNode_AssetPlayerBase.h"
 #include "Animation/AnimNode_StateMachine.h"
 #include "EdGraph/EdGraphNode.h"
-#include "Animation/AnimInstanceSubsystemData.h"
 #include "Algo/Reverse.h"
 
 /////////////////////////////////////////////////////
@@ -276,9 +275,6 @@ void UAnimBlueprintGeneratedClass::Serialize(FArchive& Ar)
 	Ar.UsingCustomVersion(FReleaseObjectVersion::GUID);
 
 	Super::Serialize(Ar);
-
-	// Init subsystem maps
-	RebuildSubsystemMaps();
 }
 
 void UAnimBlueprintGeneratedClass::Link(FArchive& Ar, bool bRelinkExistingProperties)
@@ -293,7 +289,6 @@ void UAnimBlueprintGeneratedClass::Link(FArchive& Ar, bool bRelinkExistingProper
 	DynamicResetNodeProperties.Empty();
 	StateMachineNodeProperties.Empty();
 	InitializationNodeProperties.Empty();
-	SubsystemProperties.Empty();
 
 #if WITH_EDITORONLY_DATA
 	for (FExposedValueHandler& Handler : EvaluateGraphExposedInputs)
@@ -330,15 +325,8 @@ void UAnimBlueprintGeneratedClass::Link(FArchive& Ar, bool bRelinkExistingProper
 				}
 				AnimNodeProperties.Add(StructProp);
 			}
-			else if(StructProp->Struct->IsChildOf(FAnimInstanceSubsystemData::StaticStruct()))
-			{
-				SubsystemProperties.Add(StructProp);
-			}
 		}
 	}
-
-	// Properties are in reverse order with respect to the subsystems
-	Algo::Reverse(SubsystemProperties);
 }
 
 void UAnimBlueprintGeneratedClass::PurgeClass(bool bRecompilingOnLoad)
@@ -383,6 +371,9 @@ void UAnimBlueprintGeneratedClass::PostLoad()
 {
 	Super::PostLoad();
 	GenerateAnimationBlueprintFunctions();
+
+	// Post-load property access library
+	PropertyAccess::PostLoadLibrary(PropertyAccessLibrary);
 }
 
 void UAnimBlueprintGeneratedClass::GenerateAnimationBlueprintFunctions()
@@ -601,41 +592,3 @@ const UEdGraphNode* UAnimBlueprintGeneratedClass::GetVisualNodeFromNodePropertyI
 }
 
 #endif // WITH_EDITORONLY_DATA
-
-UAnimBlueprintClassSubsystem* UAnimBlueprintGeneratedClass::GetSubsystem(TSubclassOf<UAnimBlueprintClassSubsystem> InClass) const
-{
-	return GetRootClass()->SubsystemMap.FindRef(InClass);
-}
-
-UAnimBlueprintClassSubsystem* UAnimBlueprintGeneratedClass::FindSubsystemWithInterface(TSubclassOf<UInterface> InClassInterface) const
-{
-	return GetRootClass()->SubsystemInterfaceMap.FindRef(InClassInterface);
-}
-
-void UAnimBlueprintGeneratedClass::RebuildSubsystemMaps()
-{
-	SubsystemMap.Reset();
-	SubsystemInterfaceMap.Reset();
-
-	// Init subsystem maps
-	for(UAnimBlueprintClassSubsystem* Subsystem : Subsystems)
-	{
-		SubsystemMap.Add(Subsystem->GetClass(), Subsystem);
-
-		for(const FImplementedInterface& ImplementedInterface : Subsystem->GetClass()->Interfaces)
-		{
-			SubsystemInterfaceMap.Add(ImplementedInterface.Class, Subsystem);
-		}
-	}
-}
-
-const UAnimBlueprintGeneratedClass* UAnimBlueprintGeneratedClass::GetRootClass() const
-{
-	const UAnimBlueprintGeneratedClass* RootClass = this;
-	while(const UAnimBlueprintGeneratedClass* NextClass = Cast<UAnimBlueprintGeneratedClass>(RootClass->GetSuperClass()))
-	{
-		RootClass = NextClass;
-	}
-
-	return RootClass;
-}
\ No newline at end of file
diff --git a/Engine/Source/Runtime/Engine/Private/Animation/AnimClassData.cpp b/Engine/Source/Runtime/Engine/Private/Animation/AnimClassData.cpp
--- a/Engine/Source/Runtime/Engine/Private/Animation/AnimClassData.cpp
+++ b/Engine/Source/Runtime/Engine/Private/Animation/AnimClassData.cpp
@@ -29,35 +29,13 @@ void UAnimClassData::DynamicClassInitialization(UDynamicClass* InDynamicClass)
 		Algo::Transform(AnimBlueprintFunctionData[FunctionIndex].InputPoseNodeProperties, AnimBlueprintFunctions[FunctionIndex].InputPoseNodeProperties, ResolveTransform);
 	}
 
-	Algo::Transform(SubsystemProperties, ResolvedSubsystemProperties, ResolveTransform);
-
-	// Init subsystem maps
-	for(UAnimBlueprintClassSubsystem* Subsystem : Subsystems)
-	{
-		SubsystemMap.Add(Subsystem->GetClass(), Subsystem);
-
-		for(const FImplementedInterface& ImplementedInterface : Subsystem->GetClass()->Interfaces)
-		{
-			SubsystemInterfaceMap.Add(ImplementedInterface.Class, Subsystem);
-		}
-
-		Subsystem->PostLoadSubsystem();
-	}
+	// Init property access library
+	PropertyAccess::PostLoadLibrary(PropertyAccessLibrary);
 
 	// Init exposed value handlers
 	FExposedValueHandler::DynamicClassInitialization(EvaluateGraphExposedInputs, InDynamicClass);
 }
 
-UAnimBlueprintClassSubsystem* UAnimClassData::GetSubsystem(TSubclassOf<UAnimBlueprintClassSubsystem> InClass) const
-{
-	return SubsystemMap.FindRef(InClass);
-}
-
-UAnimBlueprintClassSubsystem* UAnimClassData::FindSubsystemWithInterface(TSubclassOf<UInterface> InClassInterface) const
-{
-	return SubsystemInterfaceMap.FindRef(InClassInterface);
-}
-
 #if WITH_EDITOR
 void UAnimClassData::CopyFrom(UAnimBlueprintGeneratedClass* AnimClass)
 {
@@ -102,14 +80,6 @@ void UAnimClassData::CopyFrom(UAnimBlueprintGeneratedClass* AnimClass)
 	EvaluateGraphExposedInputs = AnimClass->GetExposedValueHandlers();
 	GraphNameAssetPlayers = AnimClass->GetGraphAssetPlayerInformation();
 	GraphBlendOptions = AnimClass->GetGraphBlendOptions();
-
-	Algo::Transform(AnimClass->GetSubsystemProperties(), SubsystemProperties, MakePropertyPath);
-	ResolvedSubsystemProperties = AnimClass->GetSubsystemProperties();
-
-	// Deep-copy (duplicate) subsystems
-	Algo::Transform(AnimClass->GetSubsystems(), Subsystems, [this](UAnimBlueprintClassSubsystem* InSubsystem)
-	{
-		return DuplicateObject(InSubsystem, this);
-	});
+	PropertyAccessLibrary = AnimClass->GetPropertyAccessLibrary();
 }
 #endif // WITH_EDITOR
\ No newline at end of file
diff --git a/Engine/Source/Runtime/Engine/Private/Animation/AnimInstance.cpp b/Engine/Source/Runtime/Engine/Private/Animation/AnimInstance.cpp
--- a/Engine/Source/Runtime/Engine/Private/Animation/AnimInstance.cpp
+++ b/Engine/Source/Runtime/Engine/Private/Animation/AnimInstance.cpp
@@ -471,11 +471,11 @@ void UAnimInstance::UpdateAnimation(float DeltaSeconds, bool bNeedsValidRootMoti
 		UpdateMontageEvaluationData();
 	}
 
-	// Give class subsystems a crack
-	IAnimClassInterface::ForEachAnimInstanceSubsystemData(this, [this, DeltaSeconds](UAnimBlueprintClassSubsystem* InSubsystem, FAnimInstanceSubsystemData& InSubsystemData)
+	if (IAnimClassInterface* AnimBlueprintClass = IAnimClassInterface::GetFromClass(GetClass()))
 	{
-		InSubsystem->OnUpdateAnimation(this, InSubsystemData, DeltaSeconds);
-	});
+		// Process internal batched property copies
+		PropertyAccess::ProcessCopies(this, AnimBlueprintClass->GetPropertyAccessLibrary(), EPropertyAccessCopyBatch::ExternalBatched);
+	}
 
 	{
 		SCOPE_CYCLE_COUNTER(STAT_NativeUpdateAnimation);
diff --git a/Engine/Source/Runtime/Engine/Private/Animation/AnimInstanceProxy.cpp b/Engine/Source/Runtime/Engine/Private/Animation/AnimInstanceProxy.cpp
--- a/Engine/Source/Runtime/Engine/Private/Animation/AnimInstanceProxy.cpp
+++ b/Engine/Source/Runtime/Engine/Private/Animation/AnimInstanceProxy.cpp
@@ -15,6 +15,9 @@
 #include "Engine/Engine.h"
 #include "DrawDebugHelpers.h"
 #include "GameFramework/WorldSettings.h"
+#include "Animation/AnimNode_Root.h"
+#include "Animation/AnimNode_LinkedAnimLayer.h"
+#include "Animation/AnimMontage.h"
 
 #define DO_ANIMSTAT_PROCESSING(StatName) DEFINE_STAT(STAT_ ## StatName)
 #include "Animation/AnimMTStats.h"
@@ -22,9 +25,6 @@
 
 #define DO_ANIMSTAT_PROCESSING(StatName) DEFINE_STAT(STAT_ ## StatName ## _WorkerThread)
 #include "Animation/AnimMTStats.h"
-#include "Animation/AnimNode_Root.h"
-#include "Animation/AnimNode_LinkedAnimLayer.h"
-#include "Animation/AnimMontage.h"
 
 #undef DO_ANIMSTAT_PROCESSING
 
@@ -124,26 +124,23 @@ void FAnimInstanceProxy::Initialize(UAnimInstance* InAnimInstance)
 
 		// Initialize state buffers
 		int32 NumStates = 0;
-		if(IAnimClassInterface* Interface = GetAnimClassInterface())
+		const TArray<FBakedAnimationStateMachine>& BakedMachines = AnimClassInterface->GetBakedStateMachines();
+		const int32 NumMachines = BakedMachines.Num();
+		for(int32 MachineClassIndex = 0; MachineClassIndex < NumMachines; ++MachineClassIndex)
 		{
-			const TArray<FBakedAnimationStateMachine>& BakedMachines = Interface->GetBakedStateMachines();
-			const int32 NumMachines = BakedMachines.Num();
-			for(int32 MachineClassIndex = 0; MachineClassIndex < NumMachines; ++MachineClassIndex)
-			{
-				const FBakedAnimationStateMachine& Machine = BakedMachines[MachineClassIndex];
-				StateMachineClassIndexToWeightOffset.Add(MachineClassIndex, NumStates);
-				NumStates += Machine.States.Num();
-			}
-			StateWeightArrays[0].Reset(NumStates);
-			StateWeightArrays[0].AddZeroed(NumStates);
-			StateWeightArrays[1].Reset(NumStates);
-			StateWeightArrays[1].AddZeroed(NumStates);
-
-			MachineWeightArrays[0].Reset(NumMachines);
-			MachineWeightArrays[0].AddZeroed(NumMachines);
-			MachineWeightArrays[1].Reset(NumMachines);
-			MachineWeightArrays[1].AddZeroed(NumMachines);
+			const FBakedAnimationStateMachine& Machine = BakedMachines[MachineClassIndex];
+			StateMachineClassIndexToWeightOffset.Add(MachineClassIndex, NumStates);
+			NumStates += Machine.States.Num();
 		}
+		StateWeightArrays[0].Reset(NumStates);
+		StateWeightArrays[0].AddZeroed(NumStates);
+		StateWeightArrays[1].Reset(NumStates);
+		StateWeightArrays[1].AddZeroed(NumStates);
+
+		MachineWeightArrays[0].Reset(NumMachines);
+		MachineWeightArrays[0].AddZeroed(NumMachines);
+		MachineWeightArrays[1].Reset(NumMachines);
+		MachineWeightArrays[1].AddZeroed(NumMachines);
 
 #if WITH_EDITORONLY_DATA
 		if (UAnimBlueprint* Blueprint = Cast<UAnimBlueprint>(InAnimInstance->GetClass()->ClassGeneratedBy))
@@ -1227,11 +1224,11 @@ void FAnimInstanceProxy::UpdateAnimation()
 
 	FMemMark Mark(FMemStack::Get());
 
-	// Give class subsystems a crack
-	IAnimClassInterface::ForEachAnimInstanceSubsystemData(CastChecked<UAnimInstance>(GetAnimInstanceObject()), [this](UAnimBlueprintClassSubsystem* InSubsystem, FAnimInstanceSubsystemData& InSubsystemData)
+	// Process internal batched property copies
+	if(AnimClassInterface)
 	{
-		InSubsystem->OnParallelUpdateAnimation(*this, InSubsystemData, CurrentDeltaSeconds);
-	});
+		PropertyAccess::ProcessCopies(GetAnimInstanceObject(), AnimClassInterface->GetPropertyAccessLibrary(), EPropertyAccessCopyBatch::InternalBatched);
+	}
 
 #if WITH_EDITORONLY_DATA
 	UpdatedNodesThisFrame.Reset();
diff --git a/Engine/Source/Runtime/Engine/Private/Animation/AnimNodeBase.cpp b/Engine/Source/Runtime/Engine/Private/Animation/AnimNodeBase.cpp
--- a/Engine/Source/Runtime/Engine/Private/Animation/AnimNodeBase.cpp
+++ b/Engine/Source/Runtime/Engine/Private/Animation/AnimNodeBase.cpp
@@ -5,7 +5,7 @@
 #include "Animation/AnimInstanceProxy.h"
 #include "Animation/AnimTrace.h"
 #include "UObject/CoreObjectVersion.h"
-#include "IPropertyAccess.h"
+#include "PropertyAccess.h"
 
 /////////////////////////////////////////////////////
 // FAnimationBaseContext
@@ -544,29 +544,29 @@ void FNodeDebugData::GetFlattenedDebugData(TArray<FFlattenedDebugData>& Flattene
 
 void FExposedValueHandler::DynamicClassInitialization(TArray<FExposedValueHandler>& InHandlers, UDynamicClass* InDynamicClass)
 {
-	IPropertyAccess* PropertyAccess = IAnimClassInterface::FindSubsystemWithInterface<IPropertyAccess>(IAnimClassInterface::GetFromClass(InDynamicClass));
+	const FPropertyAccessLibrary& PropertyAccessLibrary = IAnimClassInterface::GetFromClass(InDynamicClass)->GetPropertyAccessLibrary();
 
 	for(FExposedValueHandler& Handler : InHandlers)
 	{
-		Handler.Initialize(InDynamicClass, PropertyAccess);
+		Handler.Initialize(InDynamicClass, PropertyAccessLibrary);
 	}
 }
 
 void FExposedValueHandler::ClassInitialization(TArray<FExposedValueHandler>& InHandlers, UObject* InClassDefaultObject)
 {
 	UClass* Class = InClassDefaultObject->GetClass();
-	IPropertyAccess* PropertyAccess = IAnimClassInterface::FindSubsystemWithInterface<IPropertyAccess>(IAnimClassInterface::GetFromClass(Class));
+	const FPropertyAccessLibrary& PropertyAccessLibrary = IAnimClassInterface::GetFromClass(Class)->GetPropertyAccessLibrary();
 
 	for(FExposedValueHandler& Handler : InHandlers)
 	{
 		FAnimNode_Base* AnimNode = Handler.ValueHandlerNodeProperty->ContainerPtrToValuePtr<FAnimNode_Base>(InClassDefaultObject);
 		check(AnimNode);
 		AnimNode->SetExposedValueHandler(&Handler);
-		Handler.Initialize(Class, PropertyAccess);
+		Handler.Initialize(Class, PropertyAccessLibrary);
 	}
 }
 
-void FExposedValueHandler::Initialize(UClass* InClass, IPropertyAccess* InPropertyAccess)
+void FExposedValueHandler::Initialize(UClass* InClass, const FPropertyAccessLibrary& InPropertyAccessLibrary)
 {
 	// bInitialized may no longer be necessary, but leaving alone for now:
 	if (bInitialized)
@@ -594,8 +594,8 @@ void FExposedValueHandler::Initialize(UClass* InClass, IPropertyAccess* InProper
 		Function = NULL;
 	}
 
-	// Cache property access
-	PropertyAccess = InPropertyAccess;
+	// Cache property access library
+	PropertyAccessLibrary = &InPropertyAccessLibrary;
 
 	bInitialized = true;
 }
@@ -609,12 +609,12 @@ void FExposedValueHandler::Execute(const FAnimationBaseContext& Context) const
 
 	if(CopyRecords.Num() > 0)
 	{
-		if(PropertyAccess != nullptr)
+		if(PropertyAccessLibrary != nullptr)
 		{
 			UObject* AnimInstanceObject = Context.AnimInstanceProxy->GetAnimInstanceObject();
 			for(const FExposedValueCopyRecord& CopyRecord : CopyRecords)
 			{
-				PropertyAccess->ProcessCopy(AnimInstanceObject, EPropertyAccessCopyBatch::InternalUnbatched, CopyRecord.CopyIndex, [&CopyRecord](const FProperty* InProperty, void* InAddress)
+				PropertyAccess::ProcessCopy(AnimInstanceObject, *PropertyAccessLibrary, EPropertyAccessCopyBatch::InternalUnbatched, CopyRecord.CopyIndex, [&CopyRecord](const FProperty* InProperty, void* InAddress)
 				{
 					if(CopyRecord.PostCopyOperation == EPostCopyOperation::LogicalNegateBool)
 					{
