PR #7317: Added nullptr check to UAnimGraphNode_ModifyCurve::GetNodeContextMenuActions (Contributed by [REDACTED_PERSON])

diff --git a/Engine/Source/Editor/AnimGraph/Private/AnimGraphNode_ModifyCurve.cpp b/Engine/Source/Editor/AnimGraph/Private/AnimGraphNode_ModifyCurve.cpp
--- a/Engine/Source/Editor/AnimGraph/Private/AnimGraphNode_ModifyCurve.cpp
+++ b/Engine/Source/Editor/AnimGraph/Private/AnimGraphNode_ModifyCurve.cpp
@@ -83,18 +83,22 @@ void UAnimGraphNode_ModifyCurve::GetNodeContextMenuActions(UToolMenu* Menu, UGra
 		// Clicked pin
 		if (Context->Pin != NULL)
 		{
-			// Get proeprty from pin
+			// Get property from pin
 			FProperty* AssociatedProperty;
 			int32 ArrayIndex;
 			GetPinAssociatedProperty(GetFNodeType(), Context->Pin, /*out*/ AssociatedProperty, /*out*/ ArrayIndex);
-			FName PinPropertyName = AssociatedProperty->GetFName();
 
-			if (PinPropertyName  == GET_MEMBER_NAME_CHECKED(FAnimNode_ModifyCurve, CurveValues) && Context->Pin->Direction == EGPD_Input)
+			if (AssociatedProperty != nullptr)
 			{
-				FString PinName = Context->Pin->PinFriendlyName.ToString();
-				FUIAction Action = FUIAction( FExecuteAction::CreateUObject(const_cast<UAnimGraphNode_ModifyCurve*>(this), &UAnimGraphNode_ModifyCurve::RemoveCurvePin, FName(*PinName)) );
-				FText RemovePinLabelText = FText::Format(LOCTEXT("RemoveThisPin", "Remove This Curve Pin: {0}"), FText::FromString(PinName));
-				Section.AddMenuEntry("RemoveThisPin", RemovePinLabelText, LOCTEXT("RemoveThisPinTooltip", "Remove this curve pin from this node"), FSlateIcon(), Action);
+				FName PinPropertyName = AssociatedProperty->GetFName();
+
+				if (PinPropertyName == GET_MEMBER_NAME_CHECKED(FAnimNode_ModifyCurve, CurveValues) && Context->Pin->Direction == EGPD_Input)
+				{
+					FString PinName = Context->Pin->PinFriendlyName.ToString();
+					FUIAction Action = FUIAction(FExecuteAction::CreateUObject(const_cast<UAnimGraphNode_ModifyCurve*>(this), &UAnimGraphNode_ModifyCurve::RemoveCurvePin, FName(*PinName)));
+					FText RemovePinLabelText = FText::Format(LOCTEXT("RemoveThisPin", "Remove This Curve Pin: {0}"), FText::FromString(PinName));
+					Section.AddMenuEntry("RemoveThisPin", RemovePinLabelText, LOCTEXT("RemoveThisPinTooltip", "Remove this curve pin from this node"), FSlateIcon(), Action);
+				}
 			}
 		}
 
