Re-adding SyncGroupNameToRequireValidMarkersRule to baked anim state transition rules

This was merged out in CL 14113036 when resolving a merge to staging

diff --git a/Engine/Source/Editor/AnimGraph/Private/AnimGraphNode_StateMachineBase.cpp b/Engine/Source/Editor/AnimGraph/Private/AnimGraphNode_StateMachineBase.cpp
--- a/Engine/Source/Editor/AnimGraph/Private/AnimGraphNode_StateMachineBase.cpp
+++ b/Engine/Source/Editor/AnimGraph/Private/AnimGraphNode_StateMachineBase.cpp
@@ -557,6 +557,7 @@ void UAnimGraphNode_StateMachineBase::OnProcessDuringCompilation(IAnimBlueprintC
 
 			// Handle automatic time remaining rules
 			Rule.bAutomaticRemainingTimeRule = TransitionNode->bAutomaticRuleBasedOnSequencePlayerInState;
+			Rule.SyncGroupNameToRequireValidMarkersRule = TransitionNode->SyncGroupNameToRequireValidMarkersRule;
 
 			// Handle custom transition graphs
 			Rule.CustomResultNodeIndex = INDEX_NONE;
