Relocating the contents of the Classes directory in AnimGraph in favor of Public.

#trivial

diff --git a/Engine/Source/Editor/AnimGraph/Private/AnimBlueprintPostCompileValidation.cpp b/Engine/Source/Editor/AnimGraph/Private/AnimBlueprintPostCompileValidation.cpp
--- a/Engine/Source/Editor/AnimGraph/Private/AnimBlueprintPostCompileValidation.cpp
+++ b/Engine/Source/Editor/AnimGraph/Private/AnimBlueprintPostCompileValidation.cpp
@@ -5,7 +5,7 @@
 #include "Animation/BlendSpaceBase.h"
 #include "Animation/AnimNode_SequencePlayer.h"
 #include "Animation/AnimInstance.h"
-#include "Editor/AnimGraph/Classes/AnimGraphNode_Base.h"
+#include "AnimGraphNode_Base.h"
 
 UAnimBlueprintPostCompileValidation::UAnimBlueprintPostCompileValidation(const FObjectInitializer& ObjectInitializer)
 	: Super(ObjectInitializer)
diff --git a/Engine/Source/Editor/AnimGraph/Public/AnimBlueprintPostCompileValidation.h b/Engine/Source/Editor/AnimGraph/Public/AnimBlueprintPostCompileValidation.h
new file mode 100644
--- /dev/null
+++ b/Engine/Source/Editor/AnimGraph/Public/AnimBlueprintPostCompileValidation.h
@@ -0,0 +1,127 @@
+// Copyright Epic Games, Inc. All Rights Reserved.
+
+#pragma once
+
+#include "CoreMinimal.h"
+#include "UObject/ObjectMacros.h"
+#include "UObject/Object.h"
+#include "Animation/AnimSequence.h"
+#include "Animation/BlendSpace.h"
+#include "Animation/AnimBlueprintGeneratedClass.h"
+#include "Kismet2/CompilerResultsLog.h"
+
+#include "AnimBlueprintPostCompileValidation.generated.h"
+
+// Encapsulated Parameters, to upgrade these without changing function signature.
+struct FAnimBPCompileValidationParams
+{
+	const class UAnimInstance* const DefaultAnimInstance;
+	const class UAnimBlueprintGeneratedClass* const NewAnimBlueprintClass;
+	FCompilerResultsLog& MessageLog;
+	const TMap<FProperty*, class UAnimGraphNode_Base*>& AllocatedNodePropertiesToNodes;
+
+	FAnimBPCompileValidationParams
+	(
+		const class UAnimInstance* const InDefaultAnimInstance,
+		const class UAnimBlueprintGeneratedClass* const InNewAnimBlueprintClass,
+		FCompilerResultsLog& InMessageLog,
+		const TMap<FProperty*, class UAnimGraphNode_Base*>& InAllocatedNodePropertiesToNodes
+	)
+		: DefaultAnimInstance(InDefaultAnimInstance)
+		, NewAnimBlueprintClass(InNewAnimBlueprintClass)
+		, MessageLog(InMessageLog)
+		, AllocatedNodePropertiesToNodes(InAllocatedNodePropertiesToNodes)
+	{}
+};
+
+// This class is a base class for performing AnimBlueprint Post Compilation Validation.
+UCLASS()
+class ANIMGRAPH_API UAnimBlueprintPostCompileValidation : public UObject
+{
+	GENERATED_UCLASS_BODY()
+
+public:
+	virtual void DoPostCompileValidation(FAnimBPCompileValidationParams& InParams) const;
+
+protected:
+	static void PCV_PreloadObject(const UObject* const ReferencedObject);
+
+	static void PCV_GatherAnimSequences(TArray<const UAnimSequence*>& OutAnimSequences, const UAnimSequenceBase* const InAnimSequenceBase);
+	static void PCV_GatherAnimSequences(TArray<const UAnimSequence*>& OutAnimSequences, const class UBlendSpaceBase* const InBlendSpace);
+
+	struct FPCV_GatherParams
+	{
+		bool bFilterBySyncGroup;
+		FName SyncGroupName;
+		bool bFilterByLoopingCondition;
+		bool bLoopingCondition;
+
+		UE_DEPRECATED(4.26, "Please use the constructor that takes an FName for the sync group")
+		FPCV_GatherParams
+		(
+			bool InbFilterBySyncGroup = false,
+			int32 InSyncGroupIndex = INDEX_NONE,
+			bool InbFilterByLoopingCondition = false,
+			bool InbLoopingCondition = false
+		)
+			: bFilterBySyncGroup(InbFilterBySyncGroup)
+			, SyncGroupName(NAME_None)
+			, bFilterByLoopingCondition(InbFilterByLoopingCondition)
+			, bLoopingCondition(InbLoopingCondition)
+		{}
+
+		FPCV_GatherParams
+		(
+			bool InbFilterBySyncGroup = false,
+			FName InSyncGroupName = NAME_None,
+			bool InbFilterByLoopingCondition = false,
+			bool InbLoopingCondition = false
+		)
+			: bFilterBySyncGroup(InbFilterBySyncGroup)
+			, SyncGroupName(InSyncGroupName)
+			, bFilterByLoopingCondition(InbFilterByLoopingCondition)
+			, bLoopingCondition(InbLoopingCondition)
+		{}
+	};
+
+	static void PCV_GatherAnimSequencesFromGraph(TArray<const UAnimSequence*>& OutAnimSequences, FAnimBPCompileValidationParams& PCV_Params, const FPCV_GatherParams& GatherParams);
+	static void PCV_GatherBlendSpacesFromGraph(TArray<const class UBlendSpaceBase*>& OutBlendSpaces, FAnimBPCompileValidationParams& PCV_Params);
+
+	struct FPCV_ReferencedAnimSequence
+	{
+		const UAnimSequence* AnimSequence;
+		const UObject* Referencer;
+
+		FPCV_ReferencedAnimSequence
+		(
+			const UAnimSequence* InAnimSequence,
+			const UObject* InReferencer
+		)
+			: AnimSequence(InAnimSequence)
+			, Referencer(InReferencer)
+		{}
+	};
+
+	struct FPCV_PropertyAndValue
+	{
+		const FProperty* Property;
+		const void* Value;
+
+		FPCV_PropertyAndValue
+		(
+			const FProperty* InProperty,
+			const void* InValue
+		)
+			: Property(InProperty)
+			, Value(InValue)
+		{}
+	};
+
+	static void PCV_GatherAllReferencedAnimSequences(TArray<FPCV_ReferencedAnimSequence>& OutRefAnimSequences, FAnimBPCompileValidationParams& PCV_Params);
+	static void PCV_GatherAnimSequencesFromStruct(TArray<FPCV_ReferencedAnimSequence>& OutRefAnimSequences, FAnimBPCompileValidationParams& PCV_Params, const UStruct* InStruct, const void* InData, TArray<FPCV_PropertyAndValue> InPropertyCallChain);
+	static void PCV_GatherAnimSequencesFromProperty(TArray<FPCV_ReferencedAnimSequence>& OutRefAnimSequences, FAnimBPCompileValidationParams& PCV_Params, const FProperty* InProperty, const void* InData, TArray<FPCV_PropertyAndValue> InPropertyCallChain);
+
+private:
+	virtual bool NeedsLoadForClient() const override { return false; }
+	virtual bool NeedsLoadForServer() const override { return false; }
+};
diff --git a/Engine/Source/Editor/AnimGraph/Public/AnimGraphNode_AssetPlayerBase.h b/Engine/Source/Editor/AnimGraph/Public/AnimGraphNode_AssetPlayerBase.h
new file mode 100644
--- /dev/null
+++ b/Engine/Source/Editor/AnimGraph/Public/AnimGraphNode_AssetPlayerBase.h
@@ -0,0 +1,35 @@
+// Copyright Epic Games, Inc. All Rights Reserved.
+
+#pragma once
+
+#include "CoreMinimal.h"
+#include "UObject/ObjectMacros.h"
+#include "Animation/AnimationAsset.h"
+#include "AnimGraphNode_Base.h"
+#include "AnimGraphNode_AssetPlayerBase.generated.h"
+
+/** Get the default anim node class for playing a particular asset */
+ANIMGRAPH_API UClass* GetNodeClassForAsset(const UClass* AssetClass);
+
+/** See if a particular anim NodeClass can play a particular anim AssetClass */
+ANIMGRAPH_API bool SupportNodeClassForAsset(const UClass* AssetClass, UClass* NodeClass);
+
+/** Helper / intermediate for asset player graphical nodes */
+UCLASS(Abstract)
+class ANIMGRAPH_API UAnimGraphNode_AssetPlayerBase : public UAnimGraphNode_Base
+{
+	GENERATED_BODY()
+public:
+	// Sync group settings for this player.  Sync groups keep related animations with different lengths synchronized.
+	UPROPERTY(EditAnywhere, Category = Settings)
+	FAnimationGroupReference SyncGroup;
+
+	/** UEdGraphNode interface */
+	virtual void PinConnectionListChanged(UEdGraphPin* Pin) override;
+	virtual void PinDefaultValueChanged(UEdGraphPin* Pin) override;
+
+	/** UAnimGraphNode_Base interface */
+	virtual void OnProcessDuringCompilation(IAnimBlueprintCompilationContext& InCompilationContext, IAnimBlueprintGeneratedClassCompiledData& OutCompiledData) override;
+
+	virtual void SetAnimationAsset(UAnimationAsset* Asset) { check(false); /*Base function called*/ }
+};
diff --git a/Engine/Source/Editor/AnimGraph/Public/AnimGraphNode_BlendListBase.h b/Engine/Source/Editor/AnimGraph/Public/AnimGraphNode_BlendListBase.h
new file mode 100644
--- /dev/null
+++ b/Engine/Source/Editor/AnimGraph/Public/AnimGraphNode_BlendListBase.h
@@ -0,0 +1,30 @@
+// Copyright Epic Games, Inc. All Rights Reserved.
+
+#pragma once
+
+#include "CoreMinimal.h"
+#include "UObject/ObjectMacros.h"
+#include "AnimGraphNode_Base.h"
+#include "AnimGraphNode_BlendListBase.generated.h"
+
+UCLASS(MinimalAPI, Abstract)
+class UAnimGraphNode_BlendListBase : public UAnimGraphNode_Base
+{
+	GENERATED_UCLASS_BODY()
+
+	// UEdGraphNode interface
+	ANIMGRAPH_API virtual FLinearColor GetNodeTitleColor() const override;
+	ANIMGRAPH_API virtual void PostEditChangeProperty(struct FPropertyChangedEvent& PropertyChangedEvent) override;
+	ANIMGRAPH_API virtual void ReallocatePinsDuringReconstruction(TArray<UEdGraphPin*>& OldPins) override;
+	// End of UEdGraphNode interface
+
+	// UAnimGraphNode_Base interface
+	virtual FString GetNodeCategory() const override;
+	// End of UAnimGraphNode_Base interface
+
+protected:
+	// removes removed pins and adjusts array indices of remained pins
+	void RemovePinsFromOldPins(TArray<UEdGraphPin*>& OldPins, int32 RemovedArrayIndex);
+
+	int32 RemovedPinArrayIndex;
+};
diff --git a/Engine/Source/Editor/AnimGraph/Public/AnimGraphNode_BoneDrivenController.h b/Engine/Source/Editor/AnimGraph/Public/AnimGraphNode_BoneDrivenController.h
new file mode 100644
--- /dev/null
+++ b/Engine/Source/Editor/AnimGraph/Public/AnimGraphNode_BoneDrivenController.h
@@ -0,0 +1,68 @@
+// Copyright Epic Games, Inc. All Rights Reserved.
+
+#pragma once
+
+#include "CoreMinimal.h"
+#include "UObject/ObjectMacros.h"
+#include "Misc/Attribute.h"
+#include "Layout/Visibility.h"
+#include "AnimGraphNode_SkeletalControlBase.h"
+#include "BoneControllers/AnimNode_BoneDrivenController.h"
+#include "AnimGraphNode_BoneDrivenController.generated.h"
+
+class FCompilerResultsLog;
+class FPrimitiveDrawInterface;
+class IDetailCategoryBuilder;
+class IDetailLayoutBuilder;
+class IPropertyHandle;
+class USkeletalMeshComponent;
+
+/**
+ * This is the 'source version' of a bone driven controller, which maps part of the state from one bone to another (e.g., 2 * source.x -> target.z)
+ */
+UCLASS()
+class ANIMGRAPH_API UAnimGraphNode_BoneDrivenController : public UAnimGraphNode_SkeletalControlBase
+{
+	GENERATED_UCLASS_BODY()
+
+	UPROPERTY(EditAnywhere, Category = Settings)
+	FAnimNode_BoneDrivenController Node;
+
+public:
+	// UObject interface
+	virtual void Serialize(FArchive& Ar) override;
+	// End of UObject interface
+
+	// UEdGraphNode interface
+	virtual FText GetNodeTitle(ENodeTitleType::Type TitleType) const override;
+	virtual FText GetTooltipText() const override;
+	// End of UEdGraphNode interface
+
+	// UAnimGraphNode_Base interface
+	virtual void ValidateAnimNodeDuringCompilation(USkeleton* ForSkeleton, FCompilerResultsLog& MessageLog) override;
+	virtual void CustomizeDetails(IDetailLayoutBuilder& DetailBuilder) override;
+	// End of UAnimGraphNode_Base interface
+
+	// UAnimGraphNode_SkeletalControlBase interface
+	virtual void Draw(FPrimitiveDrawInterface* PDI, USkeletalMeshComponent* SkelMeshComp) const override;
+	// End of UAnimGraphNode_SkeletalControlBase interface
+
+protected:
+
+	// UAnimGraphNode_SkeletalControlBase protected interface
+	virtual FText GetControllerDescription() const override;
+	virtual const FAnimNode_SkeletalControlBase* GetNode() const override { return &Node; }
+	// End of UAnimGraphNode_SkeletalControlBase protected interface
+
+	// Should non-curve mapping values be shown (multiplier, range)?
+	static EVisibility AreNonCurveMappingValuesVisible(IDetailLayoutBuilder* DetailLayoutBuilder);
+	static EVisibility AreRemappingValuesVisible(IDetailLayoutBuilder* DetailLayoutBuilder);
+
+	// Should destination bone or morph target properties be visible
+	static EVisibility AreTargetBonePropertiesVisible(IDetailLayoutBuilder* DetailLayoutBuilder);
+	static EVisibility AreTargetCurvePropertiesVisible(IDetailLayoutBuilder* DetailLayoutBuilder);
+
+	static void AddTripletPropertyRow(const FText& Name, const FText& Tooltip, IDetailCategoryBuilder& Category, TSharedRef<IPropertyHandle> PropertyHandle, const FName XPropertyName, const FName YPropertyName, const FName ZPropertyName, TAttribute<EVisibility> VisibilityAttribute);
+	static void AddRangePropertyRow(const FText& Name, const FText& Tooltip, IDetailCategoryBuilder& Category, TSharedRef<IPropertyHandle> PropertyHandle, const FName MinPropertyName, const FName MaxPropertyName, TAttribute<EVisibility> VisibilityAttribute);
+	static FText ComponentTypeToText(EComponentType::Type Component);
+};
diff --git a/Engine/Source/Editor/AnimGraph/Public/AnimGraphNode_Constraint.h b/Engine/Source/Editor/AnimGraph/Public/AnimGraphNode_Constraint.h
new file mode 100644
--- /dev/null
+++ b/Engine/Source/Editor/AnimGraph/Public/AnimGraphNode_Constraint.h
@@ -0,0 +1,44 @@
+// Copyright Epic Games, Inc. All Rights Reserved.
+
+#pragma once
+
+#include "AnimGraphNode_SkeletalControlBase.h"
+#include "BoneControllers/AnimNode_Constraint.h"
+#include "EdGraph/EdGraphNodeUtils.h" // for FNodeTitleTextTable
+#include "AnimGraphNode_Constraint.generated.h"
+
+UCLASS(Experimental)
+class ANIMGRAPH_API UAnimGraphNode_Constraint : public UAnimGraphNode_SkeletalControlBase
+{
+	GENERATED_UCLASS_BODY()
+
+	UPROPERTY(EditAnywhere, Category=Settings)
+	FAnimNode_Constraint Node;
+
+public:
+	// UEdGraphNode interface
+	virtual FText GetNodeTitle(ENodeTitleType::Type TitleType) const override;
+	virtual FText GetTooltipText() const override;
+	// End of UEdGraphNode interface
+
+	// UObject interface
+	virtual void PostEditChangeProperty(struct FPropertyChangedEvent& PropertyChangedEvent) override;
+	// End of UObject interface
+
+protected:
+	// UAnimGraphNode_Base interface
+	virtual void ValidateAnimNodeDuringCompilation(USkeleton* ForSkeleton, FCompilerResultsLog& MessageLog) override;
+	virtual void PostProcessPinName(const UEdGraphPin* Pin, FString& DisplayName) const override;
+	// End of UAnimGraphNode_Base interface
+
+	// UAnimGraphNode_SkeletalControlBase interface
+	virtual void Draw(FPrimitiveDrawInterface* PDI, USkeletalMeshComponent* SkelMeshComp) const override;
+	virtual FText GetControllerDescription() const override;
+	virtual const FAnimNode_SkeletalControlBase* GetNode() const override { return &Node; }
+	// End of UAnimGraphNode_SkeletalControlBase interface
+
+private:
+	/** Constructing FText strings can be costly, so we cache the node's title */
+	FNodeTitleTextTable CachedNodeTitles;
+};
+
diff --git a/Engine/Source/Editor/AnimGraph/Public/AnimGraphNode_CustomProperty.h b/Engine/Source/Editor/AnimGraph/Public/AnimGraphNode_CustomProperty.h
new file mode 100644
--- /dev/null
+++ b/Engine/Source/Editor/AnimGraph/Public/AnimGraphNode_CustomProperty.h
@@ -0,0 +1,93 @@
+// Copyright Epic Games, Inc. All Rights Reserved.
+
+#pragma once
+
+#include "CoreMinimal.h"
+#include "UObject/ObjectMacros.h"
+#include "Misc/Guid.h"
+#include "AnimGraphNode_Base.h"
+#include "Animation/AnimNode_CustomProperty.h"
+#include "IClassVariableCreator.h"
+
+#include "AnimGraphNode_CustomProperty.generated.h"
+
+class FCompilerResultsLog;
+class IDetailLayoutBuilder;
+class IPropertyHandle;
+class FAnimBlueprintCompilerHandler;
+
+UCLASS(Abstract)
+class ANIMGRAPH_API UAnimGraphNode_CustomProperty : public UAnimGraphNode_Base, public IClassVariableCreator
+{
+	GENERATED_BODY()
+
+public:
+
+	// IClassVariableCreator interface
+	virtual void CreateClassVariablesFromBlueprint(IAnimBlueprintVariableCreationContext& InCreationContext) override;
+
+	//~ Begin UEdGraphNode Interface.
+	virtual void ValidateAnimNodeDuringCompilation(USkeleton* ForSkeleton, FCompilerResultsLog& MessageLog) override;
+	virtual void ReallocatePinsDuringReconstruction(TArray<UEdGraphPin*>& OldPins) override;
+	virtual UObject* GetJumpTargetForDoubleClick() const override;
+	virtual bool HasExternalDependencies(TArray<class UStruct*>* OptionalOutput /*= NULL*/) const override;
+	//~ End UEdGraphNode Interface.
+
+	// UAnimGraphNode_Base interface
+	virtual void CustomizeDetails(IDetailLayoutBuilder& DetailBuilder) override;
+	virtual void OnProcessDuringCompilation(IAnimBlueprintCompilationContext& InCompilationContext, IAnimBlueprintGeneratedClassCompiledData& OutCompiledData) override;
+
+	// Gets the property on InOwnerInstanceClass that corresponds to InInputPin
+	void GetInstancePinProperty(const IAnimBlueprintCompilationContext& InCompilationContext, UEdGraphPin* InInputPin, FProperty*& OutProperty);
+	// Gets the unique name for the property linked to a given pin
+	FString GetPinTargetVariableName(const UEdGraphPin* InPin) const;
+	// Gets Target Class this properties to link
+	UClass* GetTargetClass() const;
+	// Add Source and Target Properties - Check FAnimNode_CustomProperty
+	void AddSourceTargetProperties(const FName& InSourcePropertyName, const FName& InTargetPropertyName);
+	// Helper used to get the skeleton class we are targeting
+	virtual UClass* GetTargetSkeletonClass() const;
+
+	// ----- UI CALLBACKS ----- //
+	// User changed the instance class etc.
+	void OnStructuralPropertyChanged(IDetailLayoutBuilder* DetailBuilder);
+	// If given property exposed on this node
+	virtual ECheckBoxState IsPropertyExposed(FName PropertyName) const;
+	// User chose to expose, or unexpose a property
+	virtual void OnPropertyExposeCheckboxChanged(ECheckBoxState NewState, FName PropertyName);
+	// If all possible properties are exposed on this node
+	virtual ECheckBoxState AreAllPropertiesExposed() const;
+	// User chose to expose, or unexpose all properties
+	virtual void OnPropertyExposeAllCheckboxChanged(ECheckBoxState NewState);
+	// User changed the instance class
+	void OnInstanceClassChanged(IDetailLayoutBuilder* DetailBuilder);
+protected:
+
+	/** List of property names we know to exist on the target class, so we can detect when
+	 *  Properties are added or removed on reconstruction
+	 */
+	UPROPERTY()
+	TArray<FName> KnownExposableProperties;
+
+	/** Names of properties the user has chosen to expose */
+	UPROPERTY()
+	TArray<FName> ExposedPropertyNames;
+
+	// Searches the instance class for properties that we can expose (public and BP visible)
+	virtual void GetExposableProperties(TArray<FProperty*>& OutExposableProperties) const;
+	// Gets a property's type as FText (for UI)
+	FText GetPropertyTypeText(FProperty* Property);
+	// Given a new class, rebuild the known property list (for tracking class changes and moving pins)
+	virtual void RebuildExposedProperties();
+
+	// internal node accessor
+	virtual FAnimNode_CustomProperty* GetCustomPropertyNode() PURE_VIRTUAL(UAnimGraphNode_CustomProperty::GetCustomPropertyNode, return nullptr;);
+	virtual const FAnimNode_CustomProperty* GetCustomPropertyNode() const PURE_VIRTUAL(UAnimGraphNode_CustomProperty::GetCustomPropertyNode, return nullptr;);
+
+	// Check whether the specified property is structural (i.e. should we rebuild the UI if it changes)
+	virtual bool IsStructuralProperty(FProperty* InProperty) const { return false; }
+
+	// Whether this node needs a valid target class up-front
+	virtual bool NeedsToSpecifyValidTargetClass() const { return true; }
+
+};
diff --git a/Engine/Source/Editor/AnimGraph/Public/AnimGraphNode_ModifyBone.h b/Engine/Source/Editor/AnimGraph/Public/AnimGraphNode_ModifyBone.h
new file mode 100644
--- /dev/null
+++ b/Engine/Source/Editor/AnimGraph/Public/AnimGraphNode_ModifyBone.h
@@ -0,0 +1,48 @@
+// Copyright Epic Games, Inc. All Rights Reserved.
+
+#pragma once
+
+#include "CoreMinimal.h"
+#include "UObject/ObjectMacros.h"
+#include "BoneControllers/AnimNode_ModifyBone.h"
+#include "EdGraph/EdGraphNodeUtils.h"
+#include "AnimGraphNode_SkeletalControlBase.h"
+#include "AnimGraphNode_ModifyBone.generated.h"
+
+class FCompilerResultsLog;
+
+UCLASS(meta=(Keywords = "Modify Transform"))
+class ANIMGRAPH_API UAnimGraphNode_ModifyBone : public UAnimGraphNode_SkeletalControlBase
+{
+	GENERATED_UCLASS_BODY()
+
+	UPROPERTY(EditAnywhere, Category=Settings)
+	FAnimNode_ModifyBone Node;
+
+public:
+	// UEdGraphNode interface
+	virtual FText GetNodeTitle(ENodeTitleType::Type TitleType) const override;
+	virtual FText GetTooltipText() const override;
+	// End of UEdGraphNode interface
+
+protected:
+	// UAnimGraphNode_Base interface
+	virtual void ValidateAnimNodeDuringCompilation(USkeleton* ForSkeleton, FCompilerResultsLog& MessageLog) override;
+	virtual FEditorModeID GetEditorMode() const override;
+	virtual void CopyNodeDataToPreviewNode(FAnimNode_Base* InPreviewNode) override;
+	virtual void CopyPinDefaultsToNodeData(UEdGraphPin* InPin) override;
+	// End of UAnimGraphNode_Base interface
+
+	// UAnimGraphNode_SkeletalControlBase interface
+	virtual FText GetControllerDescription() const override;
+	virtual const FAnimNode_SkeletalControlBase* GetNode() const override { return &Node; }
+	// End of UAnimGraphNode_SkeletalControlBase interface
+
+private:
+	/** Constructing FText strings can be costly, so we cache the node's title */
+	FNodeTitleTextTable CachedNodeTitles;
+
+	// storing current widget mode 
+	int32 CurWidgetMode;
+};
+
diff --git a/Engine/Source/Editor/AnimGraph/Public/AnimGraphNode_PoseDriver.h b/Engine/Source/Editor/AnimGraph/Public/AnimGraphNode_PoseDriver.h
new file mode 100644
--- /dev/null
+++ b/Engine/Source/Editor/AnimGraph/Public/AnimGraphNode_PoseDriver.h
@@ -0,0 +1,91 @@
+// Copyright Epic Games, Inc. All Rights Reserved.
+
+#pragma once
+
+#include "CoreMinimal.h"
+#include "UObject/ObjectMacros.h"
+#include "AnimNodes/AnimNode_PoseDriver.h"
+#include "AnimGraphNode_PoseHandler.h"
+#include "EdGraph/EdGraphNodeUtils.h"
+#include "AnimGraphNode_PoseDriver.generated.h"
+
+class FCompilerResultsLog;
+
+UCLASS(Experimental, BlueprintType)
+class ANIMGRAPH_API UAnimGraphNode_PoseDriver : public UAnimGraphNode_PoseHandler
+{
+	GENERATED_UCLASS_BODY()
+
+	UPROPERTY(EditAnywhere, BlueprintReadWrite, Category = Settings)
+	FAnimNode_PoseDriver Node;
+
+	/** Length of axis in world units used for debug drawing */
+	UPROPERTY(EditAnywhere, Category = Debugging, meta = (DefaultValue = "20.0", UIMin = "1.0", UIMax = "100.0"))
+	float AxisLength;
+
+	/** Number of subdivisions / lines used when debug drawing a cone */
+	UPROPERTY(EditAnywhere, Category = Debugging, meta = (DefaultValue = "32", UIMin = "6", UIMax = "128"))
+	int32 ConeSubdivision;
+
+	/** If checked the cones will be drawn in 3d for debugging */
+	UPROPERTY(EditAnywhere, Category = Debugging, meta = (DefaultValue = "True"))
+	bool bDrawDebugCones;
+
+	/** Used to indicate selected target to edit mode drawing */
+	int32 SelectedTargetIndex;
+	/** Delegate to call when selection changes */
+	FSimpleMulticastDelegate SelectedTargetChangeDelegate;
+
+public:
+
+	/** Get the current preview node instance */
+	FAnimNode_PoseDriver* GetPreviewPoseDriverNode() const;
+
+	/** Util to replace current contents of PoseTargets with info from assigned PoseAsset */
+	void CopyTargetsFromPoseAsset();
+
+	/** Automatically modify TargetScale for each PoseTarget, based on distance to nearest neighbor */
+	void AutoSetTargetScales(float& OutMaxDistance);
+
+	/** Adds a new target, reallocating transforms array appropriately */
+	void AddNewTarget();
+
+	/** Reallocates transforms arrays as necessary to accommodate source bones */
+	void ReserveTargetTransforms();
+
+	/** Return the color for a given weight. Used for Details and EditMode */
+	FLinearColor GetColorFromWeight(float InWeight);
+
+	/** Used to refer back to preview instance in anim tools */
+	UPROPERTY(Transient)
+	USkeletalMeshComponent* LastPreviewComponent;
+
+	// Begin UObject Interface.
+	virtual void PostLoad() override;
+	virtual void Serialize(FArchive& Ar) override;
+	// End UObject Interface.
+
+	// UEdGraphNode interface
+	virtual FText GetNodeTitle(ENodeTitleType::Type TitleType) const override;
+	virtual FText GetTooltipText() const override;
+	virtual FText GetMenuCategory() const override;
+	virtual FLinearColor GetNodeBodyTintColor() const override;
+	// End of UEdGraphNode interface
+
+	// UAnimGraphNode_Base interface
+	virtual void ValidateAnimNodeDuringCompilation(USkeleton* ForSkeleton, FCompilerResultsLog& MessageLog) override;
+	virtual FEditorModeID GetEditorMode() const override;
+	virtual EAnimAssetHandlerType SupportsAssetClass(const UClass* AssetClass) const override;
+	virtual void CopyNodeDataToPreviewNode(FAnimNode_Base* InPreviewNode) override;
+	// End of UAnimGraphNode_Base interface
+
+protected:
+	// UAnimGraphNode_PoseHandler interface
+	virtual bool IsPoseAssetRequired() { return false; }
+	virtual FAnimNode_PoseHandler* GetPoseHandlerNode() override { return &Node; }
+	virtual const FAnimNode_PoseHandler* GetPoseHandlerNode() const override { return &Node; }
+	// End of UAnimGraphNode_PoseHandler interface
+
+	/** Constructing FText strings can be costly, so we cache the node's title */
+	FNodeTitleTextTable CachedNodeTitles;
+};
diff --git a/Engine/Source/Editor/AnimGraph/Public/AnimGraphNode_Root.h b/Engine/Source/Editor/AnimGraph/Public/AnimGraphNode_Root.h
new file mode 100644
--- /dev/null
+++ b/Engine/Source/Editor/AnimGraph/Public/AnimGraphNode_Root.h
@@ -0,0 +1,39 @@
+// Copyright Epic Games, Inc. All Rights Reserved.
+
+#pragma once
+
+#include "CoreMinimal.h"
+#include "UObject/ObjectMacros.h"
+#include "AnimGraphNode_Base.h"
+#include "Animation/AnimNode_Root.h"
+#include "AnimGraphNode_Root.generated.h"
+
+class FBlueprintActionDatabaseRegistrar;
+
+UCLASS(MinimalAPI)
+class UAnimGraphNode_Root : public UAnimGraphNode_Base
+{
+	GENERATED_UCLASS_BODY()
+
+	UPROPERTY(EditAnywhere, Category=Settings)
+	FAnimNode_Root Node;
+
+	//~ Begin UEdGraphNode Interface.
+	virtual FLinearColor GetNodeTitleColor() const override;
+	virtual FText GetTooltipText() const override;
+	virtual FText GetNodeTitle(ENodeTitleType::Type TitleType) const override;
+	virtual bool CanUserDeleteNode() const override { return false; }
+	virtual bool CanDuplicateNode() const override { return false; }
+	virtual void GetMenuActions(FBlueprintActionDatabaseRegistrar& ActionRegistrar) const override;
+	virtual bool IsNodeRootSet() const override { return true; }
+	//~ End UEdGraphNode Interface.
+
+	//~ Begin UAnimGraphNode_Base Interface
+	virtual bool IsSinkNode() const override;
+	virtual void OnProcessDuringCompilation(IAnimBlueprintCompilationContext& InCompilationContext, IAnimBlueprintGeneratedClassCompiledData& OutCompiledData) override;
+
+	// Get the link to the documentation
+	virtual FString GetDocumentationLink() const override;
+
+	//~ End UAnimGraphNode_Base Interface
+};
diff --git a/Engine/Source/Editor/AnimGraph/Public/AnimGraphNode_StateMachineBase.h b/Engine/Source/Editor/AnimGraph/Public/AnimGraphNode_StateMachineBase.h
new file mode 100644
--- /dev/null
+++ b/Engine/Source/Editor/AnimGraph/Public/AnimGraphNode_StateMachineBase.h
@@ -0,0 +1,52 @@
+// Copyright Epic Games, Inc. All Rights Reserved.
+
+#pragma once
+
+#include "CoreMinimal.h"
+#include "UObject/ObjectMacros.h"
+#include "AnimGraphNode_Base.h"
+#include "EdGraph/EdGraphNodeUtils.h"
+#include "Animation/AnimNode_StateMachine.h"
+#include "AnimGraphNode_StateMachineBase.generated.h"
+
+class INameValidatorInterface;
+
+UCLASS(Abstract)
+class ANIMGRAPH_API UAnimGraphNode_StateMachineBase : public UAnimGraphNode_Base
+{
+	GENERATED_UCLASS_BODY()
+
+	// Editor state machine representation
+	UPROPERTY()
+	class UAnimationStateMachineGraph* EditorStateMachineGraph;
+
+	// UEdGraphNode interface
+	virtual FLinearColor GetNodeTitleColor() const override;
+	virtual FText GetTooltipText() const override;
+	virtual FText GetNodeTitle(ENodeTitleType::Type TitleType) const override;
+	virtual void PostPlacedNewNode() override;
+	virtual UObject* GetJumpTargetForDoubleClick() const override;
+	virtual void JumpToDefinition() const override;
+	virtual void DestroyNode() override;
+	virtual void PostPasteNode() override;
+	virtual TSharedPtr<class INameValidatorInterface> MakeNameValidator() const override;
+	virtual FString GetDocumentationLink() const override;
+	virtual void OnRenameNode(const FString& NewName) override;
+	// End of UEdGraphNode interface
+
+	// UAnimGraphNode_Base interface
+	virtual FString GetNodeCategory() const override;
+	virtual void OnProcessDuringCompilation(IAnimBlueprintCompilationContext& InCompilationContext, IAnimBlueprintGeneratedClassCompiledData& OutCompiledData) override;
+	// End of UAnimGraphNode_Base interface
+
+	//  @return the name of this state machine
+	FString GetStateMachineName();
+
+	// Interface for derived classes to implement
+	virtual FAnimNode_StateMachine& GetNode() PURE_VIRTUAL(UAnimGraphNode_StateMachineBase::GetNode, static FAnimNode_StateMachine Dummy; return Dummy;);
+	// End of my interface
+
+private:
+	/** Constructing FText strings can be costly, so we cache the node's title */
+	FNodeTextCache CachedFullTitle;
+};
diff --git a/Engine/Source/Editor/AnimGraph/Public/AnimGraphNode_TwistCorrectiveNode.h b/Engine/Source/Editor/AnimGraph/Public/AnimGraphNode_TwistCorrectiveNode.h
new file mode 100644
--- /dev/null
+++ b/Engine/Source/Editor/AnimGraph/Public/AnimGraphNode_TwistCorrectiveNode.h
@@ -0,0 +1,39 @@
+// Copyright Epic Games, Inc. All Rights Reserved.
+
+#pragma once
+
+#include "CoreMinimal.h"
+#include "UObject/ObjectMacros.h"
+#include "AnimGraphNode_SkeletalControlBase.h"
+#include "BoneControllers/AnimNode_TwistCorrectiveNode.h"
+#include "AnimGraphNode_TwistCorrectiveNode.generated.h"
+
+class FCompilerResultsLog;
+
+/**
+ * This is the 'source version' of a bone driven controller, which maps part of the state from one bone to another (e.g., 2 * source.x -> target.z)
+ */
+UCLASS()
+class ANIMGRAPH_API UAnimGraphNode_TwistCorrectiveNode : public UAnimGraphNode_SkeletalControlBase
+{
+	GENERATED_UCLASS_BODY()
+
+	UPROPERTY(EditAnywhere, Category = Settings)
+	FAnimNode_TwistCorrectiveNode Node;
+
+public:
+
+	// UEdGraphNode interface
+	virtual FText GetNodeTitle(ENodeTitleType::Type TitleType) const override;
+	virtual FText GetTooltipText() const override;
+	// End of UEdGraphNode interface
+
+protected:
+
+	// UAnimGraphNode_SkeletalControlBase protected interface
+	virtual FText GetControllerDescription() const override;
+	virtual void ValidateAnimNodeDuringCompilation(USkeleton* ForSkeleton, FCompilerResultsLog& MessageLog) override;
+	virtual const FAnimNode_SkeletalControlBase* GetNode() const override { return &Node; }
+	// End of UAnimGraphNode_SkeletalControlBase protected interface
+
+};
diff --git a/Engine/Source/Editor/AnimGraph/Public/AnimPreviewInstance.h b/Engine/Source/Editor/AnimGraph/Public/AnimPreviewInstance.h
new file mode 100644
--- /dev/null
+++ b/Engine/Source/Editor/AnimGraph/Public/AnimPreviewInstance.h
@@ -0,0 +1,310 @@
+// Copyright Epic Games, Inc. All Rights Reserved.
+
+#pragma once
+
+#include "CoreMinimal.h"
+#include "UObject/ObjectMacros.h"
+#include "BonePose.h"
+#include "BoneControllers/AnimNode_ModifyBone.h"
+#include "Animation/AnimSingleNodeInstance.h"
+#include "Animation/AnimSingleNodeInstanceProxy.h"
+#include "AnimNodes/AnimNode_CurveSource.h"
+#include "AnimNodes/AnimNode_PoseBlendNode.h"
+#include "AnimNodes/AnimNode_CopyPoseFromMesh.h"
+#include "AnimPreviewInstance.generated.h"
+
+/** Enum to know how montage is being played */
+UENUM()
+enum EMontagePreviewType
+{
+	/** Playing montage in usual way. */
+	EMPT_Normal, 
+	/** Playing all sections. */
+	EMPT_AllSections,
+	EMPT_MAX,
+};
+
+
+/** Proxy override for this UAnimInstance-derived class */
+USTRUCT()
+struct ANIMGRAPH_API FAnimPreviewInstanceProxy : public FAnimSingleNodeInstanceProxy
+{
+	GENERATED_BODY()
+
+public:
+	FAnimPreviewInstanceProxy()
+	{
+		bCanProcessAdditiveAnimations = true;
+	}
+
+	FAnimPreviewInstanceProxy(UAnimInstance* InAnimInstance)
+		: FAnimSingleNodeInstanceProxy(InAnimInstance)
+		, SkeletalControlAlpha(1.0f)
+#if WITH_EDITORONLY_DATA
+		, bForceRetargetBasePose(false)
+#endif
+		, bEnableControllers(true)
+		, bSetKey(false)
+	{
+		bCanProcessAdditiveAnimations = true;
+	}
+
+	virtual void Initialize(UAnimInstance* InAnimInstance) override;
+	virtual void Update(float DeltaSeconds) override;
+	virtual bool Evaluate(FPoseContext& Output) override;
+	virtual void PreUpdate(UAnimInstance* InAnimInstance, float DeltaSeconds) override;
+	virtual void SetAnimationAsset(UAnimationAsset* NewAsset, USkeletalMeshComponent* MeshComponent, bool bIsLooping, float InPlayRate) override;
+
+	void ResetModifiedBone(bool bCurveController = false);
+
+	FAnimNode_ModifyBone* FindModifiedBone(const FName& InBoneName, bool bCurveController = false);	
+
+	FAnimNode_ModifyBone& ModifyBone(const FName& InBoneName, bool bCurveController = false);
+
+	void RemoveBoneModification(const FName& InBoneName, bool bCurveController = false);
+
+	void SetForceRetargetBasePose(bool bInForceRetargetBasePose)
+	{ 
+		bForceRetargetBasePose = bInForceRetargetBasePose; 
+	}
+
+	bool GetForceRetargetBasePose() const 
+	{ 
+		return bForceRetargetBasePose; 
+	}
+
+	void EnableControllers(bool bEnable)
+	{
+		bEnableControllers = bEnable;
+	}
+
+	void SetSkeletalControlAlpha(float InSkeletalControlAlpha)
+	{
+		SkeletalControlAlpha = FMath::Clamp<float>(InSkeletalControlAlpha, 0.f, 1.f);
+	}
+
+#if WITH_EDITOR	
+	void SetKey()
+	{
+		bSetKey = true;
+	}
+
+	FDelegateHandle AddKeyCompleteDelegate(FSimpleMulticastDelegate::FDelegate InOnSetKeyCompleteDelegate)
+	{
+		return OnSetKeyCompleteDelegate.Add(InOnSetKeyCompleteDelegate);
+	}
+
+	void RemoveKeyCompleteDelegate(FDelegateHandle InDelegateHandle)
+	{
+		OnSetKeyCompleteDelegate.Remove(InDelegateHandle);
+	}
+#endif
+
+	void RefreshCurveBoneControllers(UAnimationAsset* AssetToRefreshFrom);
+
+	TArray<FAnimNode_ModifyBone>& GetBoneControllers()
+	{
+		return BoneControllers;
+	}
+
+	TArray<FAnimNode_ModifyBone>& GetCurveBoneControllers()
+	{
+		return CurveBoneControllers;
+	}
+
+	virtual void AddImpulseAtLocation(FVector Impulse, FVector Location, FName BoneName = NAME_None) {}
+
+	/** Sets an external debug skeletal mesh component to use to debug */
+	void SetDebugSkeletalMeshComponent(USkeletalMeshComponent* InSkeletalMeshComponent);
+
+	/** Gets the external debug skeletal mesh component we are debugging */
+	USkeletalMeshComponent* GetDebugSkeletalMeshComponent() const;
+
+private:
+	void UpdateCurveController();
+
+	void ApplyBoneControllers(TArray<FAnimNode_ModifyBone> &InBoneControllers, FComponentSpacePoseContext& ComponentSpacePoseContext);
+
+	void SetKeyImplementation(const FCompactPose& PreControllerInLocalSpace, const FCompactPose& PostControllerInLocalSpace);
+
+	void AddKeyToSequence(UAnimSequence* Sequence, float Time, const FName& BoneName, const FTransform& AdditiveTransform);
+
+private:
+	/** Controllers for individual bones */
+	TArray<FAnimNode_ModifyBone> BoneControllers;
+
+	/** Curve modifiers */
+	TArray<FAnimNode_ModifyBone> CurveBoneControllers;
+
+	/** External curve for in-editor curve sources (such as audio) */
+	FAnimNode_CurveSource CurveSource;
+
+	/** Pose blend node for evaluating pose assets (for previewing curve sources) */
+	FAnimNode_PoseBlendNode PoseBlendNode;
+
+	/** Allows us to copy a pose from the mesh being debugged */
+	FAnimNode_CopyPoseFromMesh CopyPoseNode;
+
+	/**
+	 * Delegate to call after Key is set
+	 */
+	FSimpleMulticastDelegate OnSetKeyCompleteDelegate;
+
+	/** Shared parameters for previewing blendspace or animsequence **/
+	float SkeletalControlAlpha;
+
+#if WITH_EDITORONLY_DATA
+	bool bForceRetargetBasePose;
+#endif
+
+	/*
+	 * Used to determine if controller has to be applied or not
+	 * Used to disable controller during editing
+	 */
+	bool bEnableControllers;
+
+	/* 
+	 * When this flag is true, it sets key
+	 */
+	bool bSetKey;
+};
+
+/**
+ * This Instance only contains one AnimationAsset, and produce poses
+ * Used by Preview in AnimGraph, Playing single animation in Kismet2 and etc
+ */
+
+UCLASS(transient, NotBlueprintable, noteditinlinenew)
+class ANIMGRAPH_API UAnimPreviewInstance : public UAnimSingleNodeInstance
+{
+	GENERATED_UCLASS_BODY()
+
+		/** Shared parameters for previewing blendspace or animsequence **/
+	UPROPERTY(transient)
+	TEnumAsByte<enum EMontagePreviewType> MontagePreviewType;
+
+	UPROPERTY(transient)
+	int32 MontagePreviewStartSectionIdx;
+
+	//~ Begin UObject Interface
+	virtual void Serialize(FArchive& Ar) override;
+	//~ End UObject Interface
+
+	//~ Begin UAnimInstance Interface
+	virtual void NativeInitializeAnimation() override;
+	virtual FAnimInstanceProxy* CreateAnimInstanceProxy() override;
+	virtual bool CanRunParallelWork() const { return false; }
+protected:
+	virtual void Montage_Advance(float DeltaTime) override;
+	//~ End UAnimInstance Interface
+
+public:
+	/** Set SkeletalControl Alpha**/
+	void SetSkeletalControlAlpha(float SkeletalControlAlpha);
+
+	UAnimSequence* GetAnimSequence();
+
+	//~ Begin UAnimSingleNodeInstance Interface
+	virtual void RestartMontage(UAnimMontage* Montage, FName FromSection = FName()) override;
+	virtual void SetAnimationAsset(UAnimationAsset* NewAsset, bool bIsLooping = true, float InPlayRate = 1.f) override;
+	//~ End UAnimSingleNodeInstance Interface
+
+	/** Montage preview functions */
+	void MontagePreview_JumpToStart();
+	void MontagePreview_JumpToEnd();
+	void MontagePreview_JumpToPreviewStart();
+	void MontagePreview_Restart();
+	void MontagePreview_PreviewNormal(int32 FromSectionIdx = INDEX_NONE, bool bPlay = true);
+	void MontagePreview_SetLoopNormal(bool bIsLooping, int32 PreferSectionIdx = INDEX_NONE);
+	void MontagePreview_PreviewAllSections(bool bPlay = true);
+	void MontagePreview_SetLoopAllSections(bool bIsLooping);
+	void MontagePreview_SetLoopAllSetupSections(bool bIsLooping);
+	void MontagePreview_ResetSectionsOrder();
+	void MontagePreview_SetLooping(bool bIsLooping);
+	void MontagePreview_SetPlaying(bool bIsPlaying);
+	void MontagePreview_SetReverse(bool bInReverse);
+	void MontagePreview_StepForward();
+	void MontagePreview_StepBackward();
+	void MontagePreview_JumpToPosition(float NewPosition);
+	int32 MontagePreview_FindFirstSectionAsInMontage(int32 AnySectionIdx);
+	int32 MontagePreview_FindLastSection(int32 StartSectionIdx);
+	float MontagePreview_CalculateStepLength();
+	void MontagePreview_RemoveBlendOut();
+	bool IsPlayingMontage() { return GetActiveMontageInstance() != NULL; }
+
+	/** 
+	 * Finds an already modified bone 
+	 * @param	InBoneName	The name of the bone modification to find
+	 * @return the bone modification or NULL if no current modification was found
+	 */
+	FAnimNode_ModifyBone* FindModifiedBone(const FName& InBoneName, bool bCurveController=false);
+
+	/** 
+	 * Modifies a single bone. Create a new FAnimNode_ModifyBone if one does not exist for the passed-in bone.
+	 * @param	InBoneName	The name of the bone to modify
+	 * @return the new or existing bone modification
+	 */
+	FAnimNode_ModifyBone& ModifyBone(const FName& InBoneName, bool bCurveController=false);
+
+	/**
+	 * Removes an existing bone modification
+	 * @param	InBoneName	The name of the existing modification to remove
+	 */
+	void RemoveBoneModification(const FName& InBoneName, bool bCurveController=false);
+
+	/**
+	 * Reset all bone modified
+	 */
+	void ResetModifiedBone(bool bCurveController=false);
+
+	void SetForceRetargetBasePose(bool ForceRetargetBasePose);
+
+	bool GetForceRetargetBasePose() const;
+
+#if WITH_EDITOR	
+	/**
+	 * Convert current modified bone transforms (BoneControllers) to transform curves (CurveControllers)
+	 * it does based on CurrentTime. This function does not set key directly here. 
+	 * It does wait until next update, and it gets the delta of transform before applying curves, and 
+	 * creates curves from it, so you'll need delegate if you'd like to do something after (set with SetKeyCompleteDelegate)
+	 */
+	void SetKey();
+
+	/**
+	 * Add the delegate to be called when a key is set.
+	 * 
+	 * @param Delegate To be called once set key is completed
+	 */
+	FDelegateHandle AddKeyCompleteDelegate(FSimpleMulticastDelegate::FDelegate InOnSetKeyCompleteDelegate);
+
+	/**
+	 * Add the delegate to be called when a key is set.
+	 * 
+	 * @param Delegate To be called once set key is completed
+	 */
+	void RemoveKeyCompleteDelegate(FDelegateHandle InDelegateHandle);
+#endif
+
+	/** 
+	 * Refresh Curve Bone Controllers based on TransformCurves from Animation data
+	 */
+	void RefreshCurveBoneControllers();
+
+	/** 
+	 * Enable Controllers
+	 * This is used by when editing, when controller has to be disabled
+	 */
+	void EnableControllers(bool bEnable);
+
+	/** Preview physics interaction */
+	void AddImpulseAtLocation(FVector Impulse, FVector Location, FName BoneName = NAME_None);
+
+	/** Sets an external debug skeletal mesh component to use to debug */
+	void SetDebugSkeletalMeshComponent(USkeletalMeshComponent* InSkeletalMeshComponent);
+
+	/** Gets the external debug skeletal mesh component we are debugging */
+	USkeletalMeshComponent* GetDebugSkeletalMeshComponent() const;
+};
+
+
+
diff --git a/Engine/Source/Editor/AnimGraph/Public/AnimationGraphSchema.h b/Engine/Source/Editor/AnimGraph/Public/AnimationGraphSchema.h
new file mode 100644
--- /dev/null
+++ b/Engine/Source/Editor/AnimGraph/Public/AnimationGraphSchema.h
@@ -0,0 +1,119 @@
+// Copyright Epic Games, Inc. All Rights Reserved.
+
+
+#pragma once
+
+#include "CoreMinimal.h"
+#include "UObject/ObjectMacros.h"
+#include "AssetData.h"
+#include "EdGraph/EdGraphSchema.h"
+#include "EdGraphSchema_K2.h"
+#include "AnimationGraphSchema.generated.h"
+
+class FMenuBuilder;
+class UAnimationAsset;
+class UAnimBlueprint;
+class UPhysicsAsset;
+struct FBPInterfaceDescription;
+
+UCLASS(MinimalAPI)
+class UAnimationGraphSchema : public UEdGraphSchema_K2
+{
+	GENERATED_UCLASS_BODY()
+
+	// Common PinNames
+	UPROPERTY()
+	FString PN_SequenceName;    // PC_Object+PSC_Sequence
+
+	UPROPERTY()
+	FName NAME_NeverAsPin;
+
+	UPROPERTY()
+	FName NAME_PinHiddenByDefault;
+
+	UPROPERTY()
+	FName NAME_PinShownByDefault;
+
+	UPROPERTY()
+	FName NAME_AlwaysAsPin;
+
+	UPROPERTY()
+	FName NAME_CustomizeProperty;
+
+	UPROPERTY()
+	FName NAME_OnEvaluate;
+	
+	UPROPERTY()
+	FName DefaultEvaluationHandlerName;
+
+	//~ Begin UEdGraphSchema Interface.
+	virtual FLinearColor GetPinTypeColor(const FEdGraphPinType& PinType) const override;
+	virtual EGraphType GetGraphType(const UEdGraph* TestEdGraph) const override;
+	virtual void CreateDefaultNodesForGraph(UEdGraph& Graph) const override;
+	virtual void HandleGraphBeingDeleted(UEdGraph& GraphBeingRemoved) const override;
+	virtual bool CreateAutomaticConversionNodeAndConnections(UEdGraphPin* PinA, UEdGraphPin* PinB) const override;
+	virtual void DroppedAssetsOnGraph(const TArray<FAssetData>& Assets, const FVector2D& GraphPosition, UEdGraph* Graph) const override;
+	virtual void DroppedAssetsOnNode(const TArray<FAssetData>& Assets, const FVector2D& GraphPosition, UEdGraphNode* Node) const override;
+	virtual void DroppedAssetsOnPin(const TArray<FAssetData>& Assets, const FVector2D& GraphPosition, UEdGraphPin* Pin) const override;
+	virtual void GetAssetsNodeHoverMessage(const TArray<FAssetData>& Assets, const UEdGraphNode* HoverNode, FString& OutTooltipText, bool& OutOkIcon) const override;
+	virtual void GetAssetsPinHoverMessage(const TArray<FAssetData>& Assets, const UEdGraphPin* HoverPin, FString& OutTooltipText, bool& OutOkIcon) const override;
+	virtual void GetAssetsGraphHoverMessage(const TArray<FAssetData>& Assets, const UEdGraph* HoverGraph, FString& OutTooltipText, bool& OutOkIcon) const override;
+	virtual void GetContextMenuActions(class UToolMenu* Menu, class UGraphNodeContextMenuContext* Context) const override;
+	virtual FText GetPinDisplayName(const UEdGraphPin* Pin) const override;
+	virtual bool CanDuplicateGraph(UEdGraph* InSourceGraph) const override { return InSourceGraph->GetFName() != UEdGraphSchema_K2::GN_AnimGraph; }
+	virtual bool DoesSupportEventDispatcher() const	override { return false; }
+	virtual bool ShouldAlwaysPurgeOnModification() const override { return true; }
+	virtual bool TryCreateConnection(UEdGraphPin* A, UEdGraphPin* B) const override;
+	virtual void GetGraphDisplayInformation(const UEdGraph& Graph, /*out*/ FGraphDisplayInfo& DisplayInfo) const override;
+	//~ End UEdGraphSchema Interface.
+
+	//~ Begin UEdGraphSchema_K2 Interface
+	virtual const FPinConnectionResponse DetermineConnectionResponseOfCompatibleTypedPins(const UEdGraphPin* PinA, const UEdGraphPin* PinB, const UEdGraphPin* InputPin, const UEdGraphPin* OutputPin) const override;
+	virtual bool SearchForAutocastFunction(const FEdGraphPinType& OutputPinType, const FEdGraphPinType& InputPinType, /*out*/ FName& TargetFunction, /*out*/ UClass*& FunctionOwner) const override;
+	virtual bool ArePinsCompatible(const UEdGraphPin* PinA, const UEdGraphPin* PinB, const UClass* CallingContext = NULL, bool bIgnoreArray = false) const override;
+	virtual bool DoesSupportAnimNotifyActions() const override;
+	virtual void CreateFunctionGraphTerminators(UEdGraph& Graph, UClass* Class) const override;
+	//~ End UEdGraphSchema_K2 Interface
+
+	/** Spawn the correct node in the Animation Graph using the given AnimationAsset at the supplied location */
+	static void SpawnNodeFromAsset(UAnimationAsset* Asset, const FVector2D& GraphPosition, UEdGraph* Graph, UEdGraphPin* PinIfAvailable);
+
+	/** Spawn a rigid body node if we drop a physics asset on the graph */
+	static void SpawnRigidBodyNodeFromAsset(UPhysicsAsset* Asset, const FVector2D& GraphPosition, UEdGraph* Graph);
+
+	/** Update the specified node to a new asset */
+	static void UpdateNodeWithAsset(class UK2Node* K2Node, UAnimationAsset* Asset);
+
+	/** Auto-arranges a graph's inputs and outputs. Does nothing to nodes that are not roots or inputs */
+	ANIMGRAPH_API static void AutoArrangeInterfaceGraph(UEdGraph& Graph);
+
+	/** Checks to see whether the passed-in pin type is a pose pin (local or component space) */
+	ANIMGRAPH_API static bool IsPosePin(const FEdGraphPinType& PinType);
+
+	/** Checks to see whether the passed-in pin type is a local space pose pin */
+	ANIMGRAPH_API static bool IsLocalSpacePosePin(const FEdGraphPinType& PinType);
+
+	/** Checks to see whether the passed-in pin type is a component space pose pin */
+	ANIMGRAPH_API static bool IsComponentSpacePosePin(const FEdGraphPinType& PinType);
+
+	/** Makes a local space pose pin type */
+	ANIMGRAPH_API static FEdGraphPinType MakeLocalSpacePosePin();
+
+	/** Makes a component space pose pin type */
+	ANIMGRAPH_API static FEdGraphPinType MakeComponentSpacePosePin();
+
+	/** Conforms an anim graph to an interface function */
+	ANIMGRAPH_API static void ConformAnimGraphToInterface(UBlueprint* InBlueprint, UEdGraph& InGraph, UFunction* InFunction);
+
+	/** Conforms anim layer nodes to an interface desc by GUID */
+	ANIMGRAPH_API static void ConformAnimLayersByGuid(const UAnimBlueprint* InAnimBlueprint, const FBPInterfaceDescription& CurrentInterfaceDesc);
+
+	UE_DEPRECATED(4.24, "Function renamed, please use GetPositionForNewLinkedInputPoseNode")
+	ANIMGRAPH_API static FVector2D GetPositionForNewSubInputNode(UEdGraph& InGraph) { return GetPositionForNewLinkedInputPoseNode(InGraph); }
+
+	/** Find a position for a newly created linked input pose */
+	ANIMGRAPH_API static FVector2D GetPositionForNewLinkedInputPoseNode(UEdGraph& InGraph);
+};
+
+
+
diff --git a/Engine/Source/Editor/AnimGraph/Public/AnimationStateMachineSchema.h b/Engine/Source/Editor/AnimGraph/Public/AnimationStateMachineSchema.h
new file mode 100644
--- /dev/null
+++ b/Engine/Source/Editor/AnimGraph/Public/AnimationStateMachineSchema.h
@@ -0,0 +1,93 @@
+// Copyright Epic Games, Inc. All Rights Reserved.
+
+
+#pragma once
+
+#include "CoreMinimal.h"
+#include "UObject/ObjectMacros.h"
+#include "AssetData.h"
+#include "EdGraph/EdGraphSchema.h"
+#include "AnimationStateMachineSchema.generated.h"
+
+class FMenuBuilder;
+class UEdGraph;
+
+/** Action to add a node to the graph */
+USTRUCT()
+struct ANIMGRAPH_API FEdGraphSchemaAction_NewStateNode : public FEdGraphSchemaAction
+{
+	GENERATED_USTRUCT_BODY()
+
+	UEdGraphNode* NodeTemplate;
+
+	FEdGraphSchemaAction_NewStateNode() 
+		: FEdGraphSchemaAction()
+		, NodeTemplate(nullptr)
+	{}
+
+	FEdGraphSchemaAction_NewStateNode(FText InNodeCategory, FText InMenuDesc, FText InToolTip, const int32 InGrouping)
+		: FEdGraphSchemaAction(MoveTemp(InNodeCategory), MoveTemp(InMenuDesc), MoveTemp(InToolTip), InGrouping)
+		, NodeTemplate(nullptr)
+	{}
+
+	virtual UEdGraphNode* PerformAction(class UEdGraph* ParentGraph, UEdGraphPin* FromPin, const FVector2D Location, bool bSelectNewNode = true) override;
+	virtual void AddReferencedObjects( FReferenceCollector& Collector ) override;
+
+	template <typename NodeType>
+	static NodeType* SpawnNodeFromTemplate(class UEdGraph* ParentGraph, NodeType* InTemplateNode, const FVector2D Location = FVector2D(0.0f, 0.0f), bool bSelectNewNode = true)
+	{
+		FEdGraphSchemaAction_NewStateNode Action;
+		Action.NodeTemplate = InTemplateNode;
+
+		return Cast<NodeType>(Action.PerformAction(ParentGraph, NULL, Location, bSelectNewNode));
+	}
+};
+
+/** Action to create new comment */
+USTRUCT()
+struct ANIMGRAPH_API FEdGraphSchemaAction_NewStateComment : public FEdGraphSchemaAction
+{
+	GENERATED_USTRUCT_BODY();
+
+	FEdGraphSchemaAction_NewStateComment() 
+		: FEdGraphSchemaAction()
+	{}
+
+	FEdGraphSchemaAction_NewStateComment(FText InNodeCategory, FText InMenuDesc, FText InToolTip, const int32 InGrouping)
+		: FEdGraphSchemaAction(MoveTemp(InNodeCategory), MoveTemp(InMenuDesc), MoveTemp(InToolTip), InGrouping)
+	{}
+
+	//~ Begin FEdGraphSchemaAction Interface
+	virtual UEdGraphNode* PerformAction(class UEdGraph* ParentGraph, UEdGraphPin* FromPin, const FVector2D Location, bool bSelectNewNode = true) override;
+	//~ End FEdGraphSchemaAction Interface
+};
+
+
+UCLASS(MinimalAPI)
+class UAnimationStateMachineSchema : public UEdGraphSchema
+{
+	GENERATED_UCLASS_BODY()
+
+	ANIMGRAPH_API static const FName PC_Exec;
+
+	//~ Begin UEdGraphSchema Interface
+	virtual void CreateDefaultNodesForGraph(UEdGraph& Graph) const override;
+	virtual const FPinConnectionResponse CanCreateConnection(const UEdGraphPin* A, const UEdGraphPin* B) const override;
+	virtual bool TryCreateConnection(UEdGraphPin* A, UEdGraphPin* B) const override;
+	virtual bool CreateAutomaticConversionNodeAndConnections(UEdGraphPin* PinA, UEdGraphPin* PinB) const override;
+	virtual void GetGraphContextActions(FGraphContextMenuBuilder& ContextMenuBuilder) const override;
+	virtual EGraphType GetGraphType(const UEdGraph* TestEdGraph) const override;
+	virtual void GetContextMenuActions(class UToolMenu* Menu, class UGraphNodeContextMenuContext* Context) const override;
+	virtual FLinearColor GetPinTypeColor(const FEdGraphPinType& PinType) const override;
+	virtual void GetGraphDisplayInformation(const UEdGraph& Graph, /*out*/ FGraphDisplayInfo& DisplayInfo) const override;
+	virtual void DroppedAssetsOnGraph(const TArray<FAssetData>& Assets, const FVector2D& GraphPosition, UEdGraph* Graph) const override;
+	virtual void DroppedAssetsOnNode(const TArray<FAssetData>& Assets, const FVector2D& GraphPosition, UEdGraphNode* Node) const override;
+	virtual void DroppedAssetsOnPin(const TArray<FAssetData>& Assets, const FVector2D& GraphPosition, UEdGraphPin* Pin) const override;
+	virtual void GetAssetsNodeHoverMessage(const TArray<FAssetData>& Assets, const UEdGraphNode* HoverNode, FString& OutTooltipText, bool& OutOkIcon) const override;
+	virtual void GetAssetsPinHoverMessage(const TArray<FAssetData>& Assets, const UEdGraphPin* HoverPin, FString& OutTooltipText, bool& OutOkIcon) const override;
+
+	virtual void BreakNodeLinks(UEdGraphNode& TargetNode) const override;
+	virtual void BreakPinLinks(UEdGraphPin& TargetPin, bool bSendsNodeNotification) const override;
+	virtual void BreakSinglePinLink(UEdGraphPin* SourcePin, UEdGraphPin* TargetPin) const override;
+	//~ End UEdGraphSchema Interface
+};
