Pause the animation timeline for animation sequences with 1 frame

diff --git a/Engine/Source/Editor/AnimGraph/Private/AnimPreviewInstance.cpp b/Engine/Source/Editor/AnimGraph/Private/AnimPreviewInstance.cpp
--- a/Engine/Source/Editor/AnimGraph/Private/AnimPreviewInstance.cpp
+++ b/Engine/Source/Editor/AnimGraph/Private/AnimPreviewInstance.cpp
@@ -583,6 +583,14 @@ void UAnimPreviewInstance::SetAnimationAsset(UAnimationAsset* NewAsset, bool bIs
 	Super::SetAnimationAsset(NewAsset, bIsLooping, InPlayRate);
 	RootMotionMode = Cast<UAnimMontage>(CurrentAsset) != nullptr ? ERootMotionMode::RootMotionFromMontagesOnly : ERootMotionMode::RootMotionFromEverything;
 
+	// disable playing for single frame assets
+	UAnimSequence* AnimSequence = GetAnimSequence();
+	bool bSingleFrameSequence = (AnimSequence != nullptr) ? AnimSequence->GetNumberOfSampledKeys() <= 1 : false; 
+	if (bSingleFrameSequence && Proxy.IsPlaying())
+	{
+		Proxy.SetPlaying(false);
+	}
+
 	// should re sync up curve bone controllers from new asset
 	Proxy.RefreshCurveBoneControllers(CurrentAsset);
 }
