Addresses problem of being able to change the skeleton in a child animation blueprint by generating errors from the blueprint compilation helping the user fix things up

diff --git a/Engine/Source/Editor/AnimGraph/Private/AnimBlueprintCompiler.cpp b/Engine/Source/Editor/AnimGraph/Private/AnimBlueprintCompiler.cpp
--- a/Engine/Source/Editor/AnimGraph/Private/AnimBlueprintCompiler.cpp
+++ b/Engine/Source/Editor/AnimGraph/Private/AnimBlueprintCompiler.cpp
@@ -127,7 +127,7 @@ FAnimBlueprintCompilerContext::FAnimBlueprintCompilerContext(UAnimBlueprint* Sou
 
 		if(bNodeGuidsRegenerated)
 		{
-			UE_LOG(LogAnimation, Warning, TEXT("Animation Blueprint %s has nodes with invalid node guids that have been regenerated. This blueprint will not cook deterministically until it is resaved."), *AnimBlueprint->GetPathName());
+			UE_LOG(LogAnimation, Warning, TEXT("Animation Blueprint %s has nodes with invalid node guids that have been regenerated. This blueprint will not cook deterministically until it is resaved."), AnimBlueprint);
 		}
 	}
 
@@ -528,6 +528,16 @@ void FAnimBlueprintCompilerContext::CopyTermDefaultsToDefaultObject(UObject* Def
 		{
 			if (Override->NewAsset)
 			{
+				if (UAnimSequenceBase* Sequence = Cast<UAnimSequenceBase>(Override->NewAsset))
+				{
+					USkeleton* SequenceSkeleton = Sequence->GetSkeleton();
+					if (SequenceSkeleton != AnimBlueprint->TargetSkeleton)
+					{
+						MessageLog.Error(TEXT("@@ has override @@ using skeleton @@ which is different to @@"),
+							AnimBlueprint, Sequence, SequenceSkeleton, AnimBlueprint->TargetSkeleton);
+					}
+				}
+
 				FAnimNode_Base* BaseNode = NewAnimBlueprintClass->GetPropertyInstance<FAnimNode_Base>(DefaultAnimInstance, Override->ParentNodeGuid, EPropertySearchMode::Hierarchy);
 				if (BaseNode)
 				{
@@ -772,6 +782,14 @@ void FAnimBlueprintCompilerContext::MergeUbergraphPagesIn(UEdGraph* Ubergraph)
 	if (bIsDerivedAnimBlueprint)
 	{
 		// Skip any work related to an anim graph, it's all done by the parent class
+
+		USkeleton* TargetSkeleton = AnimBlueprint->TargetSkeleton;
+		USkeleton* RootTargetSkeleton = UAnimBlueprint::FindRootTargetSkeleton(AnimBlueprint);
+		if (TargetSkeleton != RootTargetSkeleton)
+		{
+			MessageLog.Error(TEXT("@@ references skeleton @@ which is different to the root/parent skeleton @@"), 
+			                 AnimBlueprint, TargetSkeleton, RootTargetSkeleton);
+		}
 	}
 	else
 	{
diff --git a/Engine/Source/Runtime/Engine/Classes/Animation/AnimBlueprint.h b/Engine/Source/Runtime/Engine/Classes/Animation/AnimBlueprint.h
--- a/Engine/Source/Runtime/Engine/Classes/Animation/AnimBlueprint.h
+++ b/Engine/Source/Runtime/Engine/Classes/Animation/AnimBlueprint.h
@@ -144,6 +144,9 @@ class ENGINE_API UAnimBlueprint : public UBlueprint, public IInterface_PreviewMe
 	/** Returns the most base anim blueprint for a given blueprint (if it is inherited from another anim blueprint, returning null if only native / non-anim BP classes are it's parent) */
 	static UAnimBlueprint* FindRootAnimBlueprint(const UAnimBlueprint* DerivedBlueprint);
 
+	/** Returns the rootmost Target skeleton, searching up through the parent hierarchy until one is found */
+	static USkeleton* FindRootTargetSkeleton(const UAnimBlueprint* DerivedBlueprint);
+
 	DECLARE_MULTICAST_DELEGATE_TwoParams(FOnOverrideChangedMulticaster, FGuid, UAnimationAsset*);
 
 	typedef FOnOverrideChangedMulticaster::FDelegate FOnOverrideChanged;
diff --git a/Engine/Source/Runtime/Engine/Private/AnimBlueprint.cpp b/Engine/Source/Runtime/Engine/Private/AnimBlueprint.cpp
--- a/Engine/Source/Runtime/Engine/Private/AnimBlueprint.cpp
+++ b/Engine/Source/Runtime/Engine/Private/AnimBlueprint.cpp
@@ -88,7 +88,7 @@ int32 UAnimBlueprint::FindOrAddGroup(FName GroupName)
 /** Returns the most base anim blueprint for a given blueprint (if it is inherited from another anim blueprint, returning null if only native / non-anim BP classes are it's parent) */
 UAnimBlueprint* UAnimBlueprint::FindRootAnimBlueprint(const UAnimBlueprint* DerivedBlueprint)
 {
-	UAnimBlueprint* ParentBP = NULL;
+	UAnimBlueprint* ParentBP = nullptr;
 
 	// Determine if there is an anim blueprint in the ancestry of this class
 	for (UClass* ParentClass = DerivedBlueprint->ParentClass; ParentClass && (UObject::StaticClass() != ParentClass); ParentClass = ParentClass->GetSuperClass())
@@ -102,6 +102,27 @@ UAnimBlueprint* UAnimBlueprint::FindRootAnimBlueprint(const UAnimBlueprint* Deri
 	return ParentBP;
 }
 
+USkeleton* UAnimBlueprint::FindRootTargetSkeleton(const UAnimBlueprint* DerivedBlueprint)
+{
+	const UAnimBlueprint* ParentBP = nullptr;
+	USkeleton* RootTargetSkeleton = DerivedBlueprint->TargetSkeleton;
+
+	// Determine if there is an anim blueprint in the ancestry of this class
+	for (UClass* ParentClass = DerivedBlueprint->ParentClass; ParentClass && (UObject::StaticClass() != ParentClass); ParentClass = ParentClass->GetSuperClass())
+	{
+		if (UAnimBlueprint* TestBP = Cast<UAnimBlueprint>(ParentClass->ClassGeneratedBy))
+		{
+			ParentBP = TestBP;
+			if (ParentBP->TargetSkeleton)
+			{
+				RootTargetSkeleton = ParentBP->TargetSkeleton;
+			}
+		}
+	}
+
+	return RootTargetSkeleton;
+}
+
 FAnimParentNodeAssetOverride* UAnimBlueprint::GetAssetOverrideForNode(FGuid NodeGuid, bool bIgnoreSelf) const
 {
 	TArray<UBlueprint*> Hierarchy;
