Fix AnimBP transition warning if the rule is bound. Also auto-disconnect pins if they are to be replaced with a bound rule.

diff --git a/Engine/Source/Editor/AnimGraph/Private/AnimStateConduitNode.cpp b/Engine/Source/Editor/AnimGraph/Private/AnimStateConduitNode.cpp
--- a/Engine/Source/Editor/AnimGraph/Private/AnimStateConduitNode.cpp
+++ b/Engine/Source/Editor/AnimGraph/Private/AnimStateConduitNode.cpp
@@ -141,10 +141,17 @@ void UAnimStateConduitNode::ValidateNodeDuringCompilation(class FCompilerResults
 	UAnimGraphNode_TransitionResult* ResultNode = TransGraph->GetResultNode();
 	check(ResultNode);
 
-	UEdGraphPin* BoolResultPin = ResultNode->Pins[0];
-	if ((BoolResultPin->LinkedTo.Num() == 0) && (BoolResultPin->DefaultValue.ToBool() == false))
+	if (ResultNode->PropertyBindings.Num() > 0 && ResultNode->PropertyBindings.CreateIterator()->Value.bIsBound)
 	{
-		MessageLog.Warning(TEXT("@@ will never be taken, please connect something to @@"), this, BoolResultPin);
+		// Rule is bound so nothing more to check
+	}
+	else
+	{
+		UEdGraphPin* BoolResultPin = ResultNode->Pins[0];
+		if ((BoolResultPin->LinkedTo.Num() == 0) && (BoolResultPin->DefaultValue.ToBool() == false))
+		{
+			MessageLog.Warning(TEXT("@@ will never be taken, please connect something to @@"), this, BoolResultPin);
+		}
 	}
 }
 
diff --git a/Engine/Source/Editor/AnimGraph/Private/AnimStateTransitionNode.cpp b/Engine/Source/Editor/AnimGraph/Private/AnimStateTransitionNode.cpp
--- a/Engine/Source/Editor/AnimGraph/Private/AnimStateTransitionNode.cpp
+++ b/Engine/Source/Editor/AnimGraph/Private/AnimStateTransitionNode.cpp
@@ -600,7 +600,11 @@ void UAnimStateTransitionNode::ValidateNodeDuringCompilation(class FCompilerResu
 		UAnimGraphNode_TransitionResult* ResultNode = TransGraph->GetResultNode();
 		check(ResultNode);
 
-		if(ResultNode->Pins.Num() > 0)
+		if (ResultNode->PropertyBindings.Num() > 0 && ResultNode->PropertyBindings.CreateIterator()->Value.bIsBound)
+		{
+			// Rule is bound so nothing more to check
+		}
+		else if (ResultNode->Pins.Num() > 0)
 		{
 			UEdGraphPin* BoolResultPin = ResultNode->Pins[0];
 			if (BoolResultPin && (BoolResultPin->LinkedTo.Num() == 0) && (BoolResultPin->DefaultValue.ToBool() == false))
diff --git a/Engine/Source/Editor/Persona/Private/AnimGraphNodeDetails.cpp b/Engine/Source/Editor/Persona/Private/AnimGraphNodeDetails.cpp
--- a/Engine/Source/Editor/Persona/Private/AnimGraphNodeDetails.cpp
+++ b/Engine/Source/Editor/Persona/Private/AnimGraphNodeDetails.cpp
@@ -1323,6 +1323,9 @@ TSharedRef<SWidget> FAnimGraphNodeBindingExtension::GenerateExtensionWidget(cons
 			const UEdGraphSchema_K2* Schema = GetDefault<UEdGraphSchema_K2>();
 			IPropertyAccessEditor& PropertyAccessEditor = IModularFeatures::Get().GetModularFeature<IPropertyAccessEditor>("PropertyAccessEditor");
 
+			// Disconnect any connected pin since we will replace it with the bound property/function
+			ShowPinPropertyHandle->SetValue(false);
+
 			for(UObject* OuterObject : OuterObjects)
 			{
 				if(UAnimGraphNode_Base* AnimGraphNode = Cast<UAnimGraphNode_Base>(OuterObject))
@@ -1373,7 +1376,7 @@ TSharedRef<SWidget> FAnimGraphNodeBindingExtension::GenerateExtensionWidget(cons
 
 				FBlueprintEditorUtils::MarkBlueprintAsModified(Blueprint);
 			}
-
+			// Pins are exposed if we have a binding or not - and after running this we do.
 			ShowPinPropertyHandle->SetValue(true);
 		});
 
