Removed ensure that could fire with old bits of anim graph left hanging around

diff --git a/Engine/Source/Editor/AnimGraph/Private/AnimBlueprintCompilerHandler_Attributes.cpp b/Engine/Source/Editor/AnimGraph/Private/AnimBlueprintCompilerHandler_Attributes.cpp
--- a/Engine/Source/Editor/AnimGraph/Private/AnimBlueprintCompilerHandler_Attributes.cpp
+++ b/Engine/Source/Editor/AnimGraph/Private/AnimBlueprintCompilerHandler_Attributes.cpp
@@ -39,7 +39,7 @@ void FAnimBlueprintCompilerHandler_Attributes::PreProcessAnimationNodes(TArrayVi
 		// Establish links between save/used cached pose nodes
 		if(UAnimGraphNode_UseCachedPose* UseCachedPoseNode = Cast<UAnimGraphNode_UseCachedPose>(Node))
 		{
-			if(ensure(UseCachedPoseNode->SaveCachedPoseNode.IsValid()))
+			if(UseCachedPoseNode->SaveCachedPoseNode.IsValid())
 			{
 				SaveCachedPoseMap.Add(UseCachedPoseNode->SaveCachedPoseNode.Get(), UseCachedPoseNode);
 			}
