Fixed crash pasting a state machine

diff --git a/Engine/Source/Editor/AnimGraph/Private/AnimGraphNode_StateMachineBase.cpp b/Engine/Source/Editor/AnimGraph/Private/AnimGraphNode_StateMachineBase.cpp
--- a/Engine/Source/Editor/AnimGraph/Private/AnimGraphNode_StateMachineBase.cpp
+++ b/Engine/Source/Editor/AnimGraph/Private/AnimGraphNode_StateMachineBase.cpp
@@ -170,27 +170,30 @@ void UAnimGraphNode_StateMachineBase::PostPasteNode()
 {
 	Super::PostPasteNode();
 
-	// Add the new graph as a child of our parent graph
-	UEdGraph* ParentGraph = GetGraph();
-
-	if(ParentGraph->SubGraphs.Find(EditorStateMachineGraph) == INDEX_NONE)
+	if(EditorStateMachineGraph)
 	{
-		ParentGraph->SubGraphs.Add(EditorStateMachineGraph);
-	}
+		// Add the new graph as a child of our parent graph
+		UEdGraph* ParentGraph = GetGraph();
 
-	for (UEdGraphNode* GraphNode : EditorStateMachineGraph->Nodes)
-	{
-		GraphNode->CreateNewGuid();
-		GraphNode->PostPasteNode();
-		GraphNode->ReconstructNode();
-	}
+		if(ParentGraph->SubGraphs.Find(EditorStateMachineGraph) == INDEX_NONE)
+		{
+			ParentGraph->SubGraphs.Add(EditorStateMachineGraph);
+		}
 
-	// Find an interesting name
-	TSharedPtr<INameValidatorInterface> NameValidator = FNameValidatorFactory::MakeValidator(this);
-	FBlueprintEditorUtils::RenameGraphWithSuggestion(EditorStateMachineGraph, NameValidator, EditorStateMachineGraph->GetName());
+		for (UEdGraphNode* GraphNode : EditorStateMachineGraph->Nodes)
+		{
+			GraphNode->CreateNewGuid();
+			GraphNode->PostPasteNode();
+			GraphNode->ReconstructNode();
+		}
 
-	//restore transactional flag that is lost during copy/paste process
-	EditorStateMachineGraph->SetFlags(RF_Transactional);
+		// Find an interesting name
+		TSharedPtr<INameValidatorInterface> NameValidator = FNameValidatorFactory::MakeValidator(this);
+		FBlueprintEditorUtils::RenameGraphWithSuggestion(EditorStateMachineGraph, NameValidator, EditorStateMachineGraph->GetName());
+
+		//restore transactional flag that is lost during copy/paste process
+		EditorStateMachineGraph->SetFlags(RF_Transactional);
+	}
 }
 
 FString UAnimGraphNode_StateMachineBase::GetStateMachineName()
diff --git a/Engine/Source/Editor/AnimGraph/Private/AnimStateNodeBase.cpp b/Engine/Source/Editor/AnimGraph/Private/AnimStateNodeBase.cpp
--- a/Engine/Source/Editor/AnimGraph/Private/AnimStateNodeBase.cpp
+++ b/Engine/Source/Editor/AnimGraph/Private/AnimStateNodeBase.cpp
@@ -50,23 +50,26 @@ void UAnimStateNodeBase::PostPasteNode()
 
 	for(UEdGraph* SubGraph : GetSubGraphs())
 	{
-		// Add the new graph as a child of our parent graph
-		UEdGraph* ParentGraph = GetGraph();
-
-		if(ParentGraph->SubGraphs.Find(SubGraph) == INDEX_NONE)
+		if(SubGraph)
 		{
-			ParentGraph->SubGraphs.Add(SubGraph);
-		}
+			// Add the new graph as a child of our parent graph
+			UEdGraph* ParentGraph = GetGraph();
+
+			if(ParentGraph->SubGraphs.Find(SubGraph) == INDEX_NONE)
+			{
+				ParentGraph->SubGraphs.Add(SubGraph);
+			}
 
-		//@TODO: CONDUIT: Merge conflict - May no longer be necessary due to other changes?
-//		FBlueprintEditorUtils::RenameGraphWithSuggestion(SubGraph, NameValidator, GetDesiredNewNodeName());
-		//@ENDTODO
+			//@TODO: CONDUIT: Merge conflict - May no longer be necessary due to other changes?
+	//		FBlueprintEditorUtils::RenameGraphWithSuggestion(SubGraph, NameValidator, GetDesiredNewNodeName());
+			//@ENDTODO
 
-		// Restore transactional flag that is lost during copy/paste process
-		SubGraph->SetFlags(RF_Transactional);
+			// Restore transactional flag that is lost during copy/paste process
+			SubGraph->SetFlags(RF_Transactional);
 
-		UBlueprint* Blueprint = FBlueprintEditorUtils::FindBlueprintForGraphChecked(ParentGraph);
-		FBlueprintEditorUtils::MarkBlueprintAsStructurallyModified(Blueprint);
+			UBlueprint* Blueprint = FBlueprintEditorUtils::FindBlueprintForGraphChecked(ParentGraph);
+			FBlueprintEditorUtils::MarkBlueprintAsStructurallyModified(Blueprint);
+		}
 	}
 }
 
