Fixed crash opening a pose asset

The pose asset code path in the preview instance used to run its graph udpate in UpdateAnimation(), which meant that the regular scoped sync wasnt applied. Switching it to run in UpdateAnimationNode() means that everyting functions as normal.

diff --git a/Engine/Source/Editor/AnimGraph/Private/AnimPreviewInstance.cpp b/Engine/Source/Editor/AnimGraph/Private/AnimPreviewInstance.cpp
--- a/Engine/Source/Editor/AnimGraph/Private/AnimPreviewInstance.cpp
+++ b/Engine/Source/Editor/AnimGraph/Private/AnimPreviewInstance.cpp
@@ -124,21 +124,24 @@ void FAnimPreviewInstanceProxy::Update(float DeltaSeconds)
 	}
 #endif // #if WITH_EDITORONLY_DATA
 
+	FAnimSingleNodeInstanceProxy::Update(DeltaSeconds);
+}
+
+void FAnimPreviewInstanceProxy::UpdateAnimationNode(const FAnimationUpdateContext& InContext)
+{
 	if (CopyPoseNode.SourceMeshComponent.IsValid())
 	{
-		FAnimationUpdateContext UpdateContext(this, DeltaSeconds);
-		CopyPoseNode.Update_AnyThread(UpdateContext);
+		CopyPoseNode.Update_AnyThread(InContext);
 	}
 	else if (UPoseAsset* PoseAsset = Cast<UPoseAsset>(CurrentAsset))
 	{
 		PoseBlendNode.PoseAsset = PoseAsset;
 
-		FAnimationUpdateContext UpdateContext(this, DeltaSeconds);
-		PoseBlendNode.Update_AnyThread(UpdateContext);
+		PoseBlendNode.Update_AnyThread(InContext);
 	}
 	else
 	{
-		FAnimSingleNodeInstanceProxy::Update(DeltaSeconds);
+		FAnimSingleNodeInstanceProxy::UpdateAnimationNode(InContext);
 	}
 }
 
diff --git a/Engine/Source/Editor/AnimGraph/Public/AnimPreviewInstance.h b/Engine/Source/Editor/AnimGraph/Public/AnimPreviewInstance.h
--- a/Engine/Source/Editor/AnimGraph/Public/AnimPreviewInstance.h
+++ b/Engine/Source/Editor/AnimGraph/Public/AnimPreviewInstance.h
@@ -51,6 +51,7 @@ public:
 
 	virtual void Initialize(UAnimInstance* InAnimInstance) override;
 	virtual void Update(float DeltaSeconds) override;
+	virtual void UpdateAnimationNode(const FAnimationUpdateContext& InContext) override;
 	virtual bool Evaluate(FPoseContext& Output) override;
 	virtual void PreUpdate(UAnimInstance* InAnimInstance, float DeltaSeconds) override;
 	virtual void SetAnimationAsset(UAnimationAsset* NewAsset, USkeletalMeshComponent* MeshComponent, bool bIsLooping, float InPlayRate) override;
