Remove the ability to circumvent connection rules with reroute node in an animation blueprint.

diff --git a/Engine/Source/Editor/AnimGraph/Private/AnimationGraphSchema.cpp b/Engine/Source/Editor/AnimGraph/Private/AnimationGraphSchema.cpp
--- a/Engine/Source/Editor/AnimGraph/Private/AnimationGraphSchema.cpp
+++ b/Engine/Source/Editor/AnimGraph/Private/AnimationGraphSchema.cpp
@@ -190,9 +190,11 @@ bool UAnimationGraphSchema::TryCreateConnection(UEdGraphPin* A, UEdGraphPin* B)
 	}
 	check(OutputPin && InputPin);
 
-	UEdGraphNode* OutputNode = OutputPin->GetOwningNode();
+	UK2Node_Knot* OutputKnotNode = Cast<UK2Node_Knot>(OutputPin->GetOwningNode());
+	UK2Node_Knot* InputKnotNode = Cast<UK2Node_Knot>(InputPin->GetOwningNode());
+	bool bConnectionWithKnot = OutputKnotNode != nullptr || InputKnotNode != nullptr;
 
-	if(UK2Node_Knot* RerouteNode = Cast<UK2Node_Knot>(OutputNode))
+	if(bConnectionWithKnot)
 	{
 		// Double check this is our "exec"-like line
 		bool bOutputIsPose = IsPosePin(OutputPin->PinType);
