Pausing animation sequences with a single frame

diff --git a/Engine/Source/Editor/AnimGraph/Private/AnimPreviewInstance.cpp b/Engine/Source/Editor/AnimGraph/Private/AnimPreviewInstance.cpp
--- a/Engine/Source/Editor/AnimGraph/Private/AnimPreviewInstance.cpp
+++ b/Engine/Source/Editor/AnimGraph/Private/AnimPreviewInstance.cpp
@@ -589,7 +589,7 @@ void UAnimPreviewInstance::SetAnimationAsset(UAnimationAsset* NewAsset, bool bIs
 
 	// disable playing for single frame assets
 	UAnimSequence* AnimSequence = GetAnimSequence();
-	bool bSingleFrameSequence = (AnimSequence != nullptr) ? AnimSequence->GetNumberOfSampledKeys() <= 1 : false; 
+	bool bSingleFrameSequence = (AnimSequence != nullptr) ? AnimSequence->GetDataModel()->GetNumberOfFrames() <= 1 : false; 
 	if (bSingleFrameSequence && Proxy.IsPlaying())
 	{
 		Proxy.SetPlaying(false);
