Fixed cached poses not being usable nested inside blend space graphs

diff --git a/Engine/Source/Editor/AnimGraph/Private/AnimBlueprintCompilerHandler_CachedPose.cpp b/Engine/Source/Editor/AnimGraph/Private/AnimBlueprintCompilerHandler_CachedPose.cpp
--- a/Engine/Source/Editor/AnimGraph/Private/AnimBlueprintCompilerHandler_CachedPose.cpp
+++ b/Engine/Source/Editor/AnimGraph/Private/AnimBlueprintCompilerHandler_CachedPose.cpp
@@ -11,6 +11,7 @@
 #include "IAnimBlueprintGeneratedClassCompiledData.h"
 #include "IAnimBlueprintCompilerCreationContext.h"
 #include "IAnimBlueprintCompilationContext.h"
+#include "AnimationGraphSchema.h"
 
 FAnimBlueprintCompilerHandler_CachedPose::FAnimBlueprintCompilerHandler_CachedPose(IAnimBlueprintCompilerCreationContext& InCreationContext)
 {
@@ -130,6 +131,33 @@ void FAnimBlueprintCompilerHandler_CachedPose::CachePoseNodeOrdering_StartNewTra
 	UE_CLOG(bEnableDebug, LogAnimation, Display, TEXT("EndNewTraversal %s"), *RootName);
 }
 
+void FAnimBlueprintCompilerHandler_CachedPose::CachePoseNodeOrdering_TraverseInternal_SubGraph(IAnimBlueprintCompilationContext& InCompilationContext, UEdGraph* InGraph, TArray<UAnimGraphNode_SaveCachedPose*> &OrderedSavePoseNodes)
+{
+	const bool bEnableDebug = (CVarAnimDebugCachePoseNodeUpdateOrder.GetValueOnAnyThread() == 1);
+	UE_CLOG(bEnableDebug, LogAnimation, Display, TEXT("ProcessGraph %s"), *InGraph->GetName());	
+
+	if(InGraph->Schema->IsChildOf(UAnimationGraphSchema::StaticClass()))
+	{
+		TArray<UAnimGraphNode_Base*> PotentialRootNodes;
+		InGraph->GetNodesOfClass(PotentialRootNodes);
+
+		for (UAnimGraphNode_Base* PotentialRootNode : PotentialRootNodes)
+		{
+			if(PotentialRootNode->IsNodeRootSet())
+			{
+				CachePoseNodeOrdering_TraverseInternal(InCompilationContext, PotentialRootNode, OrderedSavePoseNodes);
+			}
+		}
+	}
+	else
+	{
+		for(UEdGraph* SubGraph : InGraph->SubGraphs)
+		{
+			CachePoseNodeOrdering_TraverseInternal_SubGraph(InCompilationContext, SubGraph, OrderedSavePoseNodes);
+		}
+	}
+}
+
 void FAnimBlueprintCompilerHandler_CachedPose::CachePoseNodeOrdering_TraverseInternal(IAnimBlueprintCompilationContext& InCompilationContext, UAnimGraphNode_Base* InAnimGraphNode, TArray<UAnimGraphNode_SaveCachedPose*> &OrderedSavePoseNodes)
 {
 	TArray<UAnimGraphNode_Base*> LinkedAnimNodes;
@@ -151,23 +179,20 @@ void FAnimBlueprintCompilerHandler_CachedPose::CachePoseNodeOrdering_TraverseInt
 				OrderedSavePoseNodes.Add(SaveNode);
 			}
 		}
-		else if (UAnimGraphNode_StateMachine* StateMachineNode = Cast<UAnimGraphNode_StateMachine>(LinkedNode))
+		else
 		{
-			for (UEdGraph* StateGraph : StateMachineNode->EditorStateMachineGraph->SubGraphs)
+			TArray<UEdGraph*> SubGraphs = LinkedNode->GetSubGraphs();
+			if(SubGraphs.Num() > 0)
 			{
-				TArray<UAnimGraphNode_StateResult*> ResultNodes;
-				StateGraph->GetNodesOfClass(ResultNodes);
-
-				// We should only get one here but doesn't hurt to loop here in case that changes
-				for (UAnimGraphNode_StateResult* ResultNode : ResultNodes)
+				for (UEdGraph* SubGraph : SubGraphs)
 				{
-					CachePoseNodeOrdering_TraverseInternal(InCompilationContext, ResultNode, OrderedSavePoseNodes);
+					CachePoseNodeOrdering_TraverseInternal_SubGraph(InCompilationContext, SubGraph, OrderedSavePoseNodes);
 				}
 			}
-		}
-		else
-		{
-			CachePoseNodeOrdering_TraverseInternal(InCompilationContext, LinkedNode, OrderedSavePoseNodes);
+			else
+			{
+				CachePoseNodeOrdering_TraverseInternal(InCompilationContext, LinkedNode, OrderedSavePoseNodes);
+			}
 		}
 	}
 }
diff --git a/Engine/Source/Editor/AnimGraph/Private/AnimBlueprintCompilerHandler_CachedPose.h b/Engine/Source/Editor/AnimGraph/Private/AnimBlueprintCompilerHandler_CachedPose.h
--- a/Engine/Source/Editor/AnimGraph/Private/AnimBlueprintCompilerHandler_CachedPose.h
+++ b/Engine/Source/Editor/AnimGraph/Private/AnimBlueprintCompilerHandler_CachedPose.h
@@ -29,6 +29,9 @@ private:
 	// Traverses a graph to collect save pose nodes starting at InRootNode, then processes each node
 	void CachePoseNodeOrdering_StartNewTraversal(IAnimBlueprintCompilationContext& InCompilationContext, UAnimGraphNode_Base* InRootNode, TArray<UAnimGraphNode_SaveCachedPose*> &OrderedSavePoseNodes, TArray<UAnimGraphNode_Base*> VisitedRootNodes);
 
+	// Traverse a node's subgraph to collect save pose nodes
+	void CachePoseNodeOrdering_TraverseInternal_SubGraph(IAnimBlueprintCompilationContext& InCompilationContext, UEdGraph* InGraph, TArray<UAnimGraphNode_SaveCachedPose*> &OrderedSavePoseNodes);
+	
 	// Traverses a graph to collect save pose nodes starting at InAnimGraphNode, does NOT process saved pose nodes afterwards
 	void CachePoseNodeOrdering_TraverseInternal(IAnimBlueprintCompilationContext& InCompilationContext, UAnimGraphNode_Base* InAnimGraphNode, TArray<UAnimGraphNode_SaveCachedPose*> &OrderedSavePoseNodes);
 
