Backing out //UE5/Release-5.0-EarlyAccess/Engine/Source/Editor/AnimGraph/Private/... changelist 15386468

Potentially caused T-pose problems in Frosty, backing out until I can investigate

diff --git a/Engine/Source/Editor/AnimGraph/Private/AnimBlueprintCompilerHandler_CachedPose.cpp b/Engine/Source/Editor/AnimGraph/Private/AnimBlueprintCompilerHandler_CachedPose.cpp
--- a/Engine/Source/Editor/AnimGraph/Private/AnimBlueprintCompilerHandler_CachedPose.cpp
+++ b/Engine/Source/Editor/AnimGraph/Private/AnimBlueprintCompilerHandler_CachedPose.cpp
@@ -11,7 +11,6 @@
 #include "IAnimBlueprintGeneratedClassCompiledData.h"
 #include "IAnimBlueprintCompilerCreationContext.h"
 #include "IAnimBlueprintCompilationContext.h"
-#include "AnimationGraphSchema.h"
 
 FAnimBlueprintCompilerHandler_CachedPose::FAnimBlueprintCompilerHandler_CachedPose(IAnimBlueprintCompilerCreationContext& InCreationContext)
 {
@@ -131,33 +130,6 @@ void FAnimBlueprintCompilerHandler_CachedPose::CachePoseNodeOrdering_StartNewTra
 	UE_CLOG(bEnableDebug, LogAnimation, Display, TEXT("EndNewTraversal %s"), *RootName);
 }
 
-void FAnimBlueprintCompilerHandler_CachedPose::CachePoseNodeOrdering_TraverseInternal_SubGraph(IAnimBlueprintCompilationContext& InCompilationContext, UEdGraph* InGraph, TArray<UAnimGraphNode_SaveCachedPose*> &OrderedSavePoseNodes)
-{
-	const bool bEnableDebug = (CVarAnimDebugCachePoseNodeUpdateOrder.GetValueOnAnyThread() == 1);
-	UE_CLOG(bEnableDebug, LogAnimation, Display, TEXT("ProcessGraph %s"), *InGraph->GetName());	
-
-	if(InGraph->Schema->IsChildOf(UAnimationGraphSchema::StaticClass()))
-	{
-		TArray<UAnimGraphNode_Base*> PotentialRootNodes;
-		InGraph->GetNodesOfClass(PotentialRootNodes);
-
-		for (UAnimGraphNode_Base* PotentialRootNode : PotentialRootNodes)
-		{
-			if(PotentialRootNode->IsNodeRootSet())
-			{
-				CachePoseNodeOrdering_TraverseInternal(InCompilationContext, PotentialRootNode, OrderedSavePoseNodes);
-			}
-		}
-	}
-	else
-	{
-		for(UEdGraph* SubGraph : InGraph->SubGraphs)
-		{
-			CachePoseNodeOrdering_TraverseInternal_SubGraph(InCompilationContext, SubGraph, OrderedSavePoseNodes);
-		}
-	}
-}
-
 void FAnimBlueprintCompilerHandler_CachedPose::CachePoseNodeOrdering_TraverseInternal(IAnimBlueprintCompilationContext& InCompilationContext, UAnimGraphNode_Base* InAnimGraphNode, TArray<UAnimGraphNode_SaveCachedPose*> &OrderedSavePoseNodes)
 {
 	TArray<UAnimGraphNode_Base*> LinkedAnimNodes;
@@ -179,20 +151,23 @@ void FAnimBlueprintCompilerHandler_CachedPose::CachePoseNodeOrdering_TraverseInt
 				OrderedSavePoseNodes.Add(SaveNode);
 			}
 		}
-		else
+		else if (UAnimGraphNode_StateMachine* StateMachineNode = Cast<UAnimGraphNode_StateMachine>(LinkedNode))
 		{
-			TArray<UEdGraph*> SubGraphs = LinkedNode->GetSubGraphs();
-			if(SubGraphs.Num() > 0)
+			for (UEdGraph* StateGraph : StateMachineNode->EditorStateMachineGraph->SubGraphs)
 			{
-				for (UEdGraph* SubGraph : SubGraphs)
+				TArray<UAnimGraphNode_StateResult*> ResultNodes;
+				StateGraph->GetNodesOfClass(ResultNodes);
+
+				// We should only get one here but doesn't hurt to loop here in case that changes
+				for (UAnimGraphNode_StateResult* ResultNode : ResultNodes)
 				{
-					CachePoseNodeOrdering_TraverseInternal_SubGraph(InCompilationContext, SubGraph, OrderedSavePoseNodes);
+					CachePoseNodeOrdering_TraverseInternal(InCompilationContext, ResultNode, OrderedSavePoseNodes);
 				}
 			}
-			else
-			{
-				CachePoseNodeOrdering_TraverseInternal(InCompilationContext, LinkedNode, OrderedSavePoseNodes);
-			}
+		}
+		else
+		{
+			CachePoseNodeOrdering_TraverseInternal(InCompilationContext, LinkedNode, OrderedSavePoseNodes);
 		}
 	}
 }
diff --git a/Engine/Source/Editor/AnimGraph/Private/AnimBlueprintCompilerHandler_CachedPose.h b/Engine/Source/Editor/AnimGraph/Private/AnimBlueprintCompilerHandler_CachedPose.h
--- a/Engine/Source/Editor/AnimGraph/Private/AnimBlueprintCompilerHandler_CachedPose.h
+++ b/Engine/Source/Editor/AnimGraph/Private/AnimBlueprintCompilerHandler_CachedPose.h
@@ -29,9 +29,6 @@ private:
 	// Traverses a graph to collect save pose nodes starting at InRootNode, then processes each node
 	void CachePoseNodeOrdering_StartNewTraversal(IAnimBlueprintCompilationContext& InCompilationContext, UAnimGraphNode_Base* InRootNode, TArray<UAnimGraphNode_SaveCachedPose*> &OrderedSavePoseNodes, TArray<UAnimGraphNode_Base*> VisitedRootNodes);
 
-	// Traverse a node's subgraph to collect save pose nodes
-	void CachePoseNodeOrdering_TraverseInternal_SubGraph(IAnimBlueprintCompilationContext& InCompilationContext, UEdGraph* InGraph, TArray<UAnimGraphNode_SaveCachedPose*> &OrderedSavePoseNodes);
-	
 	// Traverses a graph to collect save pose nodes starting at InAnimGraphNode, does NOT process saved pose nodes afterwards
 	void CachePoseNodeOrdering_TraverseInternal(IAnimBlueprintCompilationContext& InCompilationContext, UAnimGraphNode_Base* InAnimGraphNode, TArray<UAnimGraphNode_SaveCachedPose*> &OrderedSavePoseNodes);
 
