Double-clicking a Use Cached Pose node now goes to the Save Cached Pose node it references, if any exists

diff --git a/Engine/Source/Editor/AnimGraph/Private/AnimGraphNode_UseCachedPose.cpp b/Engine/Source/Editor/AnimGraph/Private/AnimGraphNode_UseCachedPose.cpp
--- a/Engine/Source/Editor/AnimGraph/Private/AnimGraphNode_UseCachedPose.cpp
+++ b/Engine/Source/Editor/AnimGraph/Private/AnimGraphNode_UseCachedPose.cpp
@@ -9,8 +9,9 @@
 #include "BlueprintNodeSpawner.h"
 #include "BlueprintActionDatabaseRegistrar.h"
 #include "AnimBlueprintCompilerHandler_CachedPose.h"
-#include "IAnimBlueprintCompilerHandlerCollection.h"
 #include "IAnimBlueprintCompilationContext.h"
+#include "Kismet2/KismetEditorUtilities.h"
+#include "BlueprintEditorModule.h"
 
 /////////////////////////////////////////////////////
 // UAnimGraphNode_UseCachedPose
@@ -88,6 +89,24 @@ FText UAnimGraphNode_UseCachedPose::GetNodeTitle(ENodeTitleType::Type TitleType)
 	return FText::Format(LOCTEXT("AnimGraphNode_UseCachedPose_Title", "Use cached pose '{CachePoseName}'"), Args);
 }
 
+void UAnimGraphNode_UseCachedPose::JumpToDefinition() const
+{
+	if (UObject* HyperlinkTarget = GetJumpTargetForDoubleClick())
+	{
+		TSharedPtr<IBlueprintEditor> BlueprintEditor = FKismetEditorUtilities::GetIBlueprintEditorForObject(HyperlinkTarget, true);
+		if(BlueprintEditor.IsValid())
+		{
+			BlueprintEditor->JumpToHyperlink(HyperlinkTarget, false);
+		}
+	}
+}
+
+UObject* UAnimGraphNode_UseCachedPose::GetJumpTargetForDoubleClick() const
+{
+	// Double click goes to the save cached pose for this node
+	return SaveCachedPoseNode.Get();
+}
+
 FString UAnimGraphNode_UseCachedPose::GetNodeCategory() const
 {
 	return TEXT("Cached Poses");
diff --git a/Engine/Source/Editor/AnimGraph/Public/AnimGraphNode_UseCachedPose.h b/Engine/Source/Editor/AnimGraph/Public/AnimGraphNode_UseCachedPose.h
--- a/Engine/Source/Editor/AnimGraph/Public/AnimGraphNode_UseCachedPose.h
+++ b/Engine/Source/Editor/AnimGraph/Public/AnimGraphNode_UseCachedPose.h
@@ -27,6 +27,8 @@ public:
 	// UEdGraphNode interface
 	virtual FText GetTooltipText() const override;
 	virtual FText GetNodeTitle(ENodeTitleType::Type TitleType) const override;
+	virtual UObject* GetJumpTargetForDoubleClick() const override;
+	virtual void JumpToDefinition() const override;
 	// End of UEdGraphNode interface
 
 	// UK2Node interface.
