Fix non-unity CIS

diff --git a/Engine/Source/Editor/AnimGraph/Private/AnimGraphNode_BlendSpaceGraphBase.cpp b/Engine/Source/Editor/AnimGraph/Private/AnimGraphNode_BlendSpaceGraphBase.cpp
--- a/Engine/Source/Editor/AnimGraph/Private/AnimGraphNode_BlendSpaceGraphBase.cpp
+++ b/Engine/Source/Editor/AnimGraph/Private/AnimGraphNode_BlendSpaceGraphBase.cpp
@@ -25,6 +25,7 @@
 #include "Kismet2/Kismet2NameValidators.h"
 #include "AnimNodes/AnimNode_BlendSpaceGraphBase.h"
 #include "BlueprintEditor.h"
+#include "Animation/AnimSequence.h"
 
 #define LOCTEXT_NAMESPACE "UAnimGraphNode_BlendSpaceGraphBase"
 
