Ensure that blend space graphs have the correct schema applied

Additional fix for cached poses in blend space graphs

diff --git a/Engine/Source/Editor/AnimGraph/Private/AnimGraphNode_BlendSpaceGraphBase.cpp b/Engine/Source/Editor/AnimGraph/Private/AnimGraphNode_BlendSpaceGraphBase.cpp
--- a/Engine/Source/Editor/AnimGraph/Private/AnimGraphNode_BlendSpaceGraphBase.cpp
+++ b/Engine/Source/Editor/AnimGraph/Private/AnimGraphNode_BlendSpaceGraphBase.cpp
@@ -139,7 +139,7 @@ void UAnimGraphNode_BlendSpaceGraphBase::SetupFromAsset(UBlendSpace* InBlendSpac
 		UAnimBlueprint* AnimBlueprint = GetAnimBlueprint();
 		const UAnimationGraphSchema* AnimationGraphSchema = GetDefault<UAnimationGraphSchema>();
 
-		BlendSpaceGraph = CastChecked<UBlendSpaceGraph>(FBlueprintEditorUtils::CreateNewGraph(this, InBlendSpace->GetFName(), UBlendSpaceGraph::StaticClass(), UAnimationGraphSchema::StaticClass()));
+		BlendSpaceGraph = CastChecked<UBlendSpaceGraph>(FBlueprintEditorUtils::CreateNewGraph(this, InBlendSpace->GetFName(), UBlendSpaceGraph::StaticClass(), UEdGraphSchema::StaticClass()));
 		BlendSpaceGraph->BlendSpace = BlendSpace = DuplicateObject(InBlendSpace, BlendSpaceGraph);
 		BlendSpace->ClearFlags(RF_Public | RF_Standalone);
 		BlendSpaceGraph->BlendSpace->ResetSkeleton(nullptr);
@@ -179,7 +179,7 @@ void UAnimGraphNode_BlendSpaceGraphBase::SetupFromClass(TSubclassOf<UBlendSpace>
 	}
 	else
 	{
-		BlendSpaceGraph = CastChecked<UBlendSpaceGraph>(FBlueprintEditorUtils::CreateNewGraph(this, InBlendSpaceClass.Get()->GetFName(), UBlendSpaceGraph::StaticClass(), UAnimationGraphSchema::StaticClass()));
+		BlendSpaceGraph = CastChecked<UBlendSpaceGraph>(FBlueprintEditorUtils::CreateNewGraph(this, InBlendSpaceClass.Get()->GetFName(), UBlendSpaceGraph::StaticClass(), UEdGraphSchema::StaticClass()));
 		BlendSpaceGraph->BlendSpace = BlendSpace = NewObject<UBlendSpace>(BlendSpaceGraph, InBlendSpaceClass.Get());
 		BlendSpaceGraph->BlendSpace->SetFlags(RF_Transactional);
 		BlendSpaceGraph->bAllowDeletion = false;
@@ -327,7 +327,7 @@ void UAnimGraphNode_BlendSpaceGraphBase::PostPlacedNewNode()
 	{
 		check(BlendSpace == nullptr);
 
-		BlendSpaceGraph = CastChecked<UBlendSpaceGraph>(FBlueprintEditorUtils::CreateNewGraph(this, NAME_None, UBlendSpaceGraph::StaticClass(), UAnimationGraphSchema::StaticClass()));
+		BlendSpaceGraph = CastChecked<UBlendSpaceGraph>(FBlueprintEditorUtils::CreateNewGraph(this, NAME_None, UBlendSpaceGraph::StaticClass(), UEdGraphSchema::StaticClass()));
 		BlendSpaceGraph->BlendSpace = BlendSpace = NewObject<UBlendSpace>();
 		
 		// Find an interesting name
diff --git a/Engine/Source/Editor/AnimGraph/Private/BlendSpaceGraph.cpp b/Engine/Source/Editor/AnimGraph/Private/BlendSpaceGraph.cpp
new file mode 100644
--- /dev/null
+++ b/Engine/Source/Editor/AnimGraph/Private/BlendSpaceGraph.cpp
@@ -0,0 +1,14 @@
+// Copyright Epic Games, Inc. All Rights Reserved.
+
+#include "BlendSpaceGraph.h"
+
+void UBlendSpaceGraph::PostLoad()
+{
+	Super::PostLoad();
+
+	// Fixup graph schema
+	if(Schema && Schema != UEdGraphSchema::StaticClass())
+	{
+		Schema = UEdGraphSchema::StaticClass();
+	}
+}
diff --git a/Engine/Source/Editor/AnimGraph/Public/BlendSpaceGraph.h b/Engine/Source/Editor/AnimGraph/Public/BlendSpaceGraph.h
--- a/Engine/Source/Editor/AnimGraph/Public/BlendSpaceGraph.h
+++ b/Engine/Source/Editor/AnimGraph/Public/BlendSpaceGraph.h
@@ -14,6 +14,10 @@ class UBlendSpaceGraph : public UEdGraph
 {
 	GENERATED_BODY()
 
+private:
+	// UObject interface
+	virtual void PostLoad() override;
+
 public:
 	// Blendspace that we wrap
 	UPROPERTY()
