Fix copy-paste error in fast-path detection code that ended up rejecting split pin access for structs with a native make function

diff --git a/Engine/Source/Editor/AnimGraph/Private/AnimBlueprintCompilerHandler_Base.cpp b/Engine/Source/Editor/AnimGraph/Private/AnimBlueprintCompilerHandler_Base.cpp
--- a/Engine/Source/Editor/AnimGraph/Private/AnimBlueprintCompilerHandler_Base.cpp
+++ b/Engine/Source/Editor/AnimGraph/Private/AnimBlueprintCompilerHandler_Base.cpp
@@ -932,7 +932,7 @@ bool FAnimBlueprintCompilerHandler_Base::FEvaluationHandlerRecord::CheckForSplit
 			if(Function && Function->HasMetaData(TEXT("NativeMakeFunc")) && IsWhitelistedNativeMake(Function->GetFName()))
 			{
 				// Idea here is to account for split pins, so we want to narrow the scope to not also include user-placed makes
-				UObject* SourceObject = Context.MessageLog.FindSourceObject(MakeStructNode);
+				UObject* SourceObject = Context.MessageLog.FindSourceObject(NativeMakeNode);
 				if(SourceObject && SourceObject->IsA<UAnimGraphNode_Base>())
 				{
 					return ForEachInputPin(NativeMakeNode, [this, &Context, &OriginalRecord](UEdGraphPin* InputPin)
