Fix non-unity CIS

diff --git a/Engine/Source/Editor/AnimGraph/Private/BlendSpaceGraph.cpp b/Engine/Source/Editor/AnimGraph/Private/BlendSpaceGraph.cpp
--- a/Engine/Source/Editor/AnimGraph/Private/BlendSpaceGraph.cpp
+++ b/Engine/Source/Editor/AnimGraph/Private/BlendSpaceGraph.cpp
@@ -1,6 +1,7 @@
 // Copyright Epic Games, Inc. All Rights Reserved.
 
 #include "BlendSpaceGraph.h"
+#include "EdGraph/EdGraphSchema.h"
 
 void UBlendSpaceGraph::PostLoad()
 {
