UE-77877
Automatically connect Anim node's out exec pin to the following node when placed within an execution chain

diff --git a/Engine/Source/Editor/AnimGraph/Private/AnimGraphNode_Base.cpp b/Engine/Source/Editor/AnimGraph/Private/AnimGraphNode_Base.cpp
--- a/Engine/Source/Editor/AnimGraph/Private/AnimGraphNode_Base.cpp
+++ b/Engine/Source/Editor/AnimGraph/Private/AnimGraphNode_Base.cpp
@@ -525,4 +525,40 @@ void UAnimGraphNode_Base::PinConnectionListChanged(UEdGraphPin* Pin)
 	}
 }
 
+void UAnimGraphNode_Base::AutowireNewNode(UEdGraphPin* FromPin)
+{
+	// Ensure the pin is valid, a pose pin, and has a single link
+	if (FromPin && UAnimationGraphSchema::IsPosePin(FromPin->PinType))
+	{
+		auto FindFirstPosePinInDirection = [this](EEdGraphPinDirection Direction) -> UEdGraphPin*
+		{
+			UEdGraphPin** PinToConnectTo = Pins.FindByPredicate([Direction](UEdGraphPin* Pin) -> bool
+            {
+                return Pin && Pin->Direction == Direction && UAnimationGraphSchema::IsPosePin(Pin->PinType);
+            });
+
+			return PinToConnectTo ? *PinToConnectTo : nullptr;
+		};
+		
+		// Get the linked pin, if valid, and ensure it iss also a pose pin
+		UEdGraphPin* LinkedPin = FromPin->LinkedTo.Num() == 1 ? FromPin->LinkedTo[0] : nullptr;
+		if (LinkedPin && UAnimationGraphSchema::IsPosePin(LinkedPin->PinType))
+		{
+			// Retrieve the first pin, of similar direction, from this node
+			UEdGraphPin* PinToConnectTo = FindFirstPosePinInDirection(FromPin->Direction);
+			if (PinToConnectTo)
+			{
+				ensure(GetSchema()->TryCreateConnection(LinkedPin, PinToConnectTo));
+			}
+		}
+
+		// Link this node to the FromPin, so find the first pose pin of opposite direction on this node
+		UEdGraphPin* PinToConnectTo = FindFirstPosePinInDirection(FromPin->Direction == EEdGraphPinDirection::EGPD_Input ? EEdGraphPinDirection::EGPD_Output : EEdGraphPinDirection::EGPD_Input);
+		if (PinToConnectTo)
+		{
+			ensure(GetSchema()->TryCreateConnection(FromPin, PinToConnectTo));
+		}
+	}
+}
+
 #undef LOCTEXT_NAMESPACE
diff --git a/Engine/Source/Editor/AnimGraph/Public/AnimGraphNode_Base.h b/Engine/Source/Editor/AnimGraph/Public/AnimGraphNode_Base.h
--- a/Engine/Source/Editor/AnimGraph/Public/AnimGraphNode_Base.h
+++ b/Engine/Source/Editor/AnimGraph/Public/AnimGraphNode_Base.h
@@ -204,6 +204,7 @@ class ANIMGRAPH_API UAnimGraphNode_Base : public UK2Node
 	virtual FString GetPinMetaData(FName InPinName, FName InKey) override;
 	virtual void AddSearchMetaDataInfo(TArray<struct FSearchTagDataPair>& OutTaggedMetaData) const override;
 	virtual void PinConnectionListChanged(UEdGraphPin* Pin) override;
+	virtual void AutowireNewNode(UEdGraphPin* FromPin) override;
 	// End of UEdGraphNode interface
 
 	// UK2Node interface
