Fix for crash when opening anim layer interfaces
#rnx

diff --git a/Engine/Source/Editor/AnimGraph/Private/AnimGraphNode_SequencePlayer.cpp b/Engine/Source/Editor/AnimGraph/Private/AnimGraphNode_SequencePlayer.cpp
--- a/Engine/Source/Editor/AnimGraph/Private/AnimGraphNode_SequencePlayer.cpp
+++ b/Engine/Source/Editor/AnimGraph/Private/AnimGraphNode_SequencePlayer.cpp
@@ -308,7 +308,8 @@ bool UAnimGraphNode_SequencePlayer::IsActionFilteredOut(class FBlueprintActionFi
 
 	for (UBlueprint* Blueprint : FilterContext.Blueprints)
 	{
-		if (UAnimBlueprint* AnimBlueprint = Cast<UAnimBlueprint>(Blueprint))
+		UAnimBlueprint* AnimBlueprint = Cast<UAnimBlueprint>(Blueprint);
+		if (AnimBlueprint && AnimBlueprint->TargetSkeleton)
 		{
 			if(Node.Sequence)
 			{
@@ -330,7 +331,7 @@ bool UAnimGraphNode_SequencePlayer::IsActionFilteredOut(class FBlueprintActionFi
 		}
 		else
 		{
-			// Not an animation Blueprint, cannot use
+			// Not an animation Blueprint or has no target skeleton, cannot use
 			bIsFilteredOut = true;
 			break;
 		}
