Added state aliasing for the anim graph state machine.

diff --git a/Engine/Source/Editor/AnimGraph/Private/AnimBlueprintCompilerHandler_StateMachine.cpp b/Engine/Source/Editor/AnimGraph/Private/AnimBlueprintCompilerHandler_StateMachine.cpp
--- a/Engine/Source/Editor/AnimGraph/Private/AnimBlueprintCompilerHandler_StateMachine.cpp
+++ b/Engine/Source/Editor/AnimGraph/Private/AnimBlueprintCompilerHandler_StateMachine.cpp
@@ -258,9 +258,9 @@ void FAnimBlueprintCompilerHandler_StateMachine::ProcessTransitionGetter(UK2Node
 							const int32 MachinePropertyIndex = InCompilationContext.GetAllocatedAnimNodeIndices().FindChecked(CompiledMachineInstanceNode);
 							int32 TransitionPropertyIndex = INDEX_NONE;
 
-							for(TMap<TWeakObjectPtr<UEdGraphNode>, int32>::TIterator TransIt(DebugData->NodeToTransitionIndex); TransIt; ++TransIt)
+							for(auto TransIt = DebugData->NodeToTransitionIndex.CreateConstIterator(); TransIt; ++TransIt)
 							{
-								UEdGraphNode* CurrTransNode = TransIt.Key().Get();
+								const UEdGraphNode* CurrTransNode = TransIt.Key().Get();
 								
 								if(CurrTransNode == SourceTransitionNode)
 								{
@@ -442,9 +442,13 @@ void FAnimBlueprintCompilerHandler_StateMachine::AutoWireAnimGetter(class UK2Nod
 		{
 			if(FStateMachineDebugData* DebugData = OutCompiledData.GetAnimBlueprintDebugData().StateMachineDebugData.Find(TransitionNode->GetGraph()))
 			{
-				if(int32* TransitionIndexPtr = DebugData->NodeToTransitionIndex.Find(TransitionNode))
+				TArray<int32> TransitionIndices;
+				DebugData->NodeToTransitionIndex.MultiFind(TransitionNode, TransitionIndices);
+				const int32 TransNum = TransitionIndices.Num();
+				
+				if (TransNum > 0)
 				{
-					SubNodeIndex = *TransitionIndexPtr;
+					SubNodeIndex = TransitionIndices[0];
 				}
 			}
 		}
diff --git a/Engine/Source/Editor/AnimGraph/Private/AnimGraphNode_StateMachineBase.cpp b/Engine/Source/Editor/AnimGraph/Private/AnimGraphNode_StateMachineBase.cpp
--- a/Engine/Source/Editor/AnimGraph/Private/AnimGraphNode_StateMachineBase.cpp
+++ b/Engine/Source/Editor/AnimGraph/Private/AnimGraphNode_StateMachineBase.cpp
@@ -29,6 +29,7 @@
 #include "IAnimBlueprintGeneratedClassCompiledData.h"
 #include "IAnimBlueprintCompilationContext.h"
 #include "Animation/AnimNode_Inertialization.h"
+#include "AnimStateAliasNode.h"
 
 /////////////////////////////////////////////////////
 // FAnimStateMachineNodeNameValidator
@@ -226,9 +227,15 @@ void UAnimGraphNode_StateMachineBase::OnProcessDuringCompilation(IAnimBlueprintC
 	struct FMachineCreator
 	{
 	public:
+
+		// A transition node with a state alias as source may have multiple transitions belonging to the same node. The referenced aliased state to differentiates them during compilation. 
+		using FUniqueTransition = TPair<UAnimStateTransitionNode*, UAnimStateNodeBase*>;
+
 		int32 MachineIndex;
 		TMap<UAnimStateNodeBase*, int32> StateIndexTable;
-		TMap<UAnimStateTransitionNode*, int32> TransitionIndexTable;
+		TMultiMap<int32, int32> StateIndexToStateAliasNodeIndices;
+		TArray<const UAnimStateAliasNode*> StateAliasNodes;
+		TMap<FUniqueTransition, int32> TransitionIndexTable;
 		UAnimGraphNode_StateMachineBase* StateMachineInstance;
 		TArray<FBakedAnimationStateMachine>& BakedMachines;
 		IAnimBlueprintGeneratedClassCompiledData& CompiledData;
@@ -264,6 +271,16 @@ void UAnimGraphNode_StateMachineBase::OnProcessDuringCompilation(IAnimBlueprintC
 			return CompiledData.GetAnimBlueprintDebugData().StateMachineDebugData.FindOrAdd(SourceGraph);
 		}
 
+		int32 FindState(UAnimStateNodeBase* StateNode)
+		{
+			if (int32* pResult = StateIndexTable.Find(StateNode))
+			{
+				return *pResult;
+			}
+			
+			return INDEX_NONE;
+		}
+
 		int32 FindOrAddState(UAnimStateNodeBase* StateNode)
 		{
 			if (int32* pResult = StateIndexTable.Find(StateNode))
@@ -280,6 +297,7 @@ void UAnimGraphNode_StateMachineBase::OnProcessDuringCompilation(IAnimBlueprintC
 
 				UAnimStateNodeBase* SourceNode = CompilationContext.GetMessageLog().FindSourceObjectTypeChecked<UAnimStateNodeBase>(StateNode);
 				GetMachineSpecificDebugData().NodeToStateIndex.Add(SourceNode, StateIndex);
+				GetMachineSpecificDebugData().StateIndexToNode.Add(StateIndex, SourceNode);
 				if (UAnimStateNode* SourceStateNode = Cast<UAnimStateNode>(SourceNode))
 				{
 					CompiledData.GetAnimBlueprintDebugData().StateGraphToNodeMap.Add(SourceStateNode->BoundGraph, SourceStateNode);
@@ -289,9 +307,15 @@ void UAnimGraphNode_StateMachineBase::OnProcessDuringCompilation(IAnimBlueprintC
 			}
 		}
 
-		int32 FindOrAddTransition(UAnimStateTransitionNode* TransitionNode)
+		void AddStateAliasTransitionMapping(UAnimStateAliasNode* AliasNode, const FStateMachineDebugData::FStateAliasTransitionStateIndexPair& TransitionStateIndexPair)
 		{
-			if (int32* pResult = TransitionIndexTable.Find(TransitionNode))
+			UAnimStateAliasNode* SourceNode = CompilationContext.GetMessageLog().FindSourceObjectTypeChecked<UAnimStateAliasNode>(AliasNode);
+			GetMachineSpecificDebugData().StateAliasNodeToTransitionStatePairs.Add(SourceNode, TransitionStateIndexPair);
+		}
+
+		int32 FindOrAddTransition(FUniqueTransition UniqueTransition)
+		{
+			if (int32* pResult = TransitionIndexTable.Find(UniqueTransition))
 			{
 				return *pResult;
 			}
@@ -300,9 +324,10 @@ void UAnimGraphNode_StateMachineBase::OnProcessDuringCompilation(IAnimBlueprintC
 				FBakedAnimationStateMachine& BakedMachine = GetMachine();
 
 				const int32 TransitionIndex = BakedMachine.Transitions.Num();
-				TransitionIndexTable.Add(TransitionNode, TransitionIndex);
+				TransitionIndexTable.Add(UniqueTransition, TransitionIndex);
 				new (BakedMachine.Transitions) FAnimationTransitionBetweenStates();
 
+				UAnimStateTransitionNode* TransitionNode = UniqueTransition.Get<0>();
 				UAnimStateTransitionNode* SourceTransitionNode = CompilationContext.GetMessageLog().FindSourceObjectTypeChecked<UAnimStateTransitionNode>(TransitionNode);
 				GetMachineSpecificDebugData().NodeToTransitionIndex.Add(SourceTransitionNode, TransitionIndex);
 				CompiledData.GetAnimBlueprintDebugData().TransitionGraphToNodeMap.Add(SourceTransitionNode->BoundGraph, SourceTransitionNode);
@@ -361,6 +386,11 @@ void UAnimGraphNode_StateMachineBase::OnProcessDuringCompilation(IAnimBlueprintC
 	{
 		UEdGraphNode* Node = *StateNodeIt;
 
+		if (UAnimStateNodeBase* StateNodeBase = Cast<UAnimStateNodeBase>(Node))
+		{
+			StateNodeBase->ValidateNodeDuringCompilation(InCompilationContext.GetMessageLog());
+		}
+
 		if (UAnimStateEntryNode* EntryNode = Cast<UAnimStateEntryNode>(Node))
 		{
 			// Handle the state graph entry
@@ -369,6 +399,10 @@ void UAnimGraphNode_StateMachineBase::OnProcessDuringCompilation(IAnimBlueprintC
 			{
 				InCompilationContext.GetMessageLog().Error(*LOCTEXT("TooManyStateMachineEntryNodes", "Found an extra entry node @@").ToString(), EntryNode);
 			}
+			else if (UAnimStateAliasNode* AliasState = Cast<UAnimStateAliasNode>(EntryNode->GetOutputNode()))
+			{
+				InCompilationContext.GetMessageLog().Error(*LOCTEXT("AliasAsEntryState", "An alias (@@) cannot be used as the entry node for a state machine").ToString(), AliasState);
+			}
 			else if (UAnimStateNodeBase* StartState = Cast<UAnimStateNodeBase>(EntryNode->GetOutputNode()))
 			{
 				BakedMachine.InitialState = Oven.FindOrAddState(StartState);
@@ -378,49 +412,8 @@ void UAnimGraphNode_StateMachineBase::OnProcessDuringCompilation(IAnimBlueprintC
 				InCompilationContext.GetMessageLog().Warning(*LOCTEXT("NoConnection", "Entry node @@ is not connected to state").ToString(), EntryNode);
 			}
 		}
-		else if (UAnimStateTransitionNode* TransitionNode = Cast<UAnimStateTransitionNode>(Node))
-		{
-			TransitionNode->ValidateNodeDuringCompilation(InCompilationContext.GetMessageLog());
-
-			const int32 TransitionIndex = Oven.FindOrAddTransition(TransitionNode);
-			FAnimationTransitionBetweenStates& BakedTransition = Oven.GetMachine().Transitions[TransitionIndex];
-
-			BakedTransition.CrossfadeDuration = TransitionNode->CrossfadeDuration;
-			BakedTransition.StartNotify = OutCompiledData.FindOrAddNotify(TransitionNode->TransitionStart);
-			BakedTransition.EndNotify = OutCompiledData.FindOrAddNotify(TransitionNode->TransitionEnd);
-			BakedTransition.InterruptNotify = OutCompiledData.FindOrAddNotify(TransitionNode->TransitionInterrupt);
-			BakedTransition.BlendMode = TransitionNode->BlendMode;
-			BakedTransition.CustomCurve = TransitionNode->CustomBlendCurve;
-			BakedTransition.BlendProfile = TransitionNode->BlendProfile;
-			BakedTransition.LogicType = TransitionNode->LogicType;
-
-			UAnimStateNodeBase* PreviousState = TransitionNode->GetPreviousState();
-			UAnimStateNodeBase* NextState = TransitionNode->GetNextState();
-
-			if ((PreviousState != NULL) && (NextState != NULL))
-			{
-				const int32 PreviousStateIndex = Oven.FindOrAddState(PreviousState);
-				const int32 NextStateIndex = Oven.FindOrAddState(NextState);
-
-				if (TransitionNode->Bidirectional)
-				{
-					InCompilationContext.GetMessageLog().Warning(*LOCTEXT("BidirectionalTransWarning", "Bidirectional transitions aren't supported yet @@").ToString(), TransitionNode);
-				}
-
-				BakedTransition.PreviousState = PreviousStateIndex;
-				BakedTransition.NextState = NextStateIndex;
-			}
-			else
-			{
-				InCompilationContext.GetMessageLog().Warning(*LOCTEXT("BogusTransition", "@@ is incomplete, without a previous and next state").ToString(), TransitionNode);
-				BakedTransition.PreviousState = INDEX_NONE;
-				BakedTransition.NextState = INDEX_NONE;
-			}
-		}
 		else if (UAnimStateNode* StateNode = Cast<UAnimStateNode>(Node))
 		{
-			StateNode->ValidateNodeDuringCompilation(InCompilationContext.GetMessageLog());
-
 			const int32 StateIndex = Oven.FindOrAddState(StateNode);
 			FBakedAnimationState& BakedState = Oven.GetMachine().States[StateIndex];
 
@@ -470,8 +463,6 @@ void UAnimGraphNode_StateMachineBase::OnProcessDuringCompilation(IAnimBlueprintC
 		}
 		else if (UAnimStateConduitNode* ConduitNode = Cast<UAnimStateConduitNode>(Node))
 		{
-			ConduitNode->ValidateNodeDuringCompilation(InCompilationContext.GetMessageLog());
-
 			const int32 StateIndex = Oven.FindOrAddState(ConduitNode);
 			FBakedAnimationState& BakedState = Oven.GetMachine().States[StateIndex];
 
@@ -490,6 +481,42 @@ void UAnimGraphNode_StateMachineBase::OnProcessDuringCompilation(IAnimBlueprintC
 			// be a separate allocation, and a state machine inside of this one caused stuff to shift around
 			checkSlow(&BakedState == &(Oven.GetMachine().States[StateIndex]));
 		}
+		else if (UAnimStateAliasNode* StateAliasNode = Cast<UAnimStateAliasNode>(Node))
+		{
+			Oven.StateAliasNodes.Add(StateAliasNode);
+		}
+	}
+
+	int32 NumStateAlias = Oven.StateAliasNodes.Num();
+	for (auto AliasIndex = 0; AliasIndex < NumStateAlias; ++AliasIndex)
+	{
+		const UAnimStateAliasNode* StateAliasNode = Oven.StateAliasNodes[AliasIndex];
+
+		auto MapAlias = [&](const UAnimStateNodeBase* StateNode)
+		{
+			if (StateNode)
+			{
+				if (int32* StateIndexPtr = Oven.StateIndexTable.Find(StateNode))
+				{
+					Oven.StateIndexToStateAliasNodeIndices.Add(*StateIndexPtr, AliasIndex);
+				}
+			}
+		};
+
+		if (StateAliasNode->bGlobalAlias)
+		{
+			for (auto StateNodeIt = Oven.StateIndexTable.CreateConstIterator(); StateNodeIt; ++StateNodeIt)
+			{
+				MapAlias(StateNodeIt->Key);
+			}
+		}
+		else
+		{
+			for (auto StateNodeWeakIt = StateAliasNode->GetAliasedStates().CreateConstIterator(); StateNodeWeakIt; ++StateNodeWeakIt)
+			{
+				MapAlias(StateNodeWeakIt->Get());
+			}
+		}
 	}
 
 	// Process transitions after all the states because getters within custom graphs may want to
@@ -535,19 +562,84 @@ void UAnimGraphNode_StateMachineBase::OnProcessDuringCompilation(IAnimBlueprintC
 		}
 		// Handle all the transitions out of this node
 		TArray<class UAnimStateTransitionNode*> TransitionList;
+
+		// Add aliased state transitions to transition list
+		TArray<int32> StateAliasIndices;
+		Oven.StateIndexToStateAliasNodeIndices.MultiFind(StateIndex, StateAliasIndices);
+		for (const int32 AliasIndex : StateAliasIndices)
+		{
+			// Let the final transition list do the sort.
+			Oven.StateAliasNodes[AliasIndex]->GetTransitionList(TransitionList, /*bWantSortedList=*/ false);
+		}
+		
 		StateNode->GetTransitionList(/*out*/ TransitionList, /*bWantSortedList=*/ true);
 
 		for (auto TransitionIt = TransitionList.CreateIterator(); TransitionIt; ++TransitionIt)
 		{
 			UAnimStateTransitionNode* TransitionNode = *TransitionIt;
-			const int32 TransitionIndex = Oven.FindOrAddTransition(TransitionNode);
+
+			const int32 TransitionIndex = Oven.FindOrAddTransition(FMachineCreator::FUniqueTransition(TransitionNode, StateNode));
+			FAnimationTransitionBetweenStates& BakedTransition = Oven.GetMachine().Transitions[TransitionIndex];
+
+			BakedTransition.CrossfadeDuration = TransitionNode->CrossfadeDuration;
+			BakedTransition.StartNotify = OutCompiledData.FindOrAddNotify(TransitionNode->TransitionStart);
+			BakedTransition.EndNotify = OutCompiledData.FindOrAddNotify(TransitionNode->TransitionEnd);
+			BakedTransition.InterruptNotify = OutCompiledData.FindOrAddNotify(TransitionNode->TransitionInterrupt);
+			BakedTransition.BlendMode = TransitionNode->BlendMode;
+			BakedTransition.CustomCurve = TransitionNode->CustomBlendCurve;
+			BakedTransition.BlendProfile = TransitionNode->BlendProfile;
+			BakedTransition.LogicType = TransitionNode->LogicType;
+
+			UAnimStateNodeBase* PreviousState = StateNode;
+			UAnimStateNodeBase* NextState = TransitionNode->GetNextState();
+
+			UAnimStateAliasNode* NextAliasNode = Cast<UAnimStateAliasNode>(NextState);
+			if (NextAliasNode)
+			{
+				NextState = NextAliasNode->GetAliasedState();
+			}
+
+			if ((PreviousState != nullptr) && (NextState != nullptr))
+			{
+				const int32 PreviousStateIndex = Oven.FindState(PreviousState);
+				const int32 NextStateIndex = Oven.FindState(NextState);
+
+				if (NextAliasNode)
+				{
+					Oven.AddStateAliasTransitionMapping(NextAliasNode, { TransitionIndex, NextStateIndex });
+				}
+
+				if (TransitionNode->Bidirectional)
+				{
+					InCompilationContext.GetMessageLog().Warning(*LOCTEXT("BidirectionalTransWarning", "Bidirectional transitions aren't supported yet @@").ToString(), TransitionNode);
+				}
+
+				BakedTransition.PreviousState = PreviousStateIndex;
+				BakedTransition.NextState = NextStateIndex;
+			}
+
+			if((BakedTransition.PreviousState == INDEX_NONE) || (BakedTransition.NextState == INDEX_NONE))
+			{
+				InCompilationContext.GetMessageLog().Error(*LOCTEXT("BogusTransition", "@@ is incomplete, without a previous or next state").ToString(), TransitionNode);
+			}
 
 			// Validate the blend profile for this transition - incase the skeleton of the node has
 			// changed or the blend profile no longer exists.
 			TransitionNode->ValidateBlendProfile();
 
 			FBakedStateExitTransition& Rule = *new (BakedState.Transitions) FBakedStateExitTransition();
-			Rule.bDesiredTransitionReturnValue = (TransitionNode->GetPreviousState() == StateNode);
+
+			UAnimStateNodeBase* TransitionPrevNode = TransitionNode->GetPreviousState();
+			if (UAnimStateAliasNode* PrevAliasNode = Cast<UAnimStateAliasNode>(TransitionPrevNode))
+			{
+				Rule.bDesiredTransitionReturnValue = PrevAliasNode->bGlobalAlias ? true : PrevAliasNode->GetAliasedStates().Contains(StateNode);
+				Oven.AddStateAliasTransitionMapping(PrevAliasNode, { TransitionIndex, StateIndex });
+			}
+			else
+			{
+				Rule.bDesiredTransitionReturnValue = (TransitionPrevNode == StateNode);
+			}
+
 			Rule.TransitionIndex = TransitionIndex;
 			
 			if (UAnimGraphNode_TransitionResult* TransitionResultNode = CastChecked<UAnimationTransitionGraph>(TransitionNode->BoundGraph)->GetResultNode())
diff --git a/Engine/Source/Editor/AnimGraph/Private/AnimStateAliasNode.cpp b/Engine/Source/Editor/AnimGraph/Private/AnimStateAliasNode.cpp
new file mode 100644
--- /dev/null
+++ b/Engine/Source/Editor/AnimGraph/Private/AnimStateAliasNode.cpp
@@ -0,0 +1,173 @@
+// Copyright Epic Games, Inc. All Rights Reserved.
+
+#include "AnimStateAliasNode.h"
+#include "Kismet2/Kismet2NameValidators.h"
+#include "Kismet2/CompilerResultsLog.h"
+#include "Kismet2/Kismet2NameValidators.h"
+#include "AnimationStateMachineGraph.h"
+#include "EdGraph/EdGraph.h"
+#include "AnimStateNode.h"
+
+#define LOCTEXT_NAMESPACE "AnimStateAliasNode"
+
+/////////////////////////////////////////////////////
+// UAnimStateAliasNode
+
+UAnimStateAliasNode::UAnimStateAliasNode(const FObjectInitializer& ObjectInitializer)
+	: Super(ObjectInitializer)
+{
+	bCanRenameNode = true;
+	StateAliasName = TEXT("Alias");
+}
+
+void UAnimStateAliasNode::Serialize(FArchive& Ar)
+{
+	Super::Serialize(Ar);
+	RebuildAliasedStateNodeReferences();
+}
+
+void UAnimStateAliasNode::AllocateDefaultPins()
+{
+	CreatePin(EGPD_Input, TEXT("Transition"), TEXT("In"));
+	CreatePin(EGPD_Output, TEXT("Transition"), TEXT("Out"));
+}
+
+void UAnimStateAliasNode::ValidateNodeDuringCompilation(class FCompilerResultsLog& MessageLog) const
+{
+	Super::ValidateNodeDuringCompilation(MessageLog);
+
+	if ((GetInputPin()->LinkedTo.Num() > 0) && (bGlobalAlias || AliasedStateNodes.Num() != 1))
+	{
+		MessageLog.Error(*LOCTEXT("AliasAsEntryState", "A alias (@@) used as a transition's target must alias a single state").ToString(), this);
+	}
+}
+
+void UAnimStateAliasNode::AutowireNewNode(UEdGraphPin* FromPin)
+{
+	Super::AutowireNewNode(FromPin);
+
+	if (FromPin)
+	{
+		if (GetSchema()->TryCreateConnection(FromPin, GetInputPin()))
+		{
+			FromPin->GetOwningNode()->NodeConnectionListChanged();
+		}
+	}
+}
+
+void UAnimStateAliasNode::PostPasteNode()
+{
+	// Find an interesting name, but try to keep the same if possible
+	TSharedPtr<INameValidatorInterface> NameValidator = FNameValidatorFactory::MakeValidator(this);
+	NameValidator->FindValidString(StateAliasName);
+}
+
+void UAnimStateAliasNode::PostPlacedNewNode()
+{
+	// Find an interesting name, but try to keep the same if possible
+	TSharedPtr<INameValidatorInterface> NameValidator = FNameValidatorFactory::MakeValidator(this);
+	NameValidator->FindValidString(StateAliasName);
+}
+
+FText UAnimStateAliasNode::GetNodeTitle(ENodeTitleType::Type TitleType) const
+{
+	return FText::FromString(GetStateName());
+}
+
+FText UAnimStateAliasNode::GetTooltipText() const
+{
+	return LOCTEXT("ConduitNodeTooltip", "This is a conduit, which allows specification of a predicate condition for an entire group of transitions");
+}
+
+FString UAnimStateAliasNode::GetStateName() const
+{
+	return StateAliasName;
+}
+
+void UAnimStateAliasNode::OnRenameNode(const FString& NewName)
+{
+	StateAliasName = NewName;
+}
+
+UEdGraphPin* UAnimStateAliasNode::GetInputPin() const
+{
+	return Pins[0];
+}
+
+UEdGraphPin* UAnimStateAliasNode::GetOutputPin() const
+{
+	return Pins[1];
+}
+
+FString UAnimStateAliasNode::GetDesiredNewNodeName() const
+{
+	return TEXT("Alias");
+}
+
+UObject* UAnimStateAliasNode::GetJumpTargetForDoubleClick() const
+{
+	return GetAliasedState();
+}
+
+const TSet<TWeakObjectPtr<UAnimStateNodeBase>>& UAnimStateAliasNode::GetAliasedStates() const
+{
+	return AliasedStateNodes;
+}
+
+TSet<TWeakObjectPtr<UAnimStateNodeBase>>& UAnimStateAliasNode::GetAliasedStates()
+{
+	return AliasedStateNodes;
+}
+
+UAnimStateNodeBase* UAnimStateAliasNode::GetAliasedState() const
+{
+	// If we alias more than one state, we return null
+	if (bGlobalAlias || (AliasedStateNodes.Num() != 1))
+	{
+		return nullptr;
+	}
+
+
+	if (UAnimStateNodeBase* AliasedState = AliasedStateNodes.CreateConstIterator()->Get())
+	{
+		if (!AliasedState->IsPendingKill())
+		{
+			if (const UEdGraph* Graph = GetGraph())
+			{
+				TArray<UAnimStateNodeBase*> StateNodes;
+				Graph->GetNodesOfClassEx<UAnimStateNode, UAnimStateNodeBase>(StateNodes);
+				return StateNodes.Contains(AliasedState) ? AliasedState : nullptr;
+			}
+		}
+	}
+	
+	return nullptr;
+}
+
+void UAnimStateAliasNode::RebuildAliasedStateNodeReferences()
+{
+	TSet<TWeakObjectPtr<UAnimStateNodeBase>> NewAliasedStateNodes;
+
+	// We don't use UEdGraphNode::GetGraph because this may be called during deletion and we don't want to assert on a missing graph.
+	if (const UEdGraph* Graph = Cast<UEdGraph>(GetOuter()))
+	{
+		TArray<UAnimStateNodeBase*> StateNodes;
+		Graph->GetNodesOfClassEx<UAnimStateNode, UAnimStateNodeBase>(StateNodes);
+		TSet<UAnimStateNodeBase*> StateNodesSet(StateNodes);
+
+		for (auto StateNodeIt = AliasedStateNodes.CreateIterator(); StateNodeIt; ++StateNodeIt)
+		{
+			UAnimStateNodeBase* AliasedState = StateNodeIt->Get();
+
+			// Keep only nodes that are still in the graph
+			if (AliasedState && !AliasedState->IsPendingKill() && StateNodesSet.Contains(AliasedState))
+			{
+				NewAliasedStateNodes.Add(*StateNodeIt);
+			}
+		}
+
+		AliasedStateNodes = NewAliasedStateNodes;
+	}
+}
+
+#undef LOCTEXT_NAMESPACE
diff --git a/Engine/Source/Editor/AnimGraph/Private/AnimStateConduitNode.cpp b/Engine/Source/Editor/AnimGraph/Private/AnimStateConduitNode.cpp
--- a/Engine/Source/Editor/AnimGraph/Private/AnimStateConduitNode.cpp
+++ b/Engine/Source/Editor/AnimGraph/Private/AnimStateConduitNode.cpp
@@ -64,33 +64,6 @@ UEdGraphPin* UAnimStateConduitNode::GetOutputPin() const
 	return Pins[1];
 }
 
-void UAnimStateConduitNode::GetTransitionList(TArray<class UAnimStateTransitionNode*>& OutTransitions, bool bWantSortedList)
-{
-	// Normal transitions
-	for (int32 LinkIndex = 0; LinkIndex < Pins[1]->LinkedTo.Num(); ++LinkIndex)
-	{
-		UEdGraphNode* TargetNode = Pins[1]->LinkedTo[LinkIndex]->GetOwningNode();
-		if (UAnimStateTransitionNode* Transition = Cast<UAnimStateTransitionNode>(TargetNode))
-		{
-			OutTransitions.Add(Transition);
-		}
-	}
-
-	// Sort the transitions by priority order, lower numbers are higher priority
-	if (bWantSortedList)
-	{
-		struct FCompareTransitionsByPriority
-		{
-			FORCEINLINE bool operator()(const UAnimStateTransitionNode& A, const UAnimStateTransitionNode& B) const
-			{
-				return A.PriorityOrder < B.PriorityOrder;
-			}
-		};
-
-		OutTransitions.Sort(FCompareTransitionsByPriority());
-	}
-}
-
 void UAnimStateConduitNode::PostPlacedNewNode()
 {
 	// Create a new animation graph
diff --git a/Engine/Source/Editor/AnimGraph/Private/AnimStateNode.cpp b/Engine/Source/Editor/AnimGraph/Private/AnimStateNode.cpp
--- a/Engine/Source/Editor/AnimGraph/Private/AnimStateNode.cpp
+++ b/Engine/Source/Editor/AnimGraph/Private/AnimStateNode.cpp
@@ -70,46 +70,6 @@ FString UAnimStateNode::GetStateName() const
 	return (BoundGraph != NULL) ? *(BoundGraph->GetName()) : TEXT("(null)");
 }
 
-void UAnimStateNode::GetTransitionList(TArray<UAnimStateTransitionNode*>& OutTransitions, bool bWantSortedList)
-{
-	// Normal transitions
-	for (int32 LinkIndex = 0; LinkIndex < Pins[1]->LinkedTo.Num(); ++LinkIndex)
-	{
-		UEdGraphNode* TargetNode = Pins[1]->LinkedTo[LinkIndex]->GetOwningNode();
-		if (UAnimStateTransitionNode* Transition = Cast<UAnimStateTransitionNode>(TargetNode))
-		{
-			OutTransitions.Add(Transition);
-		}
-	}
-
-	// Bidirectional transitions where we are the 'backwards' link
-	for (int32 LinkIndex = 0; LinkIndex < Pins[0]->LinkedTo.Num(); ++LinkIndex)
-	{
-		UEdGraphNode* TargetNode = Pins[0]->LinkedTo[LinkIndex]->GetOwningNode();
-		if (UAnimStateTransitionNode* Transition = Cast<UAnimStateTransitionNode>(TargetNode))
-		{
-			if (Transition->Bidirectional)
-			{
-				OutTransitions.Add(Transition);
-			}
-		}
-	}
-
-	// Sort the transitions by priority order, lower numbers are higher priority
-	if (bWantSortedList)
-	{
-		struct FCompareTransitionsByPriority
-		{
-			FORCEINLINE bool operator()(const UAnimStateTransitionNode& A, const UAnimStateTransitionNode& B) const
-			{
-				return A.PriorityOrder < B.PriorityOrder;
-			}
-		};
-
-		OutTransitions.Sort(FCompareTransitionsByPriority());
-	}
-}
-
 
 UEdGraphPin* UAnimStateNode::GetInputPin() const
 {
diff --git a/Engine/Source/Editor/AnimGraph/Private/AnimStateNodeBase.cpp b/Engine/Source/Editor/AnimGraph/Private/AnimStateNodeBase.cpp
--- a/Engine/Source/Editor/AnimGraph/Private/AnimStateNodeBase.cpp
+++ b/Engine/Source/Editor/AnimGraph/Private/AnimStateNodeBase.cpp
@@ -12,6 +12,7 @@
 #include "AnimationStateMachineSchema.h"
 #include "Kismet2/Kismet2NameValidators.h"
 #include "Kismet2/KismetEditorUtilities.h"
+#include "AnimStateTransitionNode.h"
 /////////////////////////////////////////////////////
 // FAnimStateNodeNameValidator
 
@@ -146,6 +147,48 @@ void UAnimStateNodeBase::PostLoad()
 	}
 }
 
+void UAnimStateNodeBase::GetTransitionList(TArray<UAnimStateTransitionNode*>& OutTransitions, bool bWantSortedList /*= false*/) const
+{
+	// Normal transitions
+	for (int32 LinkIndex = 0; LinkIndex < Pins[1]->LinkedTo.Num(); ++LinkIndex)
+	{
+		UEdGraphNode* TargetNode = Pins[1]->LinkedTo[LinkIndex]->GetOwningNode();
+		if (UAnimStateTransitionNode* Transition = Cast<UAnimStateTransitionNode>(TargetNode))
+		{
+			OutTransitions.Add(Transition);
+		}
+	}
+
+	// Bidirectional transitions where we are the 'backwards' link.
+	// Conduits and other states types that don't support bidirectional transitions should hide it from the details panel.
+	for (int32 LinkIndex = 0; LinkIndex < Pins[0]->LinkedTo.Num(); ++LinkIndex)
+	{
+		UEdGraphNode* TargetNode = Pins[0]->LinkedTo[LinkIndex]->GetOwningNode();
+		if (UAnimStateTransitionNode* Transition = Cast<UAnimStateTransitionNode>(TargetNode))
+		{
+			// Anim state nodes that don't support bidirectional transitions should hide this property in FAnimTransitionNodeDetails::CustomizeDetails
+			if (Transition->Bidirectional)
+			{
+				OutTransitions.Add(Transition);
+			}
+		}
+	}
+
+	// Sort the transitions by priority order, lower numbers are higher priority
+	if (bWantSortedList)
+	{
+		struct FCompareTransitionsByPriority
+		{
+			FORCEINLINE bool operator()(const UAnimStateTransitionNode& A, const UAnimStateTransitionNode& B) const
+			{
+				return A.PriorityOrder < B.PriorityOrder;
+			}
+		};
+
+		OutTransitions.Sort(FCompareTransitionsByPriority());
+	}
+}
+
 UAnimBlueprint* UAnimStateNodeBase::GetAnimBlueprint() const
 {
 	UBlueprint* Blueprint = FBlueprintEditorUtils::FindBlueprintForNode(this);
diff --git a/Engine/Source/Editor/AnimGraph/Private/AnimationStateMachineSchema.cpp b/Engine/Source/Editor/AnimGraph/Private/AnimationStateMachineSchema.cpp
--- a/Engine/Source/Editor/AnimGraph/Private/AnimationStateMachineSchema.cpp
+++ b/Engine/Source/Editor/AnimGraph/Private/AnimationStateMachineSchema.cpp
@@ -21,6 +21,7 @@
 #include "AnimStateNode.h"
 #include "AnimStateTransitionNode.h"
 #include "AnimStateConduitNode.h"
+#include "AnimStateAliasNode.h"
 #include "AnimGraphNode_AssetPlayerBase.h"
 #include "ScopedTransaction.h"
 #include "GraphEditorActions.h"
@@ -277,6 +278,12 @@ void UAnimationStateMachineSchema::GetGraphContextActions(FGraphContextMenuBuild
 		Action->NodeTemplate = NewObject<UAnimStateNode>(ContextMenuBuilder.OwnerOfTemporaries);
 	}
 
+	// Add state alias node
+	{
+		TSharedPtr<FEdGraphSchemaAction_NewStateNode> Action = AddNewStateNodeAction(ContextMenuBuilder, FText::GetEmpty(), LOCTEXT("AddStateAlias", "Add State Alias..."), LOCTEXT("AddStateAliasTooltip", "A new state alias"));
+		Action->NodeTemplate = NewObject<UAnimStateAliasNode>(ContextMenuBuilder.OwnerOfTemporaries);
+	}
+
 	// Add conduit node
 	{
 		TSharedPtr<FEdGraphSchemaAction_NewStateNode> Action = AddNewStateNodeAction(ContextMenuBuilder, FText::GetEmpty(), LOCTEXT("AddConduit", "Add Conduit..."), LOCTEXT("AddConduitTooltip", "A new conduit state"));
diff --git a/Engine/Source/Editor/AnimGraph/Public/AnimStateAliasNode.h b/Engine/Source/Editor/AnimGraph/Public/AnimStateAliasNode.h
new file mode 100644
--- /dev/null
+++ b/Engine/Source/Editor/AnimGraph/Public/AnimStateAliasNode.h
@@ -0,0 +1,64 @@
+// Copyright Epic Games, Inc. All Rights Reserved.
+
+#pragma once
+
+#include "CoreMinimal.h"
+#include "UObject/ObjectMacros.h"
+#include "Animation/AnimTypes.h"
+#include "AnimStateNodeBase.h"
+#include "AnimStateAliasNode.generated.h"
+
+class UEdGraph;
+class UEdGraphPin;
+class UAnimStateNode;
+
+UCLASS(MinimalAPI)
+class UAnimStateAliasNode : public UAnimStateNodeBase
+{
+	GENERATED_UCLASS_BODY()
+public:
+
+	UPROPERTY(EditAnywhere, BlueprintReadWrite, Category = "State Alias")
+	bool bGlobalAlias;
+
+	UPROPERTY()
+	FString StateAliasName;
+
+	// UObject
+	virtual void Serialize(FArchive& Ar) override;
+
+	//~ Begin UEdGraphNode Interface
+	virtual void AllocateDefaultPins() override;
+	virtual void AutowireNewNode(UEdGraphPin* FromPin) override;
+	virtual FText GetNodeTitle(ENodeTitleType::Type TitleType) const override;
+	virtual FText GetTooltipText() const override;
+	virtual bool CanDuplicateNode() const override { return true; }
+	virtual void OnRenameNode(const FString& NewName) override;
+	virtual void PostPasteNode() override;
+	virtual void PostPlacedNewNode() override;
+	virtual void ValidateNodeDuringCompilation(class FCompilerResultsLog& MessageLog) const override;
+	//~ End UEdGraphNode Interface
+
+	//~ Begin UAnimStateNodeBase Interface
+	virtual UEdGraphPin* GetInputPin() const override;
+	virtual UEdGraphPin* GetOutputPin() const override;
+	virtual FString GetStateName() const override;
+	virtual FString GetDesiredNewNodeName() const override;
+	UObject* GetJumpTargetForDoubleClick() const;
+	//~ End UAnimStateNodeBase Interface
+	
+	ANIMGRAPH_API const TSet<TWeakObjectPtr<UAnimStateNodeBase>>& GetAliasedStates() const;
+	ANIMGRAPH_API TSet<TWeakObjectPtr<UAnimStateNodeBase>>& GetAliasedStates();
+
+	// Returns null if aliasing more than one state.
+	ANIMGRAPH_API UAnimStateNodeBase* GetAliasedState() const;
+
+	static FName GetAliasedStateNodesPropertyName() { return GET_MEMBER_NAME_CHECKED(UAnimStateAliasNode, AliasedStateNodes); }
+
+private:
+
+	void RebuildAliasedStateNodeReferences();
+
+	UPROPERTY()
+	TSet<TWeakObjectPtr<UAnimStateNodeBase>> AliasedStateNodes;
+};
diff --git a/Engine/Source/Editor/AnimGraph/Public/AnimStateConduitNode.h b/Engine/Source/Editor/AnimGraph/Public/AnimStateConduitNode.h
--- a/Engine/Source/Editor/AnimGraph/Public/AnimStateConduitNode.h
+++ b/Engine/Source/Editor/AnimGraph/Public/AnimStateConduitNode.h
@@ -37,7 +37,6 @@ public:
 	virtual UEdGraphPin* GetInputPin() const override;
 	virtual UEdGraphPin* GetOutputPin() const override;
 	virtual FString GetStateName() const override;
-	virtual void GetTransitionList(TArray<class UAnimStateTransitionNode*>& OutTransitions, bool bWantSortedList = true) override;
 	virtual FString GetDesiredNewNodeName() const override;
 	virtual UEdGraph* GetBoundGraph() const override { return BoundGraph; }
 	virtual void ClearBoundGraph() override { BoundGraph = nullptr; }
diff --git a/Engine/Source/Editor/AnimGraph/Public/AnimStateNode.h b/Engine/Source/Editor/AnimGraph/Public/AnimStateNode.h
--- a/Engine/Source/Editor/AnimGraph/Public/AnimStateNode.h
+++ b/Engine/Source/Editor/AnimGraph/Public/AnimStateNode.h
@@ -66,7 +66,6 @@ public:
 	virtual UEdGraphPin* GetInputPin() const override;
 	virtual UEdGraphPin* GetOutputPin() const override;
 	virtual FString GetStateName() const override;
-	virtual void GetTransitionList(TArray<class UAnimStateTransitionNode*>& OutTransitions, bool bWantSortedList = false) override;
 	//~ End UAnimStateNodeBase Interface
 
 	// @return the pose pin of the state sink node within the anim graph of this state
diff --git a/Engine/Source/Editor/AnimGraph/Public/AnimStateNodeBase.h b/Engine/Source/Editor/AnimGraph/Public/AnimStateNodeBase.h
--- a/Engine/Source/Editor/AnimGraph/Public/AnimStateNodeBase.h
+++ b/Engine/Source/Editor/AnimGraph/Public/AnimStateNodeBase.h
@@ -41,7 +41,7 @@ class UAnimStateNodeBase : public UEdGraphNode
 	ANIMGRAPH_API virtual FString GetStateName() const { return TEXT("BaseState"); }
 
 	// Populates the OutTransitions array with a list of transition nodes connected to this state
-	ANIMGRAPH_API virtual void GetTransitionList(TArray<class UAnimStateTransitionNode*>& OutTransitions, bool bWantSortedList = false) { }
+	ANIMGRAPH_API virtual void GetTransitionList(TArray<class UAnimStateTransitionNode*>& OutTransitions, bool bWantSortedList = false) const;
 
 	virtual UEdGraph* GetBoundGraph() const { return NULL; }
 	virtual void ClearBoundGraph() {}
@@ -55,5 +55,4 @@ protected:
 public:
 	// Gets the animation state node documentation link
 	virtual FString GetDocumentationLink() const override;
-
 };
diff --git a/Engine/Source/Editor/AnimationBlueprintEditor/Private/AnimationGraphFactory.cpp b/Engine/Source/Editor/AnimationBlueprintEditor/Private/AnimationGraphFactory.cpp
--- a/Engine/Source/Editor/AnimationBlueprintEditor/Private/AnimationGraphFactory.cpp
+++ b/Engine/Source/Editor/AnimationBlueprintEditor/Private/AnimationGraphFactory.cpp
@@ -10,6 +10,7 @@
 #include "AnimGraphNode_LayeredBoneBlend.h"
 #include "AnimGraphNode_BlendSpaceBase.h"
 #include "AnimStateNode.h"
+#include "AnimStateAliasNode.h"
 #include "AnimStateEntryNode.h"
 #include "AnimStateConduitNode.h"
 #include "AnimStateTransitionNode.h"
@@ -18,6 +19,7 @@
 #include "AnimationGraphSchema.h"
 
 #include "AnimationStateNodes/SGraphNodeAnimState.h"
+#include "AnimationStateNodes/SGraphNodeAnimStateAlias.h"
 #include "AnimationStateNodes/SGraphNodeAnimTransition.h"
 #include "AnimationStateNodes/SGraphNodeAnimStateEntry.h"
 
@@ -77,6 +79,10 @@ TSharedPtr<class SGraphNode> FAnimationGraphNodeFactory::CreateNode(class UEdGra
 	{
 		return SNew(SGraphNodeAnimState, StateNode);
 	}
+	else if (UAnimStateAliasNode* StateAliasNode = Cast<UAnimStateAliasNode>(InNode))
+	{
+		return SNew(SGraphNodeAnimStateAlias, StateAliasNode);
+	}
 	else if (UAnimStateConduitNode* ConduitNode = Cast<UAnimStateConduitNode>(InNode))
 	{
 		return SNew(SGraphNodeAnimConduit, ConduitNode);
diff --git a/Engine/Source/Editor/AnimationBlueprintEditor/Private/AnimationStateNodes/SGraphNodeAnimState.cpp b/Engine/Source/Editor/AnimationBlueprintEditor/Private/AnimationStateNodes/SGraphNodeAnimState.cpp
--- a/Engine/Source/Editor/AnimationBlueprintEditor/Private/AnimationStateNodes/SGraphNodeAnimState.cpp
+++ b/Engine/Source/Editor/AnimationBlueprintEditor/Private/AnimationStateNodes/SGraphNodeAnimState.cpp
@@ -140,6 +140,11 @@ FSlateColor SGraphNodeAnimState::GetBorderBackgroundColor() const
 	FLinearColor ActiveStateColorDim(0.4f, 0.3f, 0.15f);
 	FLinearColor ActiveStateColorBright(1.f, 0.6f, 0.35f);
 
+	return GetBorderBackgroundColor_Internal(InactiveStateColor, ActiveStateColorDim, ActiveStateColorBright);
+}
+
+FSlateColor SGraphNodeAnimState::GetBorderBackgroundColor_Internal(FLinearColor InactiveStateColor, FLinearColor ActiveStateColorDim, FLinearColor ActiveStateColorBright) const
+{
 	UAnimBlueprint* AnimBlueprint = Cast<UAnimBlueprint>(FBlueprintEditorUtils::FindBlueprintForNode(GraphNode));
 	if(AnimBlueprint)
 	{
diff --git a/Engine/Source/Editor/AnimationBlueprintEditor/Private/AnimationStateNodes/SGraphNodeAnimState.h b/Engine/Source/Editor/AnimationBlueprintEditor/Private/AnimationStateNodes/SGraphNodeAnimState.h
--- a/Engine/Source/Editor/AnimationBlueprintEditor/Private/AnimationStateNodes/SGraphNodeAnimState.h
+++ b/Engine/Source/Editor/AnimationBlueprintEditor/Private/AnimationStateNodes/SGraphNodeAnimState.h
@@ -47,6 +47,7 @@ public:
 	static void GetStateInfoPopup(UEdGraphNode* GraphNode, TArray<FGraphInformationPopupInfo>& Popups);
 protected:
 	FSlateColor GetBorderBackgroundColor() const;
+	virtual FSlateColor GetBorderBackgroundColor_Internal(FLinearColor InactiveStateColor, FLinearColor ActiveStateColorDim, FLinearColor ActiveStateColorBright) const;
 
 	virtual FText GetPreviewCornerText() const;
 	virtual const FSlateBrush* GetNameIcon() const;
diff --git a/Engine/Source/Editor/AnimationBlueprintEditor/Private/AnimationStateNodes/SGraphNodeAnimStateAlias.cpp b/Engine/Source/Editor/AnimationBlueprintEditor/Private/AnimationStateNodes/SGraphNodeAnimStateAlias.cpp
new file mode 100644
--- /dev/null
+++ b/Engine/Source/Editor/AnimationBlueprintEditor/Private/AnimationStateNodes/SGraphNodeAnimStateAlias.cpp
@@ -0,0 +1,174 @@
+// Copyright Epic Games, Inc. All Rights Reserved.
+
+#include "AnimationStateNodes/SGraphNodeAnimStateAlias.h"
+#include "AnimationStateNodes/SGraphNodeAnimTransition.h"
+#include "AnimStateAliasNode.h"
+#include "Widgets/SBoxPanel.h"
+#include "Widgets/Images/SImage.h"
+#include "Widgets/SToolTip.h"
+#include "Animation/AnimInstance.h"
+#include "Animation/AnimBlueprint.h"
+#include "Kismet2/BlueprintEditorUtils.h"
+#include "AnimationStateMachineGraph.h"
+#include "AnimGraphNode_StateMachineBase.h"
+#include "SGraphPanel.h"
+
+#define LOCTEXT_NAMESPACE "SGraphNodeAnimStateAlias"
+
+/////////////////////////////////////////////////////
+// SGraphNodeAnimStateAlias
+
+void SGraphNodeAnimStateAlias::Construct(const FArguments& InArgs, UAnimStateAliasNode* InNode)
+{
+	SGraphNodeAnimState::Construct(SGraphNodeAnimState::FArguments(), InNode);
+}
+
+void SGraphNodeAnimStateAlias::GetNodeInfoPopups(FNodeInfoContext* Context, TArray<FGraphInformationPopupInfo>& Popups) const
+{
+	GetStateInfoPopup(GraphNode, Popups);
+}
+
+FSlateColor SGraphNodeAnimStateAlias::GetBorderBackgroundColor_Internal(FLinearColor InactiveStateColor, FLinearColor ActiveStateColorDim, FLinearColor ActiveStateColorBright) const
+{
+	UAnimBlueprint* AnimBlueprint = Cast<UAnimBlueprint>(FBlueprintEditorUtils::FindBlueprintForNode(GraphNode));
+	if (AnimBlueprint)
+	{
+		UAnimInstance* ActiveObject = Cast<UAnimInstance>(AnimBlueprint->GetObjectBeingDebugged());
+		UAnimBlueprintGeneratedClass* Class = AnimBlueprint->GetAnimBlueprintGeneratedClass();
+
+		// Display various types of debug data
+		if ((ActiveObject != NULL) && (Class != NULL))
+		{
+			UAnimationStateMachineGraph* StateMachineGraph = CastChecked<UAnimationStateMachineGraph>(GraphNode->GetGraph());
+			UAnimStateAliasNode* StateAliasNode = CastChecked<UAnimStateAliasNode>(GraphNode);
+
+			if (FStateMachineDebugData* DebugInfo = Class->GetAnimBlueprintDebugData().StateMachineDebugData.Find(StateMachineGraph))
+			{
+				TArray<int32, TInlineAllocator<16>> StatesToCheck;
+
+				// Check transitions associated with our state alias.
+				TArray<FStateMachineDebugData::FStateAliasTransitionStateIndexPair> TransitionStatePairs;
+				DebugInfo->StateAliasNodeToTransitionStatePairs.MultiFind(StateAliasNode, TransitionStatePairs);
+				const int32 PairsNum = TransitionStatePairs.Num();
+				for (int32 Index = 0; Index < PairsNum; ++Index)
+				{
+					auto& TransStatePair = TransitionStatePairs[Index];
+					if (SGraphNodeAnimTransition::IsTransitionActive(TransStatePair.TransitionIndex, *Class, *StateMachineGraph, *ActiveObject))
+					{
+						// Add the prev/next state we are aliasing that's associated with this transition.
+						StatesToCheck.Add(TransStatePair.AssociatedStateIndex);
+					}
+				}
+
+				// Use the highest aliased state's weight that has an active transition in/out of this alias. 
+				float Weight = 0.0f;
+				for (const int32 StateIndex : StatesToCheck)
+				{
+					if (TWeakObjectPtr<UAnimStateNodeBase>* StateNodeWeakPtr = DebugInfo->StateIndexToNode.Find(StateIndex))
+					{
+						if (StateAliasNode->GetAliasedStates().Contains(*StateNodeWeakPtr))
+						{
+							for (const FStateMachineStateDebugData& StateData : Class->GetAnimBlueprintDebugData().StateData)
+							{
+								if (StateData.StateMachineIndex == DebugInfo->MachineIndex && StateData.StateIndex == StateIndex)
+								{
+									Weight = FMath::Max(Weight, StateData.Weight);
+								}
+							}
+						}
+					}
+				}
+
+				if (Weight > 0.0f)
+				{
+					return FMath::Lerp<FLinearColor>(ActiveStateColorDim, ActiveStateColorBright, Weight);
+				}
+			}
+		}
+	}
+
+	return InactiveStateColor;
+}
+
+FText SGraphNodeAnimStateAlias::GetPreviewCornerText() const
+{
+	return FText();
+}
+
+const FSlateBrush* SGraphNodeAnimStateAlias::GetNameIcon() const
+{
+	// @todo: Change this to a unique icon and/or look for better visual distinction between state/conduit/alias.
+	return FEditorStyle::GetBrush(TEXT("Graph.TransitionNode.Icon"));
+}
+
+TSharedPtr<SToolTip> SGraphNodeAnimStateAlias::GetComplexTooltip()
+{
+	return nullptr;
+}
+
+void SGraphNodeAnimStateAlias::GetStateInfoPopup(UEdGraphNode* GraphNode, TArray<FGraphInformationPopupInfo>& Popups)
+{
+	UAnimBlueprint* AnimBlueprint = Cast<UAnimBlueprint>(FBlueprintEditorUtils::FindBlueprintForNode(GraphNode));
+	if (AnimBlueprint)
+	{
+		UAnimInstance* ActiveObject = Cast<UAnimInstance>(AnimBlueprint->GetObjectBeingDebugged());
+		UAnimBlueprintGeneratedClass* Class = AnimBlueprint->GetAnimBlueprintGeneratedClass();
+
+		static const FLinearColor PopupColor(1.f, 0.5f, 0.25f);
+
+		// Display various types of debug data
+		if ((ActiveObject != NULL) && (Class != NULL))
+		{
+			UAnimationStateMachineGraph* StateMachineGraph = CastChecked<UAnimationStateMachineGraph>(GraphNode->GetGraph());
+			UAnimStateAliasNode* StateAliasNode = CastChecked<UAnimStateAliasNode>(GraphNode);
+
+			if (FStateMachineDebugData* DebugInfo = Class->GetAnimBlueprintDebugData().StateMachineDebugData.Find(StateMachineGraph))
+			{
+				TArray<int32, TInlineAllocator<16>> StatesToCheck;
+
+				// Check transitions associated with our state alias.
+				TArray<FStateMachineDebugData::FStateAliasTransitionStateIndexPair> TransitionStatePairs;
+				DebugInfo->StateAliasNodeToTransitionStatePairs.MultiFind(StateAliasNode, TransitionStatePairs);
+				const int32 PairsNum = TransitionStatePairs.Num();
+				for (int32 Index = 0; Index < PairsNum; ++Index)
+				{
+					auto& TransStatePair = TransitionStatePairs[Index];
+					if (SGraphNodeAnimTransition::IsTransitionActive(TransStatePair.TransitionIndex, *Class, *StateMachineGraph, *ActiveObject))
+					{
+						// Add the prev/next state we are aliasing that's associated with this transition.
+						StatesToCheck.Add(TransStatePair.AssociatedStateIndex);
+					}
+				}
+
+				// Display the name and weight of any state we are aliasing that has an active incoming/outgoing transition from this alias
+				for (const int32 StateIndex : StatesToCheck)
+				{
+					if (TWeakObjectPtr<UAnimStateNodeBase>* StateNodeWeakPtr = DebugInfo->StateIndexToNode.Find(StateIndex))
+					{
+						if (UAnimStateNodeBase* StateNode = StateNodeWeakPtr->Get())
+						{
+							if (StateAliasNode->bGlobalAlias || StateAliasNode->GetAliasedStates().Contains(StateNode))
+							{
+								for (const FStateMachineStateDebugData& StateData : Class->GetAnimBlueprintDebugData().StateData)
+								{
+									if (StateData.StateMachineIndex == DebugInfo->MachineIndex && StateData.StateIndex == StateIndex)
+									{
+										FText StateText;
+										StateText = FText::Format(LOCTEXT("ActiveAliasedStateFormat", "{0}(Alias) {1}\nActive for {2}s"),
+											FText::FromString(StateNode->GetStateName()),
+											FText::AsPercent(StateData.Weight),
+											FText::AsNumber(StateData.ElapsedTime));
+										Popups.Emplace(nullptr, PopupColor, StateText.ToString());
+										break;
+									}
+								}
+							}
+						}
+					}
+				}
+			}
+		}
+	}
+}
+
+#undef LOCTEXT_NAMESPACE
\ No newline at end of file
diff --git a/Engine/Source/Editor/AnimationBlueprintEditor/Private/AnimationStateNodes/SGraphNodeAnimStateAlias.h b/Engine/Source/Editor/AnimationBlueprintEditor/Private/AnimationStateNodes/SGraphNodeAnimStateAlias.h
new file mode 100644
--- /dev/null
+++ b/Engine/Source/Editor/AnimationBlueprintEditor/Private/AnimationStateNodes/SGraphNodeAnimStateAlias.h
@@ -0,0 +1,29 @@
+// Copyright Epic Games, Inc. All Rights Reserved.
+
+#pragma once
+
+#include "SGraphNodeAnimState.h"
+
+class UAnimStateAliasNode;
+
+class SGraphNodeAnimStateAlias : public SGraphNodeAnimState
+{
+public:
+	SLATE_BEGIN_ARGS(SGraphNodeAnimStateAlias) {}
+	SLATE_END_ARGS()
+
+	void Construct(const FArguments& InArgs, UAnimStateAliasNode* InNode);
+
+	// SNodePanel::SNode interface
+	virtual void GetNodeInfoPopups(FNodeInfoContext* Context, TArray<FGraphInformationPopupInfo>& Popups) const override;
+	virtual TSharedPtr<SToolTip> GetComplexTooltip() override;
+	// End of SNodePanel::SNode interface
+
+	static void GetStateInfoPopup(UEdGraphNode* GraphNode, TArray<FGraphInformationPopupInfo>& Popups);
+
+protected:
+	virtual FSlateColor GetBorderBackgroundColor_Internal(FLinearColor InactiveStateColor, FLinearColor ActiveStateColorDim, FLinearColor ActiveStateColorBright) const;
+
+	virtual FText GetPreviewCornerText() const override;
+	virtual const FSlateBrush* GetNameIcon() const override;
+};
diff --git a/Engine/Source/Editor/AnimationBlueprintEditor/Private/AnimationStateNodes/SGraphNodeAnimTransition.cpp b/Engine/Source/Editor/AnimationBlueprintEditor/Private/AnimationStateNodes/SGraphNodeAnimTransition.cpp
--- a/Engine/Source/Editor/AnimationBlueprintEditor/Private/AnimationStateNodes/SGraphNodeAnimTransition.cpp
+++ b/Engine/Source/Editor/AnimationBlueprintEditor/Private/AnimationStateNodes/SGraphNodeAnimTransition.cpp
@@ -259,7 +259,6 @@ FLinearColor SGraphNodeAnimTransition::StaticGetTransitionColor(UAnimStateTransi
 	check(AnimBlueprint);
 	UAnimInstance* ActiveObject = Cast<UAnimInstance>(AnimBlueprint->GetObjectBeingDebugged());
 	UAnimBlueprintGeneratedClass* Class = AnimBlueprint->GetAnimBlueprintGeneratedClass();
-	UEdGraph* StateMachineGraph = TransNode->GetGraph();
 
 	//@TODO: WIP fast path / slow path coloring
 	if (AnimBlueprint->bWarnAboutBlueprintUsage || ((ActiveObject != nullptr) && (ActiveObject->PCV_ShouldNotifyAboutNodesNotUsingFastPath() || ActiveObject->PCV_ShouldWarnAboutNodesNotUsingFastPath())))
@@ -278,24 +277,19 @@ FLinearColor SGraphNodeAnimTransition::StaticGetTransitionColor(UAnimStateTransi
 
 	if ((ActiveObject != NULL) && (Class != NULL))
 	{
+		UAnimationStateMachineGraph* StateMachineGraph = CastChecked<UAnimationStateMachineGraph>(TransNode->GetGraph());
 		if (FStateMachineDebugData* DebugInfo = Class->GetAnimBlueprintDebugData().StateMachineDebugData.Find(StateMachineGraph))
 		{
-			if (int32* pTransIndex = DebugInfo->NodeToTransitionIndex.Find(TransNode))
+			// A transition node could be associated with multiple transitions indicies when coming from an alias. Check all of them
+			TArray<int32> TransitionIndices;
+			DebugInfo->NodeToTransitionIndex.MultiFind(TransNode, TransitionIndices);
+			const int32 TransNum = TransitionIndices.Num();
+			for (int32 Index = 0; Index < TransNum; ++Index)
 			{
-				const int32 TransIndex = *pTransIndex;
-
-				if (Class->GetAnimNodeProperties().Num())
-				{
-					UAnimationStateMachineGraph* TypedGraph = CastChecked<UAnimationStateMachineGraph>(StateMachineGraph);
-
-					if (FAnimNode_StateMachine* CurrentInstance = Class->GetPropertyInstance<FAnimNode_StateMachine>(ActiveObject, TypedGraph->OwnerAnimGraphNode))
-					{
-						if (CurrentInstance->IsTransitionActive(TransIndex))
-						{
-							// We're active!
-							return ActiveColor;
-						}
-					}
+				if(IsTransitionActive(TransitionIndices[Index], *Class, *StateMachineGraph, *ActiveObject))
+				{	
+					// We're active!
+					return ActiveColor;
 				}
 			}
 		}
@@ -319,6 +313,21 @@ FLinearColor SGraphNodeAnimTransition::StaticGetTransitionColor(UAnimStateTransi
 	return bIsHovered ? HoverColor : BaseColor;
 }
 
+bool SGraphNodeAnimTransition::IsTransitionActive(int32 TransitionIndex, UAnimBlueprintGeneratedClass& AnimClass, UAnimationStateMachineGraph& StateMachineGraph, UAnimInstance& AnimInstance)
+{
+	if (AnimClass.GetAnimNodeProperties().Num())
+	{
+		if (FAnimNode_StateMachine* CurrentInstance = AnimClass.GetPropertyInstance<FAnimNode_StateMachine>(&AnimInstance, StateMachineGraph.OwnerAnimGraphNode))
+		{
+			if (CurrentInstance->IsTransitionActive(TransitionIndex))
+			{
+				return true;
+			}
+		}
+	}
+	return false;
+}
+
 FSlateColor SGraphNodeAnimTransition::GetTransitionColor() const
 {	
 	// Highlight the transition node when the node is hovered or when the previous state is hovered
diff --git a/Engine/Source/Editor/AnimationBlueprintEditor/Private/AnimationStateNodes/SGraphNodeAnimTransition.h b/Engine/Source/Editor/AnimationBlueprintEditor/Private/AnimationStateNodes/SGraphNodeAnimTransition.h
--- a/Engine/Source/Editor/AnimationBlueprintEditor/Private/AnimationStateNodes/SGraphNodeAnimTransition.h
+++ b/Engine/Source/Editor/AnimationBlueprintEditor/Private/AnimationStateNodes/SGraphNodeAnimTransition.h
@@ -41,6 +41,8 @@ public:
 	void PositionBetweenTwoNodesWithOffset(const FGeometry& StartGeom, const FGeometry& EndGeom, int32 NodeIndex, int32 MaxNodes) const;
 
 	static FLinearColor StaticGetTransitionColor(UAnimStateTransitionNode* TransNode, bool bIsHovered);
+
+	static bool IsTransitionActive(int32 TransitionIndex, class UAnimBlueprintGeneratedClass& AnimClass, class UAnimationStateMachineGraph& StateMachineGraph, class UAnimInstance& AnimInstance);
 private:
 	TSharedPtr<STextEntryPopup> TextEntryWidget;
 
diff --git a/Engine/Source/Editor/DetailCustomizations/Private/AnimStateAliasNodeDetails.cpp b/Engine/Source/Editor/DetailCustomizations/Private/AnimStateAliasNodeDetails.cpp
new file mode 100644
--- /dev/null
+++ b/Engine/Source/Editor/DetailCustomizations/Private/AnimStateAliasNodeDetails.cpp
@@ -0,0 +1,244 @@
+// Copyright Epic Games, Inc. All Rights Reserved.
+
+#include "AnimStateAliasNodeDetails.h"
+#include "Widgets/DeclarativeSyntaxSupport.h"
+#include "Widgets/Text/STextBlock.h"
+#include "DetailLayoutBuilder.h"
+#include "DetailWidgetRow.h"
+#include "DetailCategoryBuilder.h"
+#include "AnimStateAliasNode.h"
+#include "AnimStateNode.h"
+#include "AnimationStateMachineGraph.h"
+#include "PropertyHandle.h"
+#include "ScopedTransaction.h"
+#include "Kismet2/BlueprintEditorUtils.h"
+
+#define LOCTEXT_NAMESPACE "FAnimStateAliasNodeDetails"
+
+/////////////////////////////////////////////////////////////////////////
+
+
+TSharedRef<IDetailCustomization> FAnimStateAliasNodeDetails::MakeInstance()
+{
+	return MakeShareable( new FAnimStateAliasNodeDetails);
+}
+
+void FAnimStateAliasNodeDetails::CustomizeDetails( IDetailLayoutBuilder& DetailBuilder )
+{
+	// Get a handle to the node we're viewing
+	const TArray< TWeakObjectPtr<UObject> >& SelectedObjects = DetailBuilder.GetSelectedObjects();
+	if (SelectedObjects.Num() != 1)
+	{
+		return;
+	}
+
+	const TWeakObjectPtr<UObject>& CurrentObject = SelectedObjects[0];
+	UAnimStateAliasNode* StateAliasNode = nullptr;
+	if (CurrentObject.IsValid())
+	{
+		StateAliasNodeWeak = StateAliasNode = Cast<UAnimStateAliasNode>(CurrentObject.Get());
+	}
+
+	if (StateAliasNode == nullptr)
+	{
+		return;
+	}
+
+	GenerateStatePickerDetails(*StateAliasNode, DetailBuilder);
+}
+
+void FAnimStateAliasNodeDetails::GetReferenceableStates(const UAnimStateAliasNode& OwningNode, TSet<TWeakObjectPtr<UAnimStateNodeBase>>& OutStates) const
+{
+	UAnimationStateMachineGraph* StateMachine = CastChecked<UAnimationStateMachineGraph>(OwningNode.GetOuter());
+
+	TArray<UAnimStateNode*> Nodes;
+	StateMachine->GetNodesOfClass<UAnimStateNode>(Nodes);
+	for (auto NodeIt = Nodes.CreateIterator(); NodeIt; ++NodeIt)
+	{
+		auto Node = *NodeIt;
+		OutStates.Add(Node);
+	}
+}
+
+bool FAnimStateAliasNodeDetails::IsGlobalAlias() const
+{
+	if (UAnimStateAliasNode* StateAliasNode = StateAliasNodeWeak.Get())
+	{
+		return StateAliasNode->bGlobalAlias;
+	}
+	
+	return false;
+}
+
+void FAnimStateAliasNodeDetails::OnPropertyAliasAllStatesCheckboxChanged(ECheckBoxState NewState)
+{
+	if (UAnimStateAliasNode* StateAliasNode = StateAliasNodeWeak.Get())
+	{
+		FScopedTransaction Transaction(LOCTEXT("Undo_SelectAllAliasState", "Select all state alias"));
+
+		StateAliasNode->Modify();
+		if (NewState == ECheckBoxState::Checked)
+		{
+			StateAliasNode->GetAliasedStates() = ReferenceableStates;
+		}
+		else
+		{
+			StateAliasNode->GetAliasedStates().Reset();
+		}
+
+		if (UBlueprint* Blueprint = FBlueprintEditorUtils::FindBlueprintForNodeChecked(StateAliasNode))
+		{
+			FBlueprintEditorUtils::MarkBlueprintAsModified(Blueprint);
+		}
+	}
+}
+
+ECheckBoxState FAnimStateAliasNodeDetails::AreAllStatesAliased() const
+{
+	if (const UAnimStateAliasNode* StateAliasNode = StateAliasNodeWeak.Get())
+	{
+		const int32 NumAlisedStates = StateAliasNode->GetAliasedStates().Num();
+		const int32 NumStates = ReferenceableStates.Num();
+
+		if (NumAlisedStates == 0)
+		{
+			return ECheckBoxState::Unchecked;
+		}
+		else if(NumAlisedStates == NumStates)
+		{
+			return ECheckBoxState::Checked;
+		}
+	}
+
+	return ECheckBoxState::Undetermined;
+}
+
+void FAnimStateAliasNodeDetails::OnPropertyIsStateAliasedCheckboxChanged(ECheckBoxState NewState, const TWeakObjectPtr<UAnimStateNodeBase> StateNodeWeak)
+{
+
+	if (UAnimStateNodeBase* StateNode = StateNodeWeak.Get())
+	{
+		if (UAnimStateAliasNode* StateAliasNode = StateAliasNodeWeak.Get())
+		{
+			FScopedTransaction Transaction(LOCTEXT("Undo_AliasState", "Select state alias"));
+
+			StateAliasNode->Modify();
+			if (NewState == ECheckBoxState::Checked)
+			{
+				StateAliasNode->GetAliasedStates().FindOrAdd(StateNodeWeak);
+			}
+			else
+			{
+				StateAliasNode->GetAliasedStates().Remove(StateNodeWeak);
+			}
+
+			if (UBlueprint* Blueprint = FBlueprintEditorUtils::FindBlueprintForNodeChecked(StateAliasNode))
+			{
+				FBlueprintEditorUtils::MarkBlueprintAsModified(Blueprint);
+			}
+		}
+	}
+}
+
+ECheckBoxState FAnimStateAliasNodeDetails::IsStateAliased(const TWeakObjectPtr<UAnimStateNodeBase> StateNodeWeak) const
+{
+	if (const UAnimStateNodeBase* StateNode = StateNodeWeak.Get())
+	{
+		if (const UAnimStateAliasNode* StateAliasNode = StateAliasNodeWeak.Get())
+		{
+			if (StateAliasNode->GetAliasedStates().Find(StateNodeWeak))
+			{
+				return ECheckBoxState::Checked;
+			}
+		}
+	}
+
+	return ECheckBoxState::Unchecked;
+}
+
+void FAnimStateAliasNodeDetails::GenerateStatePickerDetails(UAnimStateAliasNode& AliasNode, IDetailLayoutBuilder& DetailBuilder)
+{
+	ReferenceableStates.Reset();
+	GetReferenceableStates(AliasNode, ReferenceableStates);
+
+	if (ReferenceableStates.Num() > 0)
+	{
+		IDetailCategoryBuilder& CategoryBuilder = DetailBuilder.EditCategory(FName(TEXT("State Alias")));
+		CategoryBuilder.AddProperty(GET_MEMBER_NAME_CHECKED(UAnimStateAliasNode, bGlobalAlias));
+
+		FDetailWidgetRow& HeaderWidgetRow = CategoryBuilder.AddCustomRow(LOCTEXT("SelectAll", "Select All"));
+
+		HeaderWidgetRow.NameContent()
+			[
+				SNew(STextBlock)
+				.Text(LOCTEXT("StateName", "Name"))
+				.Font(IDetailLayoutBuilder::GetDetailFontBold())
+			];
+
+		HeaderWidgetRow.ValueContent()
+			[
+				SNew(SHorizontalBox)
+				+ SHorizontalBox::Slot()
+				.AutoWidth()
+				.VAlign(VAlign_Center)
+				[
+					SNew(STextBlock)
+					.Text(LOCTEXT("SelectAllStatesPropertyValue", "Select All"))
+					.Font(IDetailLayoutBuilder::GetDetailFontBold())
+				]
+				+ SHorizontalBox::Slot()
+				.FillWidth(1.0f)
+				.HAlign(HAlign_Right)
+				.VAlign(VAlign_Center)
+				[
+					SNew(SCheckBox)
+					.IsChecked(this, &FAnimStateAliasNodeDetails::AreAllStatesAliased)
+					.OnCheckStateChanged(this, &FAnimStateAliasNodeDetails::OnPropertyAliasAllStatesCheckboxChanged)
+					.IsEnabled_Lambda([this]() -> bool 
+						{
+							return !IsGlobalAlias();
+						})
+				]
+			];
+
+		for (auto StateIt = ReferenceableStates.CreateConstIterator(); StateIt; ++StateIt)
+		{
+			const TWeakObjectPtr<UAnimStateNodeBase>& StateNodeWeak = *StateIt;
+			if (const UAnimStateNodeBase* StateNode = StateNodeWeak.Get())
+			{
+				FString StateName = StateNode->GetStateName();
+				FText StateText = FText::FromString(StateName);
+
+				FDetailWidgetRow& PropertyWidgetRow = CategoryBuilder.AddCustomRow(StateText);
+
+				PropertyWidgetRow.NameContent()
+					[
+						SNew(STextBlock)
+						.Text(StateText)
+						.ToolTipText(StateText)
+						.Font(IDetailLayoutBuilder::GetDetailFont())
+					];
+
+				PropertyWidgetRow.ValueContent()
+					[
+						SNew(SHorizontalBox)
+						+ SHorizontalBox::Slot()
+						.FillWidth(1.0f)
+						.HAlign(HAlign_Right)
+						.VAlign(VAlign_Center)
+						[
+							SNew(SCheckBox)
+							.IsChecked(this, &FAnimStateAliasNodeDetails::IsStateAliased, StateNodeWeak)
+							.OnCheckStateChanged(this, &FAnimStateAliasNodeDetails::OnPropertyIsStateAliasedCheckboxChanged, StateNodeWeak)
+							.IsEnabled_Lambda([this]() -> bool 
+								{
+								return !IsGlobalAlias();
+								})
+						]
+					];
+			}
+		}
+	}
+}
+
+#undef LOCTEXT_NAMESPACE
diff --git a/Engine/Source/Editor/DetailCustomizations/Private/AnimStateAliasNodeDetails.h b/Engine/Source/Editor/DetailCustomizations/Private/AnimStateAliasNodeDetails.h
new file mode 100644
--- /dev/null
+++ b/Engine/Source/Editor/DetailCustomizations/Private/AnimStateAliasNodeDetails.h
@@ -0,0 +1,37 @@
+// Copyright Epic Games, Inc. All Rights Reserved.
+
+#pragma once
+
+#include "CoreMinimal.h"
+#include "IDetailCustomization.h"
+
+class IDetailLayoutBuilder;
+
+class UAnimStateAliasNode;
+class UAnimStateNodeBase;
+
+class FAnimStateAliasNodeDetails : public IDetailCustomization
+{
+public:
+	/** Makes a new instance of this detail layout class for a specific detail view requesting it */
+	static TSharedRef<IDetailCustomization> MakeInstance();
+
+	/** IDetailCustomization interface */
+	virtual void CustomizeDetails( IDetailLayoutBuilder& DetailBuilder ) override;
+
+private:
+
+	void GetReferenceableStates(const UAnimStateAliasNode& OwningNode, TSet<TWeakObjectPtr<UAnimStateNodeBase>>& OutStates) const;
+
+	bool IsGlobalAlias() const;
+	void OnPropertyAliasAllStatesCheckboxChanged(ECheckBoxState NewState);
+	ECheckBoxState AreAllStatesAliased() const;
+	void OnPropertyIsStateAliasedCheckboxChanged(ECheckBoxState NewState, const TWeakObjectPtr<UAnimStateNodeBase> StateNodeWeak);
+	ECheckBoxState IsStateAliased(const TWeakObjectPtr<UAnimStateNodeBase> StateNodeWeak) const;
+
+	void GenerateStatePickerDetails(UAnimStateAliasNode& AliasNode, IDetailLayoutBuilder& DetailBuilder);
+
+	TWeakObjectPtr<UAnimStateAliasNode> StateAliasNodeWeak;
+	TSet<TWeakObjectPtr<UAnimStateNodeBase>> ReferenceableStates;
+};
+
diff --git a/Engine/Source/Editor/DetailCustomizations/Private/DetailCustomizations.cpp b/Engine/Source/Editor/DetailCustomizations/Private/DetailCustomizations.cpp
--- a/Engine/Source/Editor/DetailCustomizations/Private/DetailCustomizations.cpp
+++ b/Engine/Source/Editor/DetailCustomizations/Private/DetailCustomizations.cpp
@@ -28,6 +28,7 @@
 #include "AnimSequenceDetails.h"
 #include "AnimTransitionNodeDetails.h"
 #include "AnimStateNodeDetails.h"
+#include "AnimStateAliasNodeDetails.h"
 #include "PoseAssetDetails.h"
 #include "AnimationAssetDetails.h"
 #include "AmbientSoundDetails.h"
@@ -330,6 +331,7 @@ void FDetailCustomizationsModule::RegisterObjectCustomizations()
 	RegisterCustomClassLayout("EditorAnimSegment", FOnGetDetailCustomizationInstance::CreateStatic(&FAnimMontageSegmentDetails::MakeInstance));
 	RegisterCustomClassLayout("EditorAnimCompositeSegment", FOnGetDetailCustomizationInstance::CreateStatic(&FAnimMontageSegmentDetails::MakeInstance));
 	RegisterCustomClassLayout("AnimStateNode", FOnGetDetailCustomizationInstance::CreateStatic(&FAnimStateNodeDetails::MakeInstance));
+	RegisterCustomClassLayout("AnimStateAliasNode", FOnGetDetailCustomizationInstance::CreateStatic(&FAnimStateAliasNodeDetails::MakeInstance));
 	RegisterCustomClassLayout("AnimStateTransitionNode", FOnGetDetailCustomizationInstance::CreateStatic(&FAnimTransitionNodeDetails::MakeInstance));
 	RegisterCustomClassLayout("AnimGraphNode_Trail", FOnGetDetailCustomizationInstance::CreateStatic(&FAnimTrailNodeDetails::MakeInstance));
 	RegisterCustomClassLayout("PoseAsset", FOnGetDetailCustomizationInstance::CreateStatic(&FPoseAssetDetails::MakeInstance));
diff --git a/Engine/Source/Editor/UnrealEd/Private/AnimationEditorUtils.cpp b/Engine/Source/Editor/UnrealEd/Private/AnimationEditorUtils.cpp
--- a/Engine/Source/Editor/UnrealEd/Private/AnimationEditorUtils.cpp
+++ b/Engine/Source/Editor/UnrealEd/Private/AnimationEditorUtils.cpp
@@ -816,7 +816,13 @@ namespace AnimationEditorUtils
 			}
 			else if(UAnimStateNodeBase* StateNode = Cast<UAnimStateNodeBase>(CurrentNode))
 			{
-				ChildGraphs.AddUnique(StateNode->GetBoundGraph());
+				UEdGraph* BoundGraph = StateNode->GetBoundGraph();
+				if (BoundGraph == nullptr)
+				{
+					continue;
+				}
+
+				ChildGraphs.AddUnique(BoundGraph);
 
 				if(UAnimStateTransitionNode* TransitionNode = Cast<UAnimStateTransitionNode>(StateNode))
 				{
diff --git a/Engine/Source/Runtime/Engine/Classes/Animation/AnimBlueprintGeneratedClass.h b/Engine/Source/Runtime/Engine/Classes/Animation/AnimBlueprintGeneratedClass.h
--- a/Engine/Source/Runtime/Engine/Classes/Animation/AnimBlueprintGeneratedClass.h
+++ b/Engine/Source/Runtime/Engine/Classes/Animation/AnimBlueprintGeneratedClass.h
@@ -20,6 +20,8 @@ class UAnimGraphNode_Base;
 class UAnimGraphNode_StateMachineBase;
 class UAnimInstance;
 class UAnimStateNode;
+class UAnimStateNodeBase;
+class UAnimStateAliasNode;
 class UAnimStateTransitionNode;
 class UEdGraph;
 class USkeleton;
@@ -70,9 +72,21 @@ public:
 		: MachineIndex(INDEX_NONE)
 	{}
 
+	struct FStateAliasTransitionStateIndexPair
+	{
+		int32 TransitionIndex;
+		int32 AssociatedStateIndex;
+	};
+
 	// Map from state nodes to their state entry in a state machine
 	TMap<TWeakObjectPtr<UEdGraphNode>, int32> NodeToStateIndex;
-	TMap<TWeakObjectPtr<UEdGraphNode>, int32> NodeToTransitionIndex;
+	TMap<int32, TWeakObjectPtr<UAnimStateNodeBase>> StateIndexToNode;
+
+	// Transition nodes may be associated w/ multiple transition indicies when the source state is an alias
+	TMultiMap<TWeakObjectPtr<UEdGraphNode>, int32> NodeToTransitionIndex;
+
+	// Mapping between an alias and any transition indices it might be associated to (Both as source and target).
+	TMultiMap<TWeakObjectPtr<UAnimStateAliasNode>, FStateAliasTransitionStateIndexPair> StateAliasNodeToTransitionStatePairs;
 
 	// The animation node that leads into this state machine (A3 only)
 	TWeakObjectPtr<UAnimGraphNode_StateMachineBase> MachineInstanceNode;
diff --git a/Engine/Source/Runtime/Engine/Private/Animation/AnimInstanceProxy.cpp b/Engine/Source/Runtime/Engine/Private/Animation/AnimInstanceProxy.cpp
--- a/Engine/Source/Runtime/Engine/Private/Animation/AnimInstanceProxy.cpp
+++ b/Engine/Source/Runtime/Engine/Private/Animation/AnimInstanceProxy.cpp
@@ -142,7 +142,8 @@ void FAnimInstanceProxy::Initialize(UAnimInstance* InAnimInstance)
 #if WITH_EDITORONLY_DATA
 		if (UAnimBlueprint* Blueprint = Cast<UAnimBlueprint>(InAnimInstance->GetClass()->ClassGeneratedBy))
 		{
-			if (Blueprint->Status == BS_Error)
+			// If our blueprint is in an error or dirty state upon initialization, anim graph shouldn't run
+			if ((Blueprint->Status == BS_Error) || (Blueprint->Status == BS_Dirty))
 			{
 				RootNode = nullptr;
 			}
