Fix Invalid use of pointer in warning output

diff --git a/Engine/Source/Editor/AnimGraph/Private/AnimBlueprintCompiler.cpp b/Engine/Source/Editor/AnimGraph/Private/AnimBlueprintCompiler.cpp
--- a/Engine/Source/Editor/AnimGraph/Private/AnimBlueprintCompiler.cpp
+++ b/Engine/Source/Editor/AnimGraph/Private/AnimBlueprintCompiler.cpp
@@ -127,7 +127,9 @@ FAnimBlueprintCompilerContext::FAnimBlueprintCompilerContext(UAnimBlueprint* Sou
 
 		if(bNodeGuidsRegenerated)
 		{
-			UE_LOG(LogAnimation, Warning, TEXT("Animation Blueprint %s has nodes with invalid node guids that have been regenerated. This blueprint will not cook deterministically until it is resaved."), AnimBlueprint);
+			UE_LOG(LogAnimation, Warning, TEXT(
+				"Animation Blueprint %s has nodes with invalid node guids that have been regenerated. This blueprint "
+				"will not cook deterministically until it is resaved."), *AnimBlueprint->GetPathName());
 		}
 	}
 
