Anim node data/compiler refactor

Per-node constant data is now held on a generated struct as part of sparse class data.
Per-node mutable data (i.e. pin links/property access mappings) is now held on a generated 'mutable data' struct that is compiled as part of the generated class.

The anim BP compiler is now extended more conventionally using UAnimBlueprintExtension, derived from UBlueprintExtension. This directly replaces the older 'compiler handler' pattern that was added in an emergency fashion for 4.26. Anim graph nodes now request their required extensions and these are held on the UAnimBlueprint in the UBlueprint::Extensions array. The Extensions array is potentially refreshed with any node addition or removal. The Extensions array is force-refreshed each time an anim BP is compiled for the first time to deal with newly added or removed requirements.

Const-corrected a bunch of UAnimInstance/FAnimInstanceProxy APIs that rely on (now truly) const data.
Added a split state/constant version of FInputScaleBiasClamp to allow some of its data to be split into constants.
Tweaked alignment/ordering of FPoseLinkBase to save a few bytes per pose link.
Deprecated FAnimNode_Base::OverrideAsset in favor of a more UAnimGraphNode_Base-based approach. Individual nodes can still have runtime overrides via specific accessors. The new approach will also give us the oppurtunity to override multiple assets per node if required in the future.

Moved property access into Engine module & removed event support from it - this was never used.
Includes a thread-safety fix for 4.26 that hasnt made it over to 5.0 yet.
Reworked property access compilation API a little - construction/lifetime was a bit confusing previously.

Optimized path used to create UK2Node_StructMemberSet nodes in per-node custom events. When using mutable data, the structure used is large and very sparsely connected (i.e. only a few properties are written) so we only create pins that are actually going to be used, rather than creating all of them and conly connecting a few.

Patched the following nodes to use the new data approach:

- Asset players (sequences, blendspaces, aim offsets)
- Blend lists
- Ref poses
- Roots

diff --git a/Engine/Config/BaseEngine.ini b/Engine/Config/BaseEngine.ini
--- a/Engine/Config/BaseEngine.ini
+++ b/Engine/Config/BaseEngine.ini
@@ -923,6 +923,9 @@ NearClipPlane=10.0
 +PropertyRedirects=(OldName="AnimationBlueprintLibrary.FindBonePathToRoot.AnimationSequence",NewName="AnimationSequenceBase")
 +ClassRedirects=(OldName="/Script/Engine.BlendSpaceBase",NewName="/Script/Engine.BlendSpace")
 +FunctionRedirects=(OldName="Image.SetBrushSize",NewName="SetDesiredSizeOverride")
++ClassRedirects=(OldName="/Script/PropertyAccess.AnimBlueprintClassSubsystem_PropertyAccess",NewName="/Script/Engine.AnimBlueprintClassSubsystem_PropertyAccess")
++PackageRedirects=(OldName="/Script/PropertyAccess", Removed=true)
++StructRedirects=(OldName="AnimInstanceSubsystemData",NewName="/Script/Engine.AnimSubsystemInstance")
 
 [CoreUObject.Metadata]
 ; Note: UnrealHeaderTool should be rerun after making changes to MetadataRedirects to catch any keys specified in class headers
@@ -2391,6 +2394,7 @@ DefaultASTCQualityBySize=3
 +VersionedIntRValues=r.Shaders.Optimize
 +VersionedIntRValues=r.CompileShadersForDevelopment
 +VersionedIntRValues=r.MobileHDR
++VersionedIntRValues=r.UsePreExposure
 
 [/Script/UnrealEd.ChunkDependencyInfo]
 ; Entries can be added to this section of DefaultEngine.ini to set dependencies for package chunks, used by cooking and packaging
diff --git a/Engine/Plugins/Developer/PropertyAccessNode/Source/PropertyAccessNode/Private/K2Node_PropertyAccess.cpp b/Engine/Plugins/Developer/PropertyAccessNode/Source/PropertyAccessNode/Private/K2Node_PropertyAccess.cpp
--- a/Engine/Plugins/Developer/PropertyAccessNode/Source/PropertyAccessNode/Private/K2Node_PropertyAccess.cpp
+++ b/Engine/Plugins/Developer/PropertyAccessNode/Source/PropertyAccessNode/Private/K2Node_PropertyAccess.cpp
@@ -15,12 +15,13 @@
 #include "K2Node_VariableGet.h"
 #include "AnimationGraph.h"
 #include "Kismet2/BlueprintEditorUtils.h"
-#include "PropertyAccessCompilerHandler.h"
+#include "AnimBlueprintExtension_PropertyAccess.h"
 #include "Modules/ModuleManager.h"
 #include "IPropertyAccessEditor.h"
 #include "IPropertyAccessCompiler.h"
 #include "Features/IModularFeatures.h"
 #include "IAnimBlueprintCompilationContext.h"
+#include "Animation/AnimBlueprint.h"
 
 #define LOCTEXT_NAMESPACE "K2Node_PropertyAccess"
 
@@ -49,9 +50,9 @@ void UK2Node_PropertyAccess::ExpandNode(FKismetCompilerContext& InCompilerContex
 
 		// Create a copy event in the complied generated class
 		TUniquePtr<IAnimBlueprintCompilationContext> CompilationContext = IAnimBlueprintCompilationContext::Get(InCompilerContext);
-		FPropertyAccessCompilerHandler* PropertyAccessHandler = CompilationContext->GetHandler<FPropertyAccessCompilerHandler>("PropertyAccessCompilerHandler");
-		check(PropertyAccessHandler);
-		PropertyAccessHandler->AddCopy(Path, DestPropertyPath, EPropertyAccessBatchType::Batched, this);
+		UAnimBlueprintExtension_PropertyAccess* PropertyAccessExtension = UAnimBlueprintExtension::GetExtension<UAnimBlueprintExtension_PropertyAccess>(CastChecked<UAnimBlueprint>(GetBlueprint()));
+		check(PropertyAccessExtension);
+		PropertyAccessExtension->AddCopy(Path, DestPropertyPath, EPropertyAccessBatchType::Batched, this);
 
 		// Replace us with a get node
 		UK2Node_VariableGet* VariableGetNode = InCompilerContext.SpawnIntermediateNode<UK2Node_VariableGet>(this, InSourceGraph);
diff --git a/Engine/Plugins/Developer/PropertyAccessNode/Source/PropertyAccessNode/PropertyAccessNode.Build.cs b/Engine/Plugins/Developer/PropertyAccessNode/Source/PropertyAccessNode/PropertyAccessNode.Build.cs
--- a/Engine/Plugins/Developer/PropertyAccessNode/Source/PropertyAccessNode/PropertyAccessNode.Build.cs
+++ b/Engine/Plugins/Developer/PropertyAccessNode/Source/PropertyAccessNode/PropertyAccessNode.Build.cs
@@ -12,7 +12,6 @@ public class PropertyAccessNode : ModuleRules
 				"Core",
 				"CoreUObject",
 				"Slate",	
-				"PropertyAccess",
 				"GraphEditor",
 				"BlueprintGraph",
 				"EditorStyle",
diff --git a/Engine/Plugins/Experimental/Animation/PoseSearch/Source/Runtime/Private/AnimNode_MotionMatching.cpp b/Engine/Plugins/Experimental/Animation/PoseSearch/Source/Runtime/Private/AnimNode_MotionMatching.cpp
--- a/Engine/Plugins/Experimental/Animation/PoseSearch/Source/Runtime/Private/AnimNode_MotionMatching.cpp
+++ b/Engine/Plugins/Experimental/Animation/PoseSearch/Source/Runtime/Private/AnimNode_MotionMatching.cpp
@@ -117,7 +117,7 @@ void FAnimNode_MotionMatching::Update_AnyThread(const FAnimationUpdateContext& C
 		}
 
 		// Continue with the follow up sequence if we're finishing a one shot anim
-		if (!bJumpedToPose && SequencePlayerNode.Sequence && !SequencePlayerNode.bLoopAnimation)
+		if (!bJumpedToPose && SequencePlayerNode.GetSequence() && !SequencePlayerNode.GetLoopAnimation())
 		{
 			float AssetTimeAfterUpdate = AssetTime + Context.GetDeltaTime();
 			if (AssetTimeAfterUpdate > AssetLength)
@@ -223,9 +223,9 @@ void FAnimNode_MotionMatching::JumpToPose(const FAnimationUpdateContext& Context
 
 	// Immediately jump to the pose by updating the embedded sequence player node
 	const FPoseSearchDatabaseSequence& ResultDbSequence = Database->Sequences[Result.DbSequenceIdx];
-	SequencePlayerNode.Sequence = ResultDbSequence.Sequence;
+	SequencePlayerNode.SetSequence(ResultDbSequence.Sequence);
 	SequencePlayerNode.SetAccumulatedTime(Result.TimeOffsetSeconds);
-	SequencePlayerNode.bLoopAnimation = ResultDbSequence.bLoopAnimation;
+	SequencePlayerNode.SetLoopAnimation(ResultDbSequence.bLoopAnimation);
 
 	// Use inertial blending to smooth over the transition
 	// It would be cool in the future to adjust the blend time by amount of dissimilarity, but we'll need a standardized distance metric first.
diff --git a/Engine/Plugins/Experimental/Animation/PoseSearch/Source/Runtime/Public/PoseSearch/AnimNode_MotionMatching.h b/Engine/Plugins/Experimental/Animation/PoseSearch/Source/Runtime/Public/PoseSearch/AnimNode_MotionMatching.h
--- a/Engine/Plugins/Experimental/Animation/PoseSearch/Source/Runtime/Public/PoseSearch/AnimNode_MotionMatching.h
+++ b/Engine/Plugins/Experimental/Animation/PoseSearch/Source/Runtime/Public/PoseSearch/AnimNode_MotionMatching.h
@@ -58,7 +58,7 @@ public:
 private:
 
 	// Embedded sequence player node for playing animations from the motion matching database
-	FAnimNode_SequencePlayer SequencePlayerNode;
+	FAnimNode_SequencePlayer_Standalone SequencePlayerNode;
 
 	// The current pose we're playing from the database
 	int32 DbPoseIdx = INDEX_NONE;
diff --git a/Engine/Plugins/Runtime/PropertyAccess/Source/PropertyAccessEditor/Private/PropertyAccessEditor.cpp b/Engine/Plugins/Runtime/PropertyAccess/Source/PropertyAccessEditor/Private/PropertyAccessEditor.cpp
--- a/Engine/Plugins/Runtime/PropertyAccess/Source/PropertyAccessEditor/Private/PropertyAccessEditor.cpp
+++ b/Engine/Plugins/Runtime/PropertyAccess/Source/PropertyAccessEditor/Private/PropertyAccessEditor.cpp
@@ -131,15 +131,10 @@ struct FPropertyAccessEditorSystem
 		}
 		else
 		{
+			InContext.ErrorMessage = FText::Format(LOCTEXT("UnrecognisedProperty", "Property '{0}' is unrecognised in property path for @@"), InProperty ? FText::FromName(InProperty->GetFName()) : LOCTEXT("Null", "Null"));
 			return ESegmentResolveResult::Failed;
 		}
 
-		static const FName PropertyEventMetadata("PropertyEvent");
-		if(InProperty->HasMetaData(PropertyEventMetadata))
-		{
-			InSegment.Flags |= (uint16)EPropertyAccessSegmentFlags::Event;
-		}
-
 		return Result;
 	}
 
@@ -570,13 +565,13 @@ namespace PropertyAccess
 	}
 }
 
-FPropertyAccessLibraryCompiler::FPropertyAccessLibraryCompiler()
-	: Library(nullptr)
-	, Class(nullptr)
+FPropertyAccessLibraryCompiler::FPropertyAccessLibraryCompiler(FPropertyAccessLibrary* InLibrary, const UClass* InClass)
+	: Library(InLibrary)
+	, Class(InClass)
 {
 }
 
-void FPropertyAccessLibraryCompiler::BeginCompilation(const UClass* InClass)
+void FPropertyAccessLibraryCompiler::BeginCompilation()
 {
 	if(Class && Library)
 	{
diff --git a/Engine/Plugins/Runtime/PropertyAccess/Source/PropertyAccessEditor/Private/PropertyAccessEditor.h b/Engine/Plugins/Runtime/PropertyAccess/Source/PropertyAccessEditor/Private/PropertyAccessEditor.h
--- a/Engine/Plugins/Runtime/PropertyAccess/Source/PropertyAccessEditor/Private/PropertyAccessEditor.h
+++ b/Engine/Plugins/Runtime/PropertyAccess/Source/PropertyAccessEditor/Private/PropertyAccessEditor.h
@@ -29,10 +29,10 @@ namespace PropertyAccess
 class FPropertyAccessLibraryCompiler : public IPropertyAccessLibraryCompiler
 {
 public:
-	FPropertyAccessLibraryCompiler();
+	FPropertyAccessLibraryCompiler(FPropertyAccessLibrary* InLibrary, const UClass* InClass);
 
 	// IPropertyAccessLibraryCompiler interface
-	virtual void BeginCompilation(const UClass* InClass) override;
+	virtual void BeginCompilation() override;
 	virtual int32 AddCopy(TArrayView<FString> InSourcePath, TArrayView<FString> InDestPath, EPropertyAccessBatchType InBatchType, UObject* InAssociatedObject = nullptr) override;
 	virtual bool FinishCompilation() override;
 	virtual void IterateErrors(TFunctionRef<void(const FText&, UObject*)> InFunction) const override;
diff --git a/Engine/Plugins/Runtime/PropertyAccess/Source/PropertyAccessEditor/Private/PropertyAccessEditorModule.cpp b/Engine/Plugins/Runtime/PropertyAccess/Source/PropertyAccessEditor/Private/PropertyAccessEditorModule.cpp
--- a/Engine/Plugins/Runtime/PropertyAccess/Source/PropertyAccessEditor/Private/PropertyAccessEditorModule.cpp
+++ b/Engine/Plugins/Runtime/PropertyAccess/Source/PropertyAccessEditor/Private/PropertyAccessEditorModule.cpp
@@ -8,8 +8,6 @@
 #include "PropertyAccessEditor.h"
 #include "Modules/ModuleInterface.h"
 #include "Features/IModularFeatures.h"
-#include "IAnimBlueprintCompilerHandlerCollection.h"
-#include "AnimBlueprintCompilerHandler_PropertyAccess.h"
 
 class FPropertyAccessEditorModule : public IPropertyAccessEditor, public IModuleInterface
 {
@@ -18,19 +16,10 @@ public:
 	virtual void StartupModule() override
 	{
 		IModularFeatures::Get().RegisterModularFeature("PropertyAccessEditor", this);
-
-		// Register node compilation handlers
-		IAnimBlueprintCompilerHandlerCollection::RegisterHandler("PropertyAccessCompilerHandler", [](IAnimBlueprintCompilerCreationContext& InCreationContext)
-		{
-			return MakeUnique<FAnimBlueprintCompilerHandler_PropertyAccess>(InCreationContext);
-		});
 	}
 
 	virtual void ShutdownModule() override
 	{
-		// Register node compilation handlers
-		IAnimBlueprintCompilerHandlerCollection::UnregisterHandler("PropertyAccessCompilerHandler");
-
 		IModularFeatures::Get().UnregisterModularFeature("PropertyAccessEditor", this);
 	}
 
@@ -67,6 +56,11 @@ public:
 	{
 		PropertyAccess::MakeStringPath(InBindingChain, OutStringPath);
 	}
+
+	virtual TUniquePtr<IPropertyAccessLibraryCompiler> MakePropertyAccessCompiler(const FPropertyAccessLibraryCompilerArgs& InArgs) const override
+	{
+		return MakeUnique<FPropertyAccessLibraryCompiler>(&InArgs.Library, InArgs.ClassContext);
+	}
 };
 
 IMPLEMENT_MODULE(FPropertyAccessEditorModule, PropertyAccessEditor)
\ No newline at end of file
diff --git a/Engine/Plugins/Runtime/PropertyAccess/Source/PropertyAccessEditor/PropertyAccessEditor.Build.cs b/Engine/Plugins/Runtime/PropertyAccess/Source/PropertyAccessEditor/PropertyAccessEditor.Build.cs
--- a/Engine/Plugins/Runtime/PropertyAccess/Source/PropertyAccessEditor/PropertyAccessEditor.Build.cs
+++ b/Engine/Plugins/Runtime/PropertyAccess/Source/PropertyAccessEditor/PropertyAccessEditor.Build.cs
@@ -18,7 +18,6 @@ public class PropertyAccessEditor : ModuleRules
 		PrivateDependencyModuleNames.AddRange(
 			new string[]
 			{
-				"PropertyAccess",
 				"GraphEditor",
 				"BlueprintGraph",
 				"EditorStyle",
diff --git a/Engine/Source/Developer/BlueprintCompilerCppBackend/Private/BlueprintCompilerCppBackendAnim.cpp b/Engine/Source/Developer/BlueprintCompilerCppBackend/Private/BlueprintCompilerCppBackendAnim.cpp
--- a/Engine/Source/Developer/BlueprintCompilerCppBackend/Private/BlueprintCompilerCppBackendAnim.cpp
+++ b/Engine/Source/Developer/BlueprintCompilerCppBackend/Private/BlueprintCompilerCppBackendAnim.cpp
@@ -81,20 +81,21 @@ void FBackendHelperAnim::AddAnimNodeInitializationFunction(FEmitterLocalContext&
 			FEmitDefaultValueHelper::OuterGenerate(Context, InProperty, TEXT(""), reinterpret_cast<const uint8*>(InCDO), bInNewProperty ? nullptr : reinterpret_cast<const uint8*>(InParentCDO), FEmitDefaultValueHelper::EPropertyAccessOperator::None, FEmitDefaultValueHelper::EPropertyGenerationControlFlags::AllowProtected);
 			
 			// anim nodes constructed, finish anim node initialization:
-			if (UAnimBlueprintGeneratedClass* AnimClass = Cast<UAnimBlueprintGeneratedClass>(Context.GetCurrentlyGeneratedClass()))
-			{
-				for (int32 i = 0; i < AnimClass->GetExposedValueHandlers().Num(); ++i)
-				{
-					if (AnimClass->GetExposedValueHandlers()[i].ValueHandlerNodeProperty == InProperty)
-					{
-						const FString ClassName = FEmitHelper::GetCppName(Context.GetCurrentlyGeneratedClass());
-						const FString MemberName = FEmitHelper::GetCppName(const_cast<FProperty*>(InProperty));
-						Context.Body.AddLine(FString::Printf(TEXT("%s.SetExposedValueHandler(&CastChecked<UAnimClassData>(CastChecked<UDynamicClass>(%s::StaticClass())->%s)->GetExposedValueHandlers()[%d]);"), *MemberName, *ClassName, GET_MEMBER_NAME_STRING_CHECKED(UDynamicClass, AnimClassImplementation), i));
-
-						break;
-					}
-				}
-			}
+			// @TODO: if we are to support nativizaton of anim BPs this will need to be reworked
+// 			if (UAnimBlueprintGeneratedClass* AnimClass = Cast<UAnimBlueprintGeneratedClass>(Context.GetCurrentlyGeneratedClass()))
+// 			{
+// 				for (int32 i = 0; i < AnimClass->GetExposedValueHandlers().Num(); ++i)
+// 				{
+// 					if (AnimClass->GetExposedValueHandlers()[i].ValueHandlerNodeProperty == InProperty)
+// 					{
+// 						const FString ClassName = FEmitHelper::GetCppName(Context.GetCurrentlyGeneratedClass());
+// 						const FString MemberName = FEmitHelper::GetCppName(const_cast<FProperty*>(InProperty));
+// 						Context.Body.AddLine(FString::Printf(TEXT("%s.SetExposedValueHandler(&CastChecked<UAnimClassData>(CastChecked<UDynamicClass>(%s::StaticClass())->%s)->GetExposedValueHandlers()[%d]);"), *MemberName, *ClassName, GET_MEMBER_NAME_STRING_CHECKED(UDynamicClass, AnimClassImplementation), i));
+// 
+// 						break;
+// 					}
+// 				}
+// 			}
 
 			Context.Body.DecreaseIndent();
 			Context.Body.AddLine(TEXT("}"));
diff --git a/Engine/Source/Editor/AnimGraph/Private/AnimBlueprintCompilationContext.cpp b/Engine/Source/Editor/AnimGraph/Private/AnimBlueprintCompilationContext.cpp
--- a/Engine/Source/Editor/AnimGraph/Private/AnimBlueprintCompilationContext.cpp
+++ b/Engine/Source/Editor/AnimGraph/Private/AnimBlueprintCompilationContext.cpp
@@ -2,6 +2,7 @@
 
 #include "AnimBlueprintCompilationContext.h"
 #include "AnimBlueprintCompiler.h"
+#include "AnimBlueprintExtension.h"
 
 TUniquePtr<IAnimBlueprintCompilationContext> IAnimBlueprintCompilationContext::Get(FKismetCompilerContext& InKismetCompiler)
 {
@@ -93,9 +94,19 @@ TArrayView<const FName> FAnimBlueprintCompilationContext::GetAttributesFromNodeI
 	return CompilerContext->GetAttributesFromNode(InNode);
 }
 
-IAnimBlueprintCompilerHandler* FAnimBlueprintCompilationContext::GetHandlerInternal(FName InName) const
+bool FAnimBlueprintCompilationContext::IsAnimGraphNodeFoldedImpl(UAnimGraphNode_Base* InNode) const
 {
-	return CompilerContext->AnimBlueprintCompilerHandlerCollection.GetHandler<IAnimBlueprintCompilerHandler>(InName);
+	return CompilerContext->IsAnimGraphNodeFolded(InNode);
+}
+
+const IAnimBlueprintCompilationContext::FFoldedPropertyRecord* FAnimBlueprintCompilationContext::GetFoldedPropertyRecordImpl(UAnimGraphNode_Base* InNode, FName InPropertyName) const
+{
+	return CompilerContext->GetFoldedPropertyRecord(InNode, InPropertyName);
+}
+
+const FStructProperty* FAnimBlueprintCompilationContext::GetMutableDataPropertyImpl() const
+{
+	return CompilerContext->NewMutablesProperty;
 }
 
 FKismetCompilerContext* FAnimBlueprintCompilationContext::GetKismetCompiler() const
@@ -108,10 +119,10 @@ FCompilerResultsLog& FAnimBlueprintCompilationBracketContext::GetMessageLogImpl(
 	return CompilerContext->MessageLog;
 }
 
-IAnimBlueprintCompilerHandler* FAnimBlueprintCompilationBracketContext::GetHandlerInternal(FName InName) const
+const TMap<UAnimGraphNode_Base*, int32>& FAnimBlueprintCompilationBracketContext::GetAllocatedAnimNodeIndicesImpl() const
 {
-	return CompilerContext->AnimBlueprintCompilerHandlerCollection.GetHandler<IAnimBlueprintCompilerHandler>(InName);
-}	
+	return CompilerContext->AllocatedAnimNodeIndices;
+}
 
 FCompilerResultsLog& FAnimBlueprintCopyTermDefaultsContext::GetMessageLogImpl() const
 {
@@ -137,8 +148,3 @@ const FKismetCompilerOptions& FAnimBlueprintPostExpansionStepContext::GetCompile
 {
 	return CompilerContext->CompileOptions;
 }
-
-IAnimBlueprintCompilerHandler* FAnimBlueprintPostExpansionStepContext::GetHandlerInternal(FName InName) const
-{
-	return CompilerContext->AnimBlueprintCompilerHandlerCollection.GetHandler<IAnimBlueprintCompilerHandler>(InName);
-}
diff --git a/Engine/Source/Editor/AnimGraph/Private/AnimBlueprintCompilationContext.h b/Engine/Source/Editor/AnimGraph/Private/AnimBlueprintCompilationContext.h
--- a/Engine/Source/Editor/AnimGraph/Private/AnimBlueprintCompilationContext.h
+++ b/Engine/Source/Editor/AnimGraph/Private/AnimBlueprintCompilationContext.h
@@ -7,8 +7,12 @@
 #include "IAnimBlueprintCompilationBracketContext.h"
 #include "IAnimBlueprintCopyTermDefaultsContext.h"
 #include "IAnimBlueprintPostExpansionStepContext.h"
+#include "IAnimBlueprintNodeOverrideAssetsContext.h"
+#include "Animation/AnimationAsset.h"
 
 class FAnimBlueprintCompilerContext;
+class UAnimBlueprintExtension;
+struct FAnimNode_Base;
 
 class FAnimBlueprintCompilationContext : public IAnimBlueprintCompilationContext
 {
@@ -37,7 +41,9 @@ private:
 	virtual const TMap<UAnimGraphNode_Base*, FProperty*>& GetAllocatedPropertiesByNodeImpl() const override;
 	virtual void AddAttributesToNodeImpl(UAnimGraphNode_Base* InNode, TArrayView<const FName> InAttributes) const override;
 	virtual TArrayView<const FName> GetAttributesFromNodeImpl(UAnimGraphNode_Base* InNode) const override;
-	virtual IAnimBlueprintCompilerHandler* GetHandlerInternal(FName InName) const override;
+	virtual bool IsAnimGraphNodeFoldedImpl(UAnimGraphNode_Base* InNode) const override;
+	virtual const IAnimBlueprintCompilationContext::FFoldedPropertyRecord* GetFoldedPropertyRecordImpl(UAnimGraphNode_Base* InNode, FName InPropertyName) const override;
+	virtual const FStructProperty* GetMutableDataPropertyImpl() const override;
 	virtual FKismetCompilerContext* GetKismetCompiler() const override;
 
 	FAnimBlueprintCompilerContext* CompilerContext;
@@ -53,8 +59,8 @@ private:
 	{}
 
 	virtual FCompilerResultsLog& GetMessageLogImpl() const override;
-	virtual IAnimBlueprintCompilerHandler* GetHandlerInternal(FName InName) const override;
-
+	virtual const TMap<UAnimGraphNode_Base*, int32>& GetAllocatedAnimNodeIndicesImpl() const override;
+	
 	FAnimBlueprintCompilerContext* CompilerContext;
 };
 
@@ -96,6 +102,29 @@ private:
 	int32 NodePropertyIndex;
 };
 
+class FAnimBlueprintExtensionCopyTermDefaultsContext : public IAnimBlueprintExtensionCopyTermDefaultsContext
+{
+private:
+	friend class FAnimBlueprintCompilerContext;
+
+	FAnimBlueprintExtensionCopyTermDefaultsContext(const FProperty* InTargetProperty, uint8* InDestinationPtr, const uint8* InSourcePtr, int32 InNodePropertyIndex)
+        : TargetProperty(InTargetProperty)
+        , DestinationPtr(InDestinationPtr)
+        , SourcePtr(InSourcePtr)
+        , NodePropertyIndex(InNodePropertyIndex)
+	{}
+
+	virtual const FProperty* GetTargetPropertyImpl() const override { return TargetProperty; }
+	virtual uint8* GetDestinationPtrImpl() const override { return DestinationPtr; }
+	virtual const uint8* GetSourcePtrImpl() const override { return SourcePtr; }
+	virtual int32 GetNodePropertyIndexImpl() const override { return NodePropertyIndex; }
+
+	const FProperty* TargetProperty;
+	uint8* DestinationPtr;
+	const uint8* SourcePtr;
+	int32 NodePropertyIndex;
+};
+
 class FAnimBlueprintPostExpansionStepContext : public IAnimBlueprintPostExpansionStepContext
 {
 private:
@@ -108,7 +137,33 @@ private:
 	virtual FCompilerResultsLog& GetMessageLogImpl() const override;
 	virtual UEdGraph* GetConsolidatedEventGraphImpl() const override;
 	virtual const FKismetCompilerOptions& GetCompileOptionsImpl() const override;
-	virtual IAnimBlueprintCompilerHandler* GetHandlerInternal(FName InName) const override;
 
 	FAnimBlueprintCompilerContext* CompilerContext;
+};
+
+class FAnimBlueprintNodeOverrideAssetsContext : public IAnimBlueprintNodeOverrideAssetsContext
+{
+private:
+	friend class FAnimBlueprintCompilerContext;
+	
+	FAnimBlueprintNodeOverrideAssetsContext(FAnimNode_Base* InAnimNode, UScriptStruct* InAnimNodeType)
+		: AnimNode(InAnimNode)
+		, AnimNodeType(InAnimNodeType)
+	{}
+
+	// Add an asset to this context
+	void AddAsset(UObject* InAsset) { Assets.Add(InAsset); }
+
+	virtual FAnimNode_Base* GetAnimNodeImpl() const override { return AnimNode; }
+	virtual UScriptStruct* GetAnimNodeTypeImpl() const override { return AnimNodeType; }
+	virtual TArrayView<UObject* const> GetAssetsImpl() const override { return MakeArrayView(Assets); }
+
+	// The anim node we are targeting
+	FAnimNode_Base* AnimNode;
+
+	// The type of the anim node we are targeting
+	UScriptStruct* AnimNodeType;
+	
+	// All the assets to override
+	TArray<UObject*> Assets;
 };
\ No newline at end of file
diff --git a/Engine/Source/Editor/AnimGraph/Private/AnimBlueprintCompiler.cpp b/Engine/Source/Editor/AnimGraph/Private/AnimBlueprintCompiler.cpp
--- a/Engine/Source/Editor/AnimGraph/Private/AnimBlueprintCompiler.cpp
+++ b/Engine/Source/Editor/AnimGraph/Private/AnimBlueprintCompiler.cpp
@@ -1,20 +1,12 @@
 // Copyright Epic Games, Inc. All Rights Reserved.
 
 #include "AnimBlueprintCompiler.h"
-#include "UObject/UObjectHash.h"
 #include "Animation/AnimInstance.h"
 #include "EdGraphUtilities.h"
 #include "K2Node_CallFunction.h"
-#include "K2Node_StructMemberGet.h"
-#include "K2Node_BreakStruct.h"
-#include "K2Node_MakeStruct.h"
-#include "K2Node_CallArrayFunction.h"
-#include "K2Node_CustomEvent.h"
 #include "K2Node_Knot.h"
 #include "K2Node_StructMemberSet.h"
 #include "K2Node_VariableGet.h"
-#include "K2Node_VariableSet.h"
-#include "K2Node_GetArrayItem.h"
 
 #include "AnimationGraphSchema.h"
 #include "Kismet/BlueprintFunctionLibrary.h"
@@ -24,8 +16,6 @@
 #include "Kismet2/KismetReinstanceUtilities.h"
 #include "AnimGraphNode_Root.h"
 #include "Animation/AnimNode_CustomProperty.h"
-#include "Animation/AnimNode_LinkedAnimGraph.h"
-#include "AnimGraphNode_LinkedAnimGraph.h"
 #include "AnimationEditorUtils.h"
 #include "AnimationGraph.h"
 #include "AnimBlueprintPostCompileValidation.h" 
@@ -34,25 +24,38 @@
 #include "K2Node_FunctionResult.h"
 #include "AnimGraphNode_LinkedAnimLayer.h"
 #include "String/ParseTokens.h"
-#include "Algo/Transform.h"
-#include "Algo/Accumulate.h"
+#include "Algo/Sort.h"
 #include "IClassVariableCreator.h"
 #include "AnimBlueprintGeneratedClassCompiledData.h"
 #include "AnimBlueprintCompilationContext.h"
+#include "AnimBlueprintCompilerCreationContext.h"
 #include "AnimBlueprintVariableCreationContext.h"
+#include "Animation/AnimSubsystemInstance.h"
+#include "AnimBlueprintExtension.h"
+#include "UObject/UE5MainStreamObjectVersion.h"
 
 #define LOCTEXT_NAMESPACE "AnimBlueprintCompiler"
 
+DECLARE_CYCLE_STAT(TEXT("Merge Ubergraph Pages In"), EAnimBlueprintCompilerStats_MergeUbergraphPagesIn, STATGROUP_KismetCompiler);
+DECLARE_CYCLE_STAT(TEXT("Process All Animation Nodes"), EAnimBlueprintCompilerStats_ProcessAllAnimationNodes, STATGROUP_KismetCompiler);
+DECLARE_CYCLE_STAT(TEXT("Process Animation Nodes"), EAnimBlueprintCompilerStats_ProcessAnimationNodes, STATGROUP_KismetCompiler);
+DECLARE_CYCLE_STAT(TEXT("Expand Split Pins"), EAnimBlueprintCompilerStats_ExpandSplitPins, STATGROUP_KismetCompiler);
+DECLARE_CYCLE_STAT(TEXT("Move Graphs"), EAnimBlueprintCompilerStats_MoveGraphs, STATGROUP_KismetCompiler);
+DECLARE_CYCLE_STAT(TEXT("Clone Graph"), EAnimBlueprintCompilerStats_CloneGraph, STATGROUP_KismetCompiler);
+DECLARE_CYCLE_STAT(TEXT("Process Animation Node"), EAnimBlueprintCompilerStats_ProcessAnimationNode, STATGROUP_KismetCompiler);
+DECLARE_CYCLE_STAT(TEXT("Gather Fold Records"), EAnimBlueprintCompilerStats_GatherFoldRecords, STATGROUP_KismetCompiler);
+
 //////////////////////////////////////////////////////////////////////////
 // FAnimBlueprintCompiler
 
 FAnimBlueprintCompilerContext::FAnimBlueprintCompilerContext(UAnimBlueprint* SourceSketch, FCompilerResultsLog& InMessageLog, const FKismetCompilerOptions& InCompileOptions)
 	: FKismetCompilerContext(SourceSketch, InMessageLog, InCompileOptions)
+	, NewAnimBlueprintConstants(nullptr)
+	, NewAnimBlueprintMutables(nullptr)
+	, NewMutablesProperty(nullptr)
 	, AnimBlueprint(SourceSketch)
 	, bIsDerivedAnimBlueprint(false)
 {
-	AnimBlueprintCompilerHandlerCollection.Initialize(this);
-
 	// Add the animation graph schema to skip default function processing on them
 	KnownGraphSchemas.AddUnique(UAnimationGraphSchema::StaticClass());
 
@@ -65,6 +68,13 @@ FAnimBlueprintCompilerContext::FAnimBlueprintCompilerContext(UAnimBlueprint* Sou
 		}
 	}
 
+	// If we need to, refresh all extensions here
+	if(AnimBlueprint->bRefreshExtensions)
+	{
+		UAnimBlueprintExtension::RefreshExtensions(AnimBlueprint);
+		AnimBlueprint->bRefreshExtensions = false;
+	}
+	
 	if (AnimBlueprint->HasAnyFlags(RF_NeedPostLoad))
 	{
 		//Compilation during loading .. need to verify node guids as some anim blueprints have duplicated guids
@@ -133,6 +143,12 @@ FAnimBlueprintCompilerContext::FAnimBlueprintCompilerContext(UAnimBlueprint* Sou
 		}
 	}
 
+	FAnimBlueprintCompilerCreationContext CreationContext(this);
+	UAnimBlueprintExtension::ForEachExtension(AnimBlueprint, [&CreationContext](UAnimBlueprintExtension* InExtension)
+	{
+		InExtension->BeginCompilation(CreationContext);
+	});
+
 	// Determine if there is an anim blueprint in the ancestry of this class
 	bIsDerivedAnimBlueprint = UAnimBlueprint::FindRootAnimBlueprint(AnimBlueprint) != NULL;
 
@@ -144,6 +160,11 @@ FAnimBlueprintCompilerContext::FAnimBlueprintCompilerContext(UAnimBlueprint* Sou
 FAnimBlueprintCompilerContext::~FAnimBlueprintCompilerContext()
 {
 	DestroyAnimGraphStubFunctions();
+
+	UAnimBlueprintExtension::ForEachExtension(AnimBlueprint, [](UAnimBlueprintExtension* InExtension)
+	{
+	    InExtension->EndCompilation();
+	});
 }
 
 void FAnimBlueprintCompilerContext::ForAllSubGraphs(UEdGraph* InGraph, TFunctionRef<void(UEdGraph*)> InPerGraphFunction)
@@ -199,6 +220,8 @@ UEdGraphSchema_K2* FAnimBlueprintCompilerContext::CreateSchema()
 
 void FAnimBlueprintCompilerContext::ProcessAnimationNode(UAnimGraphNode_Base* VisualAnimNode)
 {
+	BP_SCOPED_COMPILER_EVENT_STAT(EAnimBlueprintCompilerStats_ProcessAnimationNode);
+	
 	// Early out if this node has already been processed
 	if (AllocatedAnimNodes.Contains(VisualAnimNode))
 	{
@@ -207,7 +230,7 @@ void FAnimBlueprintCompilerContext::ProcessAnimationNode(UAnimGraphNode_Base* Vi
 
 	// Make sure the visual node has a runtime node template
 	const UScriptStruct* NodeType = VisualAnimNode->GetFNodeType();
-	if (NodeType == NULL)
+	if (NodeType == nullptr)
 	{
 		MessageLog.Error(TEXT("@@ has no animation node member"), VisualAnimNode);
 		return;
@@ -235,11 +258,13 @@ void FAnimBlueprintCompilerContext::ProcessAnimationNode(UAnimGraphNode_Base* Vi
 
 	FStructProperty* NewProperty = CastField<FStructProperty>(CreateVariable(FName(*NodeVariableName), NodeVariableType));
 
-	if (NewProperty == NULL)
+	if (NewProperty == nullptr)
 	{
 		MessageLog.Error(TEXT("Failed to create node property for @@"), VisualAnimNode);
 	}
 
+	GatherFoldRecordsForAnimationNode(NodeType, NewProperty, VisualAnimNode);
+
 	// Register this node with the compile-time data structures
 	const int32 AllocatedIndex = AllocateNodeIndexCounter++;
 	AllocatedAnimNodes.Add(VisualAnimNode, NewProperty);
@@ -251,6 +276,7 @@ void FAnimBlueprintCompilerContext::ProcessAnimationNode(UAnimGraphNode_Base* Vi
 	SourceNodeToProcessedNodeMap.Add(TrueSourceObject, VisualAnimNode);
 
 	// Register the slightly more permanent debug information
+	UAnimBlueprintGeneratedClass* NewAnimBlueprintClass = GetNewAnimBlueprintClass();
 	FAnimBlueprintDebugData& NewAnimBlueprintDebugData = NewAnimBlueprintClass->GetAnimBlueprintDebugData();
 	NewAnimBlueprintDebugData.NodePropertyToIndexMap.Add(TrueSourceObject, AllocatedIndex);
 	NewAnimBlueprintDebugData.NodeGuidToIndexMap.Add(TrueSourceObject->NodeGuid, AllocatedIndex);
@@ -262,6 +288,178 @@ void FAnimBlueprintCompilerContext::ProcessAnimationNode(UAnimGraphNode_Base* Vi
 	VisualAnimNode->ProcessDuringCompilation(CompilerContext, CompiledData);
 }
 
+void FAnimBlueprintCompilerContext::ProcessExtensions()
+{
+	TArray<UAnimBlueprintExtension*> Extensions = UAnimBlueprintExtension::GetExtensions(AnimBlueprint);
+	
+	// Sort extensions by class name (for determinism)
+	Algo::Sort(Extensions, [](UAnimBlueprintExtension* InExtensionA, UAnimBlueprintExtension* InExtensionB)
+	{
+		return InExtensionA->GetClass()->GetName() < InExtensionB->GetClass()->GetName();
+	});
+
+	// Process all gathered class extensions
+	for(UAnimBlueprintExtension* Extension : Extensions)
+	{
+		const FString ExtensionVariableName = ClassScopeNetNameMap.MakeValidName(Extension);
+
+		const UScriptStruct* InstanceDataType = Extension->GetInstanceDataType();
+		check(InstanceDataType->IsChildOf(FAnimSubsystemInstance::StaticStruct()));
+
+		const UScriptStruct* ClassDataType = Extension->GetClassDataType();
+		check(ClassDataType->IsChildOf(FAnimSubsystem::StaticStruct()));
+		
+		// Skip creating any properties if both are the default
+		if(InstanceDataType != FAnimSubsystemInstance::StaticStruct() || ClassDataType != FAnimSubsystem::StaticStruct())
+		{
+			// Process instance data (mutable)
+			{
+				FEdGraphPinType ExtensionVariableType;
+				ExtensionVariableType.PinCategory = UAnimationGraphSchema::PC_Struct;
+				ExtensionVariableType.PinSubCategoryObject = MakeWeakObjectPtr(const_cast<UScriptStruct*>(InstanceDataType));
+
+				FStructProperty* NewProperty = CastField<FStructProperty>(CreateVariable(FName(*ExtensionVariableName), ExtensionVariableType));
+				if (NewProperty == nullptr)
+				{
+					MessageLog.Error(*FText::Format(LOCTEXT("ExtensionPropertyCreationFailed", "Failed to create extension property for '{0}'"), FText::FromString(Extension->GetName())).ToString());
+				}
+				else
+				{
+					ExtensionToInstancePropertyMap.Add(Extension, NewProperty);
+					InstancePropertyToExtensionMap.Add(NewProperty, Extension);
+				}
+			}
+
+			// Process class data (constants)
+			{
+				FEdGraphPinType ExtensionVariableType;
+				ExtensionVariableType.PinCategory = UAnimationGraphSchema::PC_Struct;
+				ExtensionVariableType.PinSubCategoryObject = MakeWeakObjectPtr(const_cast<UScriptStruct*>(ClassDataType));
+
+				FStructProperty* NewProperty = CastField<FStructProperty>(CreateStructVariable(NewAnimBlueprintConstants, FName(*ExtensionVariableName), ExtensionVariableType));
+				if (NewProperty == nullptr)
+				{
+					MessageLog.Error(*FText::Format(LOCTEXT("ExtensionPropertyCreationFailed", "Failed to create extension property for '{0}'"), FText::FromString(Extension->GetName())).ToString());
+				}
+				else
+				{
+					ExtensionToClassPropertyMap.Add(Extension, NewProperty);
+					ClassPropertyToExtensionMap.Add(NewProperty, Extension);
+				}
+			}
+		}
+	}
+}
+
+void FAnimBlueprintCompilerContext::GatherFoldRecordsForAnimationNode(const UScriptStruct* InNodeType, FStructProperty* InNodeProperty, UAnimGraphNode_Base* InVisualAnimNode)
+{
+	BP_SCOPED_COMPILER_EVENT_STAT(EAnimBlueprintCompilerStats_GatherFoldRecords);
+
+	static const FName NAME_FoldProperty("FoldProperty");
+
+	// Run through node properties to see if any are eligible for folding
+	for(TFieldIterator<FProperty> It(InNodeType); It; ++It)
+	{
+		FProperty* SubProperty = *It;
+
+		if(FArrayProperty* ArrayProperty = CastField<FArrayProperty>(SubProperty))
+		{
+			bool bAllPinsExposed = true;
+			bool bAllPinsDisconnected = true;
+
+			// If a value is exposed on a pin but disconnected, push the value to the (intermediate) node here to simplify later logic
+			if(SubProperty->HasAnyPropertyFlags(CPF_Edit|CPF_BlueprintVisible))
+			{
+				FStructProperty* AnimGraphNodeProperty = InVisualAnimNode->GetFNodeProperty();
+
+				// Check the anim node property is contained in the anim graph node
+				check(AnimGraphNodeProperty->GetOwner<UClass>() && InVisualAnimNode->GetClass()->IsChildOf(AnimGraphNodeProperty->GetOwner<UClass>()));
+				const void* Node = AnimGraphNodeProperty->ContainerPtrToValuePtr<void>(InVisualAnimNode);
+
+				// Check the anim node's property is contained in the anim node
+				check(SubProperty->GetOwner<UStruct>() && AnimGraphNodeProperty->Struct->IsChildOf(SubProperty->GetOwner<UStruct>()));
+				const void* TargetPtr = SubProperty->ContainerPtrToValuePtr<void>(Node);
+
+				// Run through each array element - we can only fold array properties if all values are constants
+				FScriptArrayHelper ArrayHelper(ArrayProperty, TargetPtr);
+
+				for(int32 ArrayIndex = 0; ArrayIndex < ArrayHelper.Num(); ++ArrayIndex)
+				{
+					const FString ArrayElementPinName = SubProperty->GetName() + FString::Printf(TEXT("_%d"), ArrayIndex);
+					UEdGraphPin* Pin = InVisualAnimNode->FindPin(ArrayElementPinName);
+					const bool bExposedOnPin = Pin != nullptr;
+					const bool bPinConnected = InVisualAnimNode->IsPinExposedAndLinked(ArrayElementPinName, EGPD_Input);
+					const bool bPinBound = InVisualAnimNode->IsPinExposedAndBound(SubProperty->GetName(), EGPD_Input);
+					
+					if(bExposedOnPin)
+					{
+						if(!(bPinConnected || bPinBound))
+						{
+							check(Pin);
+
+							if(!FBlueprintEditorUtils::PropertyValueFromString_Direct(ArrayProperty->Inner, Pin->GetDefaultAsString(), ArrayHelper.GetRawPtr(ArrayIndex)))
+							{
+								MessageLog.Warning(TEXT("Unable to push default value for array pin @@ on @@"), Pin, InVisualAnimNode);
+							}
+						}
+						else
+						{
+							bAllPinsDisconnected = false;
+						}
+					}
+					else
+					{
+						bAllPinsExposed = false;
+					}
+				}
+			}
+
+			if(SubProperty->HasMetaData(NAME_FoldProperty))
+			{
+				// Add folding candidate
+				AddFoldedPropertyRecord(InVisualAnimNode, InNodeProperty, SubProperty, bAllPinsExposed, !bAllPinsDisconnected);
+			}
+		}
+		else
+		{
+			UEdGraphPin* Pin = InVisualAnimNode->FindPin(SubProperty->GetName());
+			const bool bExposedOnPin = Pin != nullptr;
+			const bool bPinConnected = InVisualAnimNode->IsPinExposedAndLinked(SubProperty->GetName(), EGPD_Input);
+			const bool bPinBound = InVisualAnimNode->IsPinExposedAndBound(SubProperty->GetName(), EGPD_Input);
+
+			// If a value is exposed on a pin but disconnected, push the value to the (intermediate) node here to simplify later logic
+			if(SubProperty->HasAnyPropertyFlags(CPF_Edit|CPF_BlueprintVisible))
+			{
+				if(bExposedOnPin && !(bPinConnected || bPinBound))
+				{
+					check(Pin);
+
+					FStructProperty* AnimGraphNodeProperty = InVisualAnimNode->GetFNodeProperty();
+
+					// Check the anim node property is contained in the anim graph node
+					check(AnimGraphNodeProperty->GetOwner<UClass>() && InVisualAnimNode->GetClass()->IsChildOf(AnimGraphNodeProperty->GetOwner<UClass>()));
+					const void* Node = AnimGraphNodeProperty->ContainerPtrToValuePtr<void>(InVisualAnimNode);
+
+					// Check the anim node's property is contained in the anim node
+					check(SubProperty->GetOwner<UStruct>() && AnimGraphNodeProperty->Struct->IsChildOf(SubProperty->GetOwner<UStruct>()));
+					const void* TargetPtr = SubProperty->ContainerPtrToValuePtr<void>(Node);
+
+					if(!FBlueprintEditorUtils::PropertyValueFromString_Direct(SubProperty, Pin->GetDefaultAsString(), (uint8*)TargetPtr))
+					{
+						MessageLog.Warning(TEXT("Unable to push default value for pin @@ on @@"), Pin, InVisualAnimNode);
+					}
+				}
+			}
+
+			if(SubProperty->HasMetaData(NAME_FoldProperty))
+			{
+				// Add folding candidate
+				AddFoldedPropertyRecord(InVisualAnimNode, InNodeProperty, SubProperty, bExposedOnPin, bPinConnected || bPinBound);
+			}
+		}
+	}
+}
+
 int32 FAnimBlueprintCompilerContext::GetAllocationIndexOfNode(UAnimGraphNode_Base* VisualAnimNode)
 {
 	ProcessAnimationNode(VisualAnimNode);
@@ -278,9 +476,12 @@ bool FAnimBlueprintCompilerContext::ShouldForceKeepNode(const UEdGraphNode* Node
 
 void FAnimBlueprintCompilerContext::PostExpansionStep(const UEdGraph* Graph)
 {
-	FAnimBlueprintGeneratedClassCompiledData CompiledData(NewAnimBlueprintClass);
+	FAnimBlueprintGeneratedClassCompiledData CompiledData(GetNewAnimBlueprintClass());
 	FAnimBlueprintPostExpansionStepContext CompilerContext(this);
-	OnPostExpansionStepDelegate.Broadcast(Graph, CompilerContext, CompiledData);
+	UAnimBlueprintExtension::ForEachExtension(AnimBlueprint, [Graph, &CompiledData, &CompilerContext](UAnimBlueprintExtension* InExtension)
+	{
+		InExtension->PostExpansionStep(Graph, CompilerContext, CompiledData);
+	});
 }
 
 void FAnimBlueprintCompilerContext::PruneIsolatedAnimationNodes(const TArray<UAnimGraphNode_Base*>& RootSet, TArray<UAnimGraphNode_Base*>& GraphNodes)
@@ -345,6 +546,8 @@ void FAnimBlueprintCompilerContext::PruneIsolatedAnimationNodes(const TArray<UAn
 
 void FAnimBlueprintCompilerContext::ProcessAnimationNodes(TArray<UAnimGraphNode_Base*>& AnimNodeList)
 {
+	BP_SCOPED_COMPILER_EVENT_STAT(EAnimBlueprintCompilerStats_ProcessAnimationNodes);
+	
 	// Process the remaining nodes
 	for (UAnimGraphNode_Base* AnimNode : AnimNodeList)
 	{
@@ -421,6 +624,8 @@ void FAnimBlueprintCompilerContext::GetLinkedAnimNodes_ProcessAnimNode(UAnimGrap
 
 void FAnimBlueprintCompilerContext::ProcessAllAnimationNodes()
 {
+	BP_SCOPED_COMPILER_EVENT_STAT(EAnimBlueprintCompilerStats_ProcessAllAnimationNodes);
+	
 	// Validate that we have a skeleton
 	if ((AnimBlueprint->TargetSkeleton == nullptr) && !AnimBlueprint->bIsNewlyCreated)
 	{
@@ -464,14 +669,26 @@ void FAnimBlueprintCompilerContext::ProcessAllAnimationNodes()
 		// Validate the graph
 		ValidateGraphIsWellFormed(ConsolidatedEventGraph);
 
-		FAnimBlueprintGeneratedClassCompiledData CompiledData(NewAnimBlueprintClass);
+		FAnimBlueprintGeneratedClassCompiledData CompiledData(GetNewAnimBlueprintClass());
 		FAnimBlueprintCompilationContext CompilerContext(this);
-		OnPreProcessAnimationNodesDelegate.Broadcast(AllSubGraphsAnimNodeList, CompilerContext, CompiledData);
+		UAnimBlueprintExtension::ForEachExtension(AnimBlueprint, [&AllSubGraphsAnimNodeList, &CompiledData, &CompilerContext](UAnimBlueprintExtension* InExtension)
+		{
+			InExtension->PreProcessAnimationNodes(AllSubGraphsAnimNodeList, CompilerContext, CompiledData);
+		});
 
 		// Process the animation nodes
 		ProcessAnimationNodes(RootAnimNodeList);
 
-		OnPostProcessAnimationNodesDelegate.Broadcast(AllSubGraphsAnimNodeList, CompilerContext, CompiledData);
+		// Process any extensions
+		ProcessExtensions();
+		
+		// Fold any constants
+		ProcessFoldedPropertyRecords();
+
+		UAnimBlueprintExtension::ForEachExtension(AnimBlueprint, [&AllSubGraphsAnimNodeList, &CompiledData, &CompilerContext](UAnimBlueprintExtension* InExtension)
+		{
+			InExtension->PostProcessAnimationNodes(AllSubGraphsAnimNodeList, CompilerContext, CompiledData);
+		});	
 	}
 	else
 	{
@@ -484,9 +701,13 @@ void FAnimBlueprintCompilerContext::CopyTermDefaultsToDefaultObject(UObject* Def
 	Super::CopyTermDefaultsToDefaultObject(DefaultObject);
 
 	UAnimInstance* DefaultAnimInstance = Cast<UAnimInstance>(DefaultObject);
+	UAnimBlueprintGeneratedClass* NewAnimBlueprintClass = GetNewAnimBlueprintClass();
 
 	if (bIsDerivedAnimBlueprint && DefaultAnimInstance)
 	{
+		// Make sure our sparse class data is re-created
+		RecreateSparseClassData();
+		
 		// If we are a derived animation graph; apply any stored overrides.
 		// Restore values from the root class to catch values where the override has been removed.
 		UAnimBlueprintGeneratedClass* RootAnimClass = NewAnimBlueprintClass;
@@ -494,6 +715,15 @@ void FAnimBlueprintCompilerContext::CopyTermDefaultsToDefaultObject(UObject* Def
 		{
 			RootAnimClass = NextClass;
 		}
+
+		// Copy constant/folded data from root class, remapping the class
+		NewAnimBlueprintClass->AnimNodeData = RootAnimClass->AnimNodeData;
+		NewAnimBlueprintClass->NodeTypeMap = RootAnimClass->NodeTypeMap;
+		for(FAnimNodeData& NodeData : NewAnimBlueprintClass->AnimNodeData)
+		{
+			NodeData.AnimClassInterface = NewAnimBlueprintClass;
+		}
+		
 		UObject* RootDefaultObject = RootAnimClass->GetDefaultObject();
 
 		for (TFieldIterator<FProperty> It(RootAnimClass); It; ++It)
@@ -513,6 +743,9 @@ void FAnimBlueprintCompilerContext::CopyTermDefaultsToDefaultObject(UObject* Def
 				}
 			}
 		}
+
+		// Re-initialize node data tables (they would be overwritten in the loop above)
+		NewAnimBlueprintClass->InitializeAnimNodeData(DefaultObject);
 	}
 
 	// Give game-specific logic a chance to replace animations
@@ -530,10 +763,17 @@ void FAnimBlueprintCompilerContext::CopyTermDefaultsToDefaultObject(UObject* Def
 		{
 			if (Override->NewAsset)
 			{
-				FAnimNode_Base* BaseNode = NewAnimBlueprintClass->GetPropertyInstance<FAnimNode_Base>(DefaultAnimInstance, Override->ParentNodeGuid, EPropertySearchMode::Hierarchy);
-				if (BaseNode)
+				int32 NodeIndex = NewAnimBlueprintClass->GetNodeIndexFromGuid(Override->ParentNodeGuid, EPropertySearchMode::Hierarchy);
+				if(NodeIndex != INDEX_NONE)
 				{
-					BaseNode->OverrideAsset(Override->NewAsset);
+					const UAnimGraphNode_Base* GraphNode = Cast<UAnimGraphNode_Base>(NewAnimBlueprintClass->GetVisualNodeFromNodePropertyIndex(NodeIndex, EPropertySearchMode::Hierarchy));
+					FAnimNode_Base* BaseNode = NewAnimBlueprintClass->GetPropertyInstance<FAnimNode_Base>(DefaultAnimInstance, Override->ParentNodeGuid, EPropertySearchMode::Hierarchy);
+					if (GraphNode && BaseNode)
+					{
+						FAnimBlueprintNodeOverrideAssetsContext Context(BaseNode, GraphNode->GetFNodeType());
+						Context.AddAsset(Override->NewAsset);
+						GraphNode->OverrideAssets(Context);
+					}
 				}
 			}
 		}
@@ -550,6 +790,28 @@ void FAnimBlueprintCompilerContext::CopyTermDefaultsToDefaultObject(UObject* Def
 		FAnimBlueprintGeneratedClassCompiledData CompiledData(NewAnimBlueprintClass);
 		FAnimBlueprintCopyTermDefaultsContext CompilerContext(this);
 
+		// Initialize extensions from their templates
+		for (TFieldIterator<FStructProperty> It(DefaultAnimInstance->GetClass(), EFieldIteratorFlags::ExcludeSuper); It; ++It)
+		{
+			FStructProperty* TargetProperty = *It;
+
+			if (UAnimBlueprintExtension* Extension = InstancePropertyToExtensionMap.FindRef(TargetProperty))
+			{
+				uint8* DestinationPtr = TargetProperty->ContainerPtrToValuePtr<uint8>(DefaultAnimInstance);
+				const uint8* SourcePtr = nullptr;
+				const FStructProperty* SourceExtensionProperty = Extension->GetInstanceDataProperty();
+				
+				if(SourceExtensionProperty)
+				{
+					check(TargetProperty->Struct == SourceExtensionProperty->Struct);				
+					SourcePtr = SourceExtensionProperty->ContainerPtrToValuePtr<uint8>(Extension);
+				}
+				
+				FAnimBlueprintExtensionCopyTermDefaultsContext ExtensionContext(TargetProperty, DestinationPtr, SourcePtr, LinkIndexCount);
+				Extension->CopyTermDefaultsToDefaultObject(DefaultAnimInstance, CompilerContext, ExtensionContext);
+			}
+		}
+		
 		// Initialize animation nodes from their templates
 		for (TFieldIterator<FProperty> It(DefaultAnimInstance->GetClass(), EFieldIteratorFlags::ExcludeSuper); It; ++It)
 		{
@@ -573,6 +835,80 @@ void FAnimBlueprintCompilerContext::CopyTermDefaultsToDefaultObject(UObject* Def
 			}
 		}
 
+		// Applies a set of folded property records to a data area (i.e. the constant or mutable structs)
+		auto PatchDataArea = [](void* InData, UScriptStruct* InStruct, const TArray<TSharedRef<IAnimBlueprintCompilationContext::FFoldedPropertyRecord>>& InRecords)
+		{
+			for(const TSharedRef<IAnimBlueprintCompilationContext::FFoldedPropertyRecord>& Record : InRecords)
+			{
+				if(Record->GeneratedProperty != nullptr)
+				{
+					FStructProperty* AnimGraphNodeProperty = Record->AnimGraphNode->GetFNodeProperty();
+
+					// Check the anim node property is contained in the anim graph node
+					check(AnimGraphNodeProperty->GetOwner<UClass>() && Record->AnimGraphNode->GetClass()->IsChildOf(AnimGraphNodeProperty->GetOwner<UClass>()));
+					const void* Node = AnimGraphNodeProperty->ContainerPtrToValuePtr<void>(Record->AnimGraphNode);
+
+					// Check the anim node's property is contained in the anim node
+					check(Record->Property->GetOwner<UStruct>() && AnimGraphNodeProperty->Struct->IsChildOf(Record->Property->GetOwner<UStruct>()));
+					const void* SourcePtr = Record->Property->ContainerPtrToValuePtr<void>(Node);
+
+					// Check the generated property is a member of the constants struct
+					check(Record->GeneratedProperty->GetOwner<UStruct>() && InStruct->IsChildOf(Record->GeneratedProperty->GetOwner<UStruct>()));
+					void* TargetPtr = Record->GeneratedProperty->ContainerPtrToValuePtr<void>(InData);
+
+					// Extract underlying property for enums
+					FProperty* PropertyToCopyWith = Record->GeneratedProperty;
+					if(const FEnumProperty* EnumProperty = CastField<const FEnumProperty>(Record->GeneratedProperty))
+					{
+						PropertyToCopyWith = EnumProperty->GetUnderlyingProperty();
+					}
+
+					PropertyToCopyWith->CopyCompleteValue(TargetPtr, SourcePtr);
+				}
+			}
+		};
+
+		if(void* Constants = NewAnimBlueprintClass->GetOrCreateSparseClassData())
+		{
+ 			UScriptStruct* ConstantsStruct = NewAnimBlueprintClass->GetSparseClassDataStruct();
+ 			check(ConstantsStruct == NewAnimBlueprintConstants);
+
+			// Initialize extensions from their templates
+			for (TFieldIterator<FStructProperty> It(ConstantsStruct, EFieldIteratorFlags::ExcludeSuper); It; ++It)
+			{
+				FStructProperty* TargetProperty = *It;
+
+				if (UAnimBlueprintExtension* Extension = ClassPropertyToExtensionMap.FindRef(TargetProperty))
+				{
+					if(const FStructProperty* SourceExtensionProperty = Extension->GetClassDataProperty())
+					{
+						check(SourceExtensionProperty != NULL);
+						check(TargetProperty->Struct == SourceExtensionProperty->Struct);
+
+						uint8* DestinationPtr = TargetProperty->ContainerPtrToValuePtr<uint8>(Constants);
+						const uint8* SourcePtr = SourceExtensionProperty->ContainerPtrToValuePtr<uint8>(Extension);
+					
+						FAnimBlueprintExtensionCopyTermDefaultsContext NodeContext(TargetProperty, DestinationPtr, SourcePtr, LinkIndexCount);
+						Extension->CopyTermDefaultsToSparseClassData(CompilerContext, NodeContext);
+					}
+				}
+			}
+
+			// Patch constants
+			PatchDataArea(Constants, ConstantsStruct, ConstantPropertyRecords);
+		}
+
+		if(NewMutablesProperty)
+		{
+			check(NewMutablesProperty->GetOwner<UClass>() && DefaultObject->GetClass()->IsChildOf(NewMutablesProperty->GetOwner<UClass>()));
+			void* Mutables = NewMutablesProperty->ContainerPtrToValuePtr<void>(DefaultObject);
+			UScriptStruct* MutablesStruct = NewMutablesProperty->Struct;
+			check(MutablesStruct == NewAnimBlueprintMutables);
+
+			// Patch mutables
+			PatchDataArea(Mutables, MutablesStruct, MutablePropertyRecords);
+		}
+		
 		// And wire up node links
 		for (auto PoseLinkIt = ValidPoseLinkList.CreateIterator(); PoseLinkIt; ++PoseLinkIt)
 		{
@@ -594,10 +930,8 @@ void FAnimBlueprintCompilerContext::CopyTermDefaultsToDefaultObject(UObject* Def
 			}
 		}   
 
-		OnCopyTermDefaultsToDefaultObjectDelegate.Broadcast(DefaultObject, CompilerContext, CompiledData);
-
 		UAnimBlueprintGeneratedClass* AnimBlueprintGeneratedClass = CastChecked<UAnimBlueprintGeneratedClass>(NewClass);
-
+		
 		// copy threaded update flag to CDO
 		DefaultAnimInstance->bUseMultiThreadedAnimationUpdate = AnimBlueprint->bUseMultiThreadedAnimationUpdate;
 
@@ -727,6 +1061,8 @@ void FAnimBlueprintCompilerContext::CopyTermDefaultsToDefaultObject(UObject* Def
 
 void FAnimBlueprintCompilerContext::ExpandSplitPins(UEdGraph* InGraph)
 {
+	BP_SCOPED_COMPILER_EVENT_STAT(EAnimBlueprintCompilerStats_ExpandSplitPins);
+
 	for (TArray<UEdGraphNode*>::TIterator NodeIt(InGraph->Nodes); NodeIt; ++NodeIt)
 	{
 		UK2Node* K2Node = Cast<UK2Node>(*NodeIt);
@@ -740,6 +1076,8 @@ void FAnimBlueprintCompilerContext::ExpandSplitPins(UEdGraph* InGraph)
 // Merges in any all ubergraph pages into the gathering ubergraph
 void FAnimBlueprintCompilerContext::MergeUbergraphPagesIn(UEdGraph* Ubergraph)
 {
+	BP_SCOPED_COMPILER_EVENT_STAT(EAnimBlueprintCompilerStats_MergeUbergraphPagesIn);
+
 	Super::MergeUbergraphPagesIn(Ubergraph);
 
 	if (bIsDerivedAnimBlueprint)
@@ -748,38 +1086,50 @@ void FAnimBlueprintCompilerContext::MergeUbergraphPagesIn(UEdGraph* Ubergraph)
 	}
 	else
 	{
-		// Move all animation graph nodes and associated pure logic chains into the consolidated event graph
-		auto MoveGraph = [this](UEdGraph* InGraph)
+		RecreateSparseClassData();
+		RecreateMutables();
+	
 		{
-			if (InGraph->Schema->IsChildOf(UAnimationGraphSchema::StaticClass()))
+			BP_SCOPED_COMPILER_EVENT_STAT(EAnimBlueprintCompilerStats_MoveGraphs);
+			
+			// Move all animation graph nodes and associated pure logic chains into the consolidated event graph
+			auto MoveGraph = [this](UEdGraph* InGraph)
 			{
-				// Merge all the animation nodes, contents, etc... into the ubergraph
-				UEdGraph* ClonedGraph = FEdGraphUtilities::CloneGraph(InGraph, NULL, &MessageLog, true);
-
-				// Prune the graph up-front
-				const bool bIncludePotentialRootNodes = false;
-				PruneIsolatedNodes(ClonedGraph, bIncludePotentialRootNodes);
+				if (InGraph->Schema->IsChildOf(UAnimationGraphSchema::StaticClass()))
+				{
+					UEdGraph* ClonedGraph;
+					{
+						BP_SCOPED_COMPILER_EVENT_STAT(EAnimBlueprintCompilerStats_CloneGraph);
+					
+						// Merge all the animation nodes, contents, etc... into the ubergraph
+						ClonedGraph = FEdGraphUtilities::CloneGraph(InGraph, NULL, &MessageLog, true);
+					}
 
-				const bool bIsLoading = Blueprint->bIsRegeneratingOnLoad || IsAsyncLoading();
-				const bool bIsCompiling = Blueprint->bBeingCompiled;
-				ClonedGraph->MoveNodesToAnotherGraph(ConsolidatedEventGraph, bIsLoading, bIsCompiling);
+					// Prune the graph up-front
+					const bool bIncludePotentialRootNodes = false;
+					PruneIsolatedNodes(ClonedGraph, bIncludePotentialRootNodes);
 
-				// Move subgraphs too
-				ConsolidatedEventGraph->SubGraphs.Append(ClonedGraph->SubGraphs);
-			}
-		};
+					const bool bIsLoading = Blueprint->bIsRegeneratingOnLoad || IsAsyncLoading();
+					const bool bIsCompiling = Blueprint->bBeingCompiled;
+					ClonedGraph->MoveNodesToAnotherGraph(ConsolidatedEventGraph, bIsLoading, bIsCompiling);
 
-		for (UEdGraph* Graph : Blueprint->FunctionGraphs)
-		{
-			MoveGraph(Graph);
-		}
+					// Move subgraphs too
+					ConsolidatedEventGraph->SubGraphs.Append(ClonedGraph->SubGraphs);
+				}
+			};
 
-		for(FBPInterfaceDescription& InterfaceDesc : Blueprint->ImplementedInterfaces)
-		{
-			for(UEdGraph* Graph : InterfaceDesc.Graphs)
+			for (UEdGraph* Graph : Blueprint->FunctionGraphs)
 			{
 				MoveGraph(Graph);
 			}
+
+			for(FBPInterfaceDescription& InterfaceDesc : Blueprint->ImplementedInterfaces)
+			{
+				for(UEdGraph* Graph : InterfaceDesc.Graphs)
+				{
+					MoveGraph(Graph);
+				}
+			}
 		}
 
 		// Make sure we expand any split pins here before we process animation nodes.
@@ -814,76 +1164,144 @@ void FAnimBlueprintCompilerContext::EnsureProperGeneratedClass(UClass*& TargetUC
 	}
 }
 
+void FAnimBlueprintCompilerContext::RecreateSparseClassData()
+{
+	UAnimBlueprintGeneratedClass* NewAnimBlueprintClass = GetNewAnimBlueprintClass();
+	
+	// Set up our sparse class data struct
+	const FString ConstantsStructTypeName = NewAnimBlueprintClass->GetName() + TEXT("_ConstantData");
+
+	// Rename old struct out of the way
+	if(UScriptStruct* OldAnimBlueprintConstants = FindObject<UScriptStruct>(NewAnimBlueprintClass->GetOuter(), *ConstantsStructTypeName))
+	{
+		OldAnimBlueprintConstants->Rename(nullptr, GetTransientPackage(), REN_DontCreateRedirectors | REN_NonTransactional | REN_ForceNoResetLoaders);
+	}
+
+	if (bIsDerivedAnimBlueprint)
+	{
+		// Get parent class
+		UAnimBlueprint* ParentAnimBP = UAnimBlueprint::GetParentAnimBlueprint(AnimBlueprint);
+		UAnimBlueprintGeneratedClass* ParentAnimClass = Cast<UAnimBlueprintGeneratedClass>(ParentAnimBP->GeneratedClass);
+		
+		check(ParentAnimClass);
+		check(ParentAnimClass->GetSparseClassDataStruct());
+		
+		// Derive sparse class data from parent class
+		NewAnimBlueprintConstants = NewObject<UScriptStruct>(NewAnimBlueprintClass->GetOuter(), *ConstantsStructTypeName, RF_Public);
+		NewAnimBlueprintConstants->SetSuperStruct(ParentAnimClass->GetSparseClassDataStruct());
+
+		// Just link & assign sparse class data struct here, no additional members are added
+		NewAnimBlueprintConstants->StaticLink(true);
+		NewAnimBlueprintClass->SetSparseClassDataStruct(NewAnimBlueprintConstants);
+		NewAnimBlueprintClass->GetOrCreateSparseClassData();
+		NewAnimBlueprintClass->BuildConstantProperties();
+	}
+	else
+	{
+		UClass* ParentClass = NewAnimBlueprintClass->GetSuperClass();
+		check(ParentClass);
+		
+	    // Create new sparse class data struct
+		NewAnimBlueprintConstants = NewObject<UScriptStruct>(NewAnimBlueprintClass->GetOuter(), *ConstantsStructTypeName, RF_Public);
+
+		// Inherit from archetype struct if there is any 
+		UScriptStruct* ArchetypeStruct = ParentClass->GetSparseClassDataArchetypeStruct();
+		UScriptStruct* SuperStruct = ArchetypeStruct ? ArchetypeStruct : FAnimBlueprintConstantData::StaticStruct();
+		NewAnimBlueprintConstants->SetSuperStruct(SuperStruct);
+	}
+}
+
+void FAnimBlueprintCompilerContext::RecreateMutables()
+{
+	UAnimBlueprintGeneratedClass* NewAnimBlueprintClass = GetNewAnimBlueprintClass();
+
+	// Set up our mutables struct
+	const FString MutablesStructTypeName = NewAnimBlueprintClass->GetName() + TEXT("_MutableData");
+
+	// Rename old struct out of the way
+	if(UScriptStruct* OldAnimBlueprintMutables = FindObject<UScriptStruct>(NewAnimBlueprintClass->GetOuter(), *MutablesStructTypeName))
+	{
+		OldAnimBlueprintMutables->Rename(nullptr, GetTransientPackage(), REN_DontCreateRedirectors | REN_NonTransactional | REN_ForceNoResetLoaders);
+	}
+
+	NewAnimBlueprintMutables = NewObject<UScriptStruct>(NewAnimBlueprintClass->GetOuter(), *MutablesStructTypeName);
+	NewAnimBlueprintMutables->SetSuperStruct(FAnimBlueprintMutableData::StaticStruct());
+}
+
 void FAnimBlueprintCompilerContext::SpawnNewClass(const FString& NewClassName)
 {
-	NewAnimBlueprintClass = FindObject<UAnimBlueprintGeneratedClass>(Blueprint->GetOutermost(), *NewClassName);
+	NewClass = FindObject<UAnimBlueprintGeneratedClass>(Blueprint->GetOutermost(), *NewClassName);
 
-	if (NewAnimBlueprintClass == NULL)
+	if (NewClass == NULL)
 	{
-		NewAnimBlueprintClass = NewObject<UAnimBlueprintGeneratedClass>(Blueprint->GetOutermost(), FName(*NewClassName), RF_Public | RF_Transactional);
+		NewClass = NewObject<UAnimBlueprintGeneratedClass>(Blueprint->GetOutermost(), *NewClassName, RF_Public | RF_Transactional);
 	}
 	else
 	{
 		// Already existed, but wasn't linked in the Blueprint yet due to load ordering issues
-		FBlueprintCompileReinstancer::Create(NewAnimBlueprintClass);
+		FBlueprintCompileReinstancer::Create(NewClass);
 	}
-	NewClass = NewAnimBlueprintClass;
-
-	FAnimBlueprintGeneratedClassCompiledData CompiledData(NewAnimBlueprintClass);
+	
+	FAnimBlueprintGeneratedClassCompiledData CompiledData(GetNewAnimBlueprintClass());
 	FAnimBlueprintCompilationBracketContext CompilerContext(this);
-	OnStartCompilingClassDelegate.Broadcast(NewAnimBlueprintClass, CompilerContext, CompiledData);
+	UAnimBlueprintExtension::ForEachExtension(AnimBlueprint, [this, &CompiledData, &CompilerContext](UAnimBlueprintExtension* InExtension)
+    {
+		InExtension->StartCompilingClass(GetNewAnimBlueprintClass(), CompilerContext, CompiledData);
+	});
 }
 
 void FAnimBlueprintCompilerContext::OnPostCDOCompiled()
 {
+	UAnimBlueprintGeneratedClass* NewAnimBlueprintClass = GetNewAnimBlueprintClass();
 	for (UAnimBlueprintGeneratedClass* ClassWithInputHandlers = NewAnimBlueprintClass; ClassWithInputHandlers != nullptr; ClassWithInputHandlers = Cast<UAnimBlueprintGeneratedClass>(ClassWithInputHandlers->GetSuperClass()))
 	{
-		FExposedValueHandler::ClassInitialization(ClassWithInputHandlers->EvaluateGraphExposedInputs, NewAnimBlueprintClass->ClassDefaultObject);
-
-		ClassWithInputHandlers->LinkFunctionsToDefaultObjectNodes(NewAnimBlueprintClass->ClassDefaultObject);
+		ClassWithInputHandlers->OnPostLoadDefaults(NewAnimBlueprintClass->ClassDefaultObject);
 	}
 }
 
 void FAnimBlueprintCompilerContext::OnNewClassSet(UBlueprintGeneratedClass* ClassToUse)
 {
-	NewAnimBlueprintClass = CastChecked<UAnimBlueprintGeneratedClass>(ClassToUse);
+	NewClass = CastChecked<UAnimBlueprintGeneratedClass>(ClassToUse);
 }
 
 void FAnimBlueprintCompilerContext::CleanAndSanitizeClass(UBlueprintGeneratedClass* ClassToClean, UObject*& InOldCDO)
 {
 	Super::CleanAndSanitizeClass(ClassToClean, InOldCDO);
 
-	// Make sure our typed pointer is set
-	check(ClassToClean == NewClass && NewAnimBlueprintClass == NewClass);
-
-	NewAnimBlueprintClass->AnimBlueprintDebugData = FAnimBlueprintDebugData();
+	UAnimBlueprintGeneratedClass* AnimBlueprintClassToClean= CastChecked<UAnimBlueprintGeneratedClass>(ClassToClean);
+	
+	AnimBlueprintClassToClean->AnimBlueprintDebugData = FAnimBlueprintDebugData();
 
 	// Reset the baked data
 	//@TODO: Move this into PurgeClass
-	NewAnimBlueprintClass->BakedStateMachines.Empty();
-	NewAnimBlueprintClass->AnimNotifies.Empty();
-	NewAnimBlueprintClass->AnimBlueprintFunctions.Empty();
-	NewAnimBlueprintClass->OrderedSavedPoseIndicesMap.Empty();
-	NewAnimBlueprintClass->AnimNodeProperties.Empty();
-	NewAnimBlueprintClass->LinkedAnimGraphNodeProperties.Empty();
-	NewAnimBlueprintClass->LinkedAnimLayerNodeProperties.Empty();
-	NewAnimBlueprintClass->PreUpdateNodeProperties.Empty();
-	NewAnimBlueprintClass->DynamicResetNodeProperties.Empty();
-	NewAnimBlueprintClass->StateMachineNodeProperties.Empty();
-	NewAnimBlueprintClass->InitializationNodeProperties.Empty();
-	NewAnimBlueprintClass->EvaluateGraphExposedInputs.Empty();
-	NewAnimBlueprintClass->GraphAssetPlayerInformation.Empty();
-	NewAnimBlueprintClass->GraphBlendOptions.Empty();
+	AnimBlueprintClassToClean->BakedStateMachines.Empty();
+	AnimBlueprintClassToClean->AnimNotifies.Empty();
+	AnimBlueprintClassToClean->AnimBlueprintFunctions.Empty();
+	AnimBlueprintClassToClean->OrderedSavedPoseIndicesMap.Empty();
+	AnimBlueprintClassToClean->AnimNodeProperties.Empty();
+	AnimBlueprintClassToClean->LinkedAnimGraphNodeProperties.Empty();
+	AnimBlueprintClassToClean->LinkedAnimLayerNodeProperties.Empty();
+	AnimBlueprintClassToClean->PreUpdateNodeProperties.Empty();
+	AnimBlueprintClassToClean->DynamicResetNodeProperties.Empty();
+	AnimBlueprintClassToClean->StateMachineNodeProperties.Empty();
+	AnimBlueprintClassToClean->InitializationNodeProperties.Empty();
+	AnimBlueprintClassToClean->GraphAssetPlayerInformation.Empty();
+	AnimBlueprintClassToClean->GraphBlendOptions.Empty();
+	AnimBlueprintClassToClean->AnimNodeData.Empty();
+	AnimBlueprintClassToClean->NodeTypeMap.Empty();
 
 	// Copy over runtime data from the blueprint to the class
-	NewAnimBlueprintClass->TargetSkeleton = AnimBlueprint->TargetSkeleton;
+	AnimBlueprintClassToClean->TargetSkeleton = AnimBlueprint->TargetSkeleton;
 
 	UAnimBlueprint* RootAnimBP = UAnimBlueprint::FindRootAnimBlueprint(AnimBlueprint);
 	bIsDerivedAnimBlueprint = RootAnimBP != NULL;
 
-	FAnimBlueprintGeneratedClassCompiledData CompiledData(NewAnimBlueprintClass);
+	FAnimBlueprintGeneratedClassCompiledData CompiledData(AnimBlueprintClassToClean);
 	FAnimBlueprintCompilationBracketContext CompilerContext(this);
-	OnStartCompilingClassDelegate.Broadcast(NewAnimBlueprintClass, CompilerContext, CompiledData);
+	UAnimBlueprintExtension::ForEachExtension(AnimBlueprint, [this, &CompiledData, &CompilerContext, AnimBlueprintClassToClean](UAnimBlueprintExtension* InExtension)
+	{
+		InExtension->StartCompilingClass(AnimBlueprintClassToClean, CompilerContext, CompiledData);
+	});
 }
 
 void FAnimBlueprintCompilerContext::FinishCompilingClass(UClass* Class)
@@ -926,9 +1344,12 @@ void FAnimBlueprintCompilerContext::FinishCompilingClass(UClass* Class)
 		}
 	}
 
-	FAnimBlueprintGeneratedClassCompiledData CompiledData(NewAnimBlueprintClass);
+	FAnimBlueprintGeneratedClassCompiledData CompiledData(GetNewAnimBlueprintClass());
 	FAnimBlueprintCompilationBracketContext CompilerContext(this);
-	OnFinishCompilingClassDelegate.Broadcast(AnimBlueprintGeneratedClass, CompilerContext, CompiledData);
+	UAnimBlueprintExtension::ForEachExtension(AnimBlueprint, [this, &CompiledData, &CompilerContext](UAnimBlueprintExtension* InExtension)
+	{
+		InExtension->FinishCompilingClass(GetNewAnimBlueprintClass(), CompilerContext, CompiledData);
+	});
 
 	Super::FinishCompilingClass(Class);
 }
@@ -962,6 +1383,8 @@ void FAnimBlueprintCompilerContext::PostCompileDiagnostics()
 {
 	FKismetCompilerContext::PostCompileDiagnostics();
 
+	UAnimBlueprintGeneratedClass* NewAnimBlueprintClass = GetNewAnimBlueprintClass();
+	
 #if WITH_EDITORONLY_DATA // ANIMINST_PostCompileValidation
 	// See if AnimInstance implements a PostCompileValidation Class. 
 	// If so, instantiate it, and let it perform Validation of our newly compiled AnimBlueprint.
@@ -1088,7 +1511,7 @@ void FAnimBlueprintCompilerContext::CreateAnimGraphStubFunctions()
 				UK2Node_FunctionEntry* EntryNode = SpawnIntermediateNode<UK2Node_FunctionEntry>(RootNode, StubGraph);
 				EntryNode->NodePosX = -200;
 				EntryNode->CustomGeneratedFunctionName = GraphToUseforSignature->GetFName();	// Note that the function generated from this temporary graph is undecorated
-				EntryNode->MetaData.Category = RootNode->Node.Group == NAME_None ? FText::GetEmpty() : FText::FromName(RootNode->Node.Group);
+				EntryNode->MetaData.Category = RootNode->Node.GetGroup() == NAME_None ? FText::GetEmpty() : FText::FromName(RootNode->Node.GetGroup());
 
 				// Add linked input poses as parameters
 				for(UAnimGraphNode_LinkedInputPose* LinkedInputPoseNode : LinkedInputPoseNodes)
@@ -1231,7 +1654,7 @@ void FAnimBlueprintCompilerContext::AddAttributesToNode(UAnimGraphNode_Base* InN
 {
 	if(UAnimGraphNode_Base* OriginalNode = CastChecked<UAnimGraphNode_Base>(MessageLog.FindSourceObject(InNode)))
 	{
-		TArray<FName>& AttributeSet = NewAnimBlueprintClass->GetAnimBlueprintDebugData().NodeAttributes.FindOrAdd(OriginalNode);
+		TArray<FName>& AttributeSet = GetNewAnimBlueprintClass()->GetAnimBlueprintDebugData().NodeAttributes.FindOrAdd(OriginalNode);
 		AttributeSet.Reserve(AttributeSet.Num() + InAttributes.Num());
 
 		for(const FName& Attribute : InAttributes)
@@ -1245,7 +1668,7 @@ TArrayView<const FName> FAnimBlueprintCompilerContext::GetAttributesFromNode(UAn
 {
 	if(UAnimGraphNode_Base* OriginalNode = CastChecked<UAnimGraphNode_Base>(MessageLog.FindSourceObject(InNode)))
 	{
-		if(const TArray<FName>* AttributeSetPtr = NewAnimBlueprintClass->GetAnimBlueprintDebugData().NodeAttributes.Find(OriginalNode))
+		if(const TArray<FName>* AttributeSetPtr = GetNewAnimBlueprintClass()->GetAnimBlueprintDebugData().NodeAttributes.Find(OriginalNode))
 		{
 			return MakeArrayView(*AttributeSetPtr);
 		}
@@ -1260,6 +1683,263 @@ FProperty* FAnimBlueprintCompilerContext::CreateUniqueVariable(UObject* InForObj
 	return CreateVariable(*VariableName, Type);
 }
 
+FProperty* FAnimBlueprintCompilerContext::CreateStructVariable(UScriptStruct* InStruct, const FName VarName, const FEdGraphPinType& VarType)
+{
+	FProperty* NewProperty = FKismetCompilerUtilities::CreatePropertyOnScope(InStruct, VarName, VarType, nullptr, CPF_None, Schema, MessageLog);
+	if (NewProperty != nullptr)
+	{
+		// This fixes a rare bug involving asynchronous loading of BPs in editor builds. The pattern was established
+		// in FKismetCompilerContext::CompileFunctions where we do this for the uber graph function. By setting
+		// the RF_LoadCompleted we prevent the linker from overwriting our regenerated property, although the
+		// circumstances under which this occurs are murky. More testing of BPs loading asynchronously in the editor
+		// needs to be added:
+		NewProperty->SetFlags(RF_LoadCompleted);
+		FKismetCompilerUtilities::LinkAddedProperty(InStruct, NewProperty);
+	}
+	else
+	{
+		MessageLog.Error(
+			*FText::Format(
+				LOCTEXT("VariableInvalidType_ErrorFmt", "The variable {0} declared in @@ has an invalid type {1}"),
+				FText::FromName(VarName),
+				UEdGraphSchema_K2::TypeToText(VarType)
+			).ToString(),
+			Blueprint
+		);
+	}
+
+	return NewProperty;
+}
+
+void FAnimBlueprintCompilerContext::AddFoldedPropertyRecord(UAnimGraphNode_Base* InAnimGraphNode, FStructProperty* InAnimNodeProperty, FProperty* InProperty, bool bInExposedOnPin, bool bInPinConnected)
+{
+	const bool bConstant = !bInExposedOnPin || (bInExposedOnPin && !bInPinConnected);
+
+	if(!InProperty->HasAnyPropertyFlags(CPF_EditorOnly))
+	{
+		MessageLog.Warning(*FString::Printf(TEXT("Property %s on @@ is foldable, but not editor only"), *InProperty->GetName()), InAnimGraphNode);
+	}
+
+	// Create record and add it our lookup map
+	TSharedRef<IAnimBlueprintCompilationContext::FFoldedPropertyRecord> Record = MakeShared<IAnimBlueprintCompilationContext::FFoldedPropertyRecord>(InAnimGraphNode, InAnimNodeProperty, InProperty, bConstant);
+	TArray<TSharedRef<IAnimBlueprintCompilationContext::FFoldedPropertyRecord>>& Array = NodeToFoldedPropertyRecordMap.FindOrAdd(InAnimGraphNode);
+	Array.Add(Record);
+
+	// Record it in the appropriate data area
+	if(bConstant)
+	{
+		ConstantPropertyRecords.Add(Record);
+	}
+	else
+	{
+		MutablePropertyRecords.Add(Record);
+	}
+}
+
+void FAnimBlueprintCompilerContext::ProcessFoldedPropertyRecords()
+{
+	UAnimBlueprintGeneratedClass* NewAnimBlueprintClass = GetNewAnimBlueprintClass();
+	
+	if(ConstantPropertyRecords.Num() > 0)
+	{
+		// Set constants struct as sparse class data
+		check(NewAnimBlueprintConstants);
+
+		auto GetRecordValue = [](const TSharedRef<IAnimBlueprintCompilationContext::FFoldedPropertyRecord>& InRecord)
+		{
+			FStructProperty* AnimGraphNodeProperty = InRecord->AnimGraphNode->GetFNodeProperty();
+
+			check(AnimGraphNodeProperty->GetOwner<UClass>() && InRecord->AnimGraphNode->GetClass()->IsChildOf(AnimGraphNodeProperty->GetOwner<UClass>()));
+			const void* Node = AnimGraphNodeProperty->ContainerPtrToValuePtr<void>(InRecord->AnimGraphNode);
+			check(InRecord->Property->GetOwner<UStruct>() && AnimGraphNodeProperty->Struct->IsChildOf(InRecord->Property->GetOwner<UStruct>()));
+			return InRecord->Property->ContainerPtrToValuePtr<void>(Node);
+		};
+
+		// Reduce any constant properties before patching
+		for(int32 RecordIndex0 = 0; RecordIndex0 < ConstantPropertyRecords.Num(); ++RecordIndex0)
+		{
+			const TSharedRef<IAnimBlueprintCompilationContext::FFoldedPropertyRecord>& Record0 = ConstantPropertyRecords[RecordIndex0];
+			if(Record0->FoldIndex == INDEX_NONE)
+			{
+				const void* Value0 = GetRecordValue(Record0);
+
+				for(int32 RecordIndex1 = RecordIndex0 + 1; RecordIndex1 < ConstantPropertyRecords.Num(); ++RecordIndex1)
+				{
+					const TSharedRef<IAnimBlueprintCompilationContext::FFoldedPropertyRecord>& Record1 = ConstantPropertyRecords[RecordIndex1];
+					if(Record1->FoldIndex == INDEX_NONE)
+					{
+						if(Record1->Property->SameType(Record0->Property))
+						{
+							const void* Value1 = GetRecordValue(Record1);
+
+							// same type, now test for equality
+							if(Record0->Property->Identical(Value0, Value1))
+							{
+								// Values are the same - fold
+								Record1->FoldIndex = RecordIndex0;
+							}
+						}
+					}
+				}
+			}
+		}
+	}
+
+	// Builds a 'data area', returns the total number of properties that were inserted into that area's struct
+	auto BuildDataArea = [this](TArray<TSharedRef<IAnimBlueprintCompilationContext::FFoldedPropertyRecord>>& InRecords, UScriptStruct* InStruct)
+	{
+		int32 PropertyIndex = 0;
+
+		if(InStruct)
+		{
+			const UAnimationGraphSchema* AnimationGraphSchema = GetDefault<UAnimationGraphSchema>();
+
+			for(TSharedRef<IAnimBlueprintCompilationContext::FFoldedPropertyRecord>& Record : InRecords)
+			{
+				// Skip folded records
+				if(Record->FoldIndex == INDEX_NONE)
+				{
+					FEdGraphPinType VariableType;
+					if(AnimationGraphSchema->ConvertPropertyToPinType(Record->Property, VariableType))
+					{
+						// Patch into sparse class data
+						FName PropertyName = FName(Record->Property->GetClass()->GetFName(), InRecords.Num() - 1 - PropertyIndex);
+						Record->GeneratedProperty = CreateStructVariable(InStruct, PropertyName, VariableType);
+						if (Record->GeneratedProperty == nullptr)
+						{
+							MessageLog.Error(*FString::Printf(TEXT("Property %s on node @@ could not be patched into data area."), *Record->Property->GetName()), Record->AnimGraphNode);
+						}
+						else
+						{
+							Record->PropertyIndex = PropertyIndex++;
+						}
+					}
+					else
+					{
+						MessageLog.Error(*FString::Printf(TEXT("Property %s on node @@ could not be patched into data area."), *Record->Property->GetName()), Record->AnimGraphNode);
+					}
+				}
+			}
+
+			if(InRecords.Num() > 0)
+			{
+				InStruct->StaticLink(true);
+			}
+		}
+
+		return PropertyIndex;
+	};
+
+	// Next, patch into the relevant data areas
+	const int32 NumConstantProperties = BuildDataArea(ConstantPropertyRecords, NewAnimBlueprintConstants);
+
+	// Set constants as sparse class data
+	if(ConstantPropertyRecords.Num() > 0)
+	{
+		NewAnimBlueprintClass->SetSparseClassDataStruct(NewAnimBlueprintConstants);
+	}
+
+	const int32 NumMutableProperties = BuildDataArea(MutablePropertyRecords, NewAnimBlueprintMutables);
+
+	// Create the property for our mutables, if we have any
+	if(MutablePropertyRecords.Num() > 0)
+	{
+		const FName MutablesStructName("__AnimBlueprintMutables");
+		FEdGraphPinType MutablesPinType;
+		MutablesPinType.PinCategory = UAnimationGraphSchema::PC_Struct;
+		MutablesPinType.PinSubCategoryObject = MakeWeakObjectPtr(NewAnimBlueprintMutables);
+		NewMutablesProperty = CastFieldChecked<FStructProperty>(CreateVariable(MutablesStructName, MutablesPinType));
+		NewMutablesProperty->SetMetaData(TEXT("BlueprintCompilerGeneratedDefaults"), TEXT("true"));
+	}
+
+	// Set up per-node mappings
+	NewAnimBlueprintClass->AnimNodeData.Empty();
+	NewAnimBlueprintClass->AnimNodeData.SetNum(AllocatedAnimNodeIndices.Num());
+	NewAnimBlueprintClass->NodeTypeMap.Empty();
+	
+	// First index & setup the node data
+	int32 NodeIndex = 0;
+	for(FAnimNodeData& NodeData : NewAnimBlueprintClass->AnimNodeData)
+	{
+		const int32 NodePropertyIndex = AllocatedAnimNodeIndices.Num() - 1 - NodeIndex;
+		FStructProperty* NodeStructProperty = CastFieldChecked<FStructProperty>(AllocatedPropertiesByIndex.FindChecked(NodePropertyIndex));
+		const UScriptStruct* Struct = NodeStructProperty->Struct;
+		const FAnimNodeStructData& AnimNodeStructData = NewAnimBlueprintClass->NodeTypeMap.Add(Struct, FAnimNodeStructData(Struct));
+
+		// Add any super-structs as values can be accessed via base classes
+		const UScriptStruct* SuperStruct = Cast<UScriptStruct>(Struct->GetSuperStruct());
+		while(SuperStruct)
+		{
+			NewAnimBlueprintClass->NodeTypeMap.Add(SuperStruct, FAnimNodeStructData(SuperStruct));
+			SuperStruct = Cast<UScriptStruct>(SuperStruct->GetSuperStruct());
+		}
+		
+		NodeData.AnimClassInterface = NewAnimBlueprintClass;
+		NodeData.NodeIndex = NodeIndex++;
+		const int32 NumProperties = AnimNodeStructData.GetNumProperties();
+		check(NumProperties >= 0);
+		NodeData.Entries.SetNum(NumProperties);
+		for(uint32& Entry : NodeData.Entries)
+		{
+			Entry = ANIM_NODE_DATA_INVALID_ENTRY;
+		}
+	}
+
+	auto BuildAnimNodeData = [this, &NewAnimBlueprintClass](const TArray<TSharedRef<IAnimBlueprintCompilationContext::FFoldedPropertyRecord>>& InRecords, int32 InTotalPropertyCount)
+	{
+		const int32 NumNodes = AllocatedAnimNodeIndices.Num();
+
+		for(const TSharedRef<IAnimBlueprintCompilationContext::FFoldedPropertyRecord>& Record : InRecords)
+		{
+			int32 NodeIndex = NumNodes - 1 - AllocatedAnimNodeIndices.FindChecked(Record->AnimGraphNode);
+			
+			FAnimNodeData& NodeData = NewAnimBlueprintClass->AnimNodeData[NodeIndex];
+			
+			int32 PropertyIndex = Record->PropertyIndex;
+			if(Record->FoldIndex != INDEX_NONE)
+			{
+				PropertyIndex = InRecords[Record->FoldIndex]->PropertyIndex;
+			}
+			check(PropertyIndex >= 0 && PropertyIndex < InTotalPropertyCount);
+			PropertyIndex = InTotalPropertyCount - 1 - PropertyIndex;
+
+			uint32 PropertyEntry = PropertyIndex;
+			if(!Record->bIsOnClass)
+			{
+				PropertyEntry |= ANIM_NODE_DATA_INSTANCE_DATA_FLAG;
+			}
+
+			const FAnimNodeStructData& AnimNodeStructData = NewAnimBlueprintClass->NodeTypeMap.FindChecked(Record->AnimNodeProperty->Struct);
+			const int32 EntryIndex = AnimNodeStructData.GetPropertyIndex(Record->Property->GetFName());
+			NodeData.Entries[EntryIndex] = PropertyEntry;
+		}
+	};
+
+	BuildAnimNodeData(ConstantPropertyRecords, NumConstantProperties);
+	BuildAnimNodeData(MutablePropertyRecords, NumMutableProperties);
+}
+
+bool FAnimBlueprintCompilerContext::IsAnimGraphNodeFolded(UAnimGraphNode_Base* InNode) const
+{
+	return NodeToFoldedPropertyRecordMap.Find(InNode) != nullptr;
+}
+
+const IAnimBlueprintCompilationContext::FFoldedPropertyRecord* FAnimBlueprintCompilerContext::GetFoldedPropertyRecord(UAnimGraphNode_Base* InNode, FName InPropertyName) const
+{
+	if(const TArray<TSharedRef<IAnimBlueprintCompilationContext::FFoldedPropertyRecord>>* FoundRecordsPtr = NodeToFoldedPropertyRecordMap.Find(InNode))
+	{
+		for(const TSharedRef<IAnimBlueprintCompilationContext::FFoldedPropertyRecord>& Record : *FoundRecordsPtr)
+		{
+			if(Record->Property->GetFName() == InPropertyName)
+			{
+				return &Record.Get();
+			}
+		}
+	}
+
+	return nullptr;
+}
+
 //////////////////////////////////////////////////////////////////////////
 
 #undef LOCTEXT_NAMESPACE
\ No newline at end of file
diff --git a/Engine/Source/Editor/AnimGraph/Private/AnimBlueprintCompiler.h b/Engine/Source/Editor/AnimGraph/Private/AnimBlueprintCompiler.h
--- a/Engine/Source/Editor/AnimGraph/Private/AnimBlueprintCompiler.h
+++ b/Engine/Source/Editor/AnimGraph/Private/AnimBlueprintCompiler.h
@@ -7,10 +7,9 @@
 #include "Animation/AnimNodeBase.h"
 #include "AnimGraphNode_Base.h"
 #include "KismetCompilerModule.h"
-#include "AnimBlueprintCompilerHandlerCollection.h"
-#include "IAnimBlueprintCompilerHandlerCollection.h"
 #include "IAnimBlueprintCompilerCreationContext.h"
 #include "Containers/ArrayView.h"
+#include "IAnimBlueprintCompilationContext.h"
 
 class UAnimationGraphSchema;
 class UAnimGraphNode_SaveCachedPose;
@@ -80,7 +79,9 @@ protected:
 	typedef TArray<UEdGraphPin*> UEdGraphPinArray;
 
 protected:
-	UAnimBlueprintGeneratedClass* NewAnimBlueprintClass;
+	UScriptStruct* NewAnimBlueprintConstants;
+	UScriptStruct* NewAnimBlueprintMutables;
+	FStructProperty* NewMutablesProperty;
 	UAnimBlueprint* AnimBlueprint;
 
 	UAnimationGraphSchema* AnimSchema;
@@ -88,6 +89,7 @@ protected:
 	// Map of allocated v3 nodes that are members of the class
 	TMap<class UAnimGraphNode_Base*, FProperty*> AllocatedAnimNodes;
 	TMap<FProperty*, class UAnimGraphNode_Base*> AllocatedNodePropertiesToNodes;
+	TMap<FProperty*, class UAnimGraphNode_Base*> AllocatedNodeConstantPropertiesToNodes;
 	TMap<int32, FProperty*> AllocatedPropertiesByIndex;
 
 	// Map of true source objects (user edited ones) to the cloned ones that are actually compiled
@@ -107,34 +109,31 @@ protected:
 	// True if any parent class is also generated from an animation blueprint
 	bool bIsDerivedAnimBlueprint;
 
-	// Handlers that this context is hosting
-	FAnimBlueprintCompilerHandlerCollection AnimBlueprintCompilerHandlerCollection;
-
 	// Graph schema classes that this compiler is aware of - they will skip default function processing
 	TArray<TSubclassOf<UEdGraphSchema>> KnownGraphSchemas;
 
-	/** Delegate fired when the class starts compiling. The class may be new or recycled. */
-	FOnStartCompilingClass OnStartCompilingClassDelegate;
-
-	/** Delegate fired before all animation nodes are processed */
-	FOnPreProcessAnimationNodes OnPreProcessAnimationNodesDelegate;
-
-	/** Delegate fired after all animation nodes are processed */
-	FOnPostProcessAnimationNodes OnPostProcessAnimationNodesDelegate;
+	// Expose compile options to handlers
+	using FKismetCompilerContext::CompileOptions;
 
-	/** Delegate fired post- graph expansion */
-	FOnPostExpansionStep OnPostExpansionStepDelegate;
+	// Records of folded properties gleaned from nodes
+	TArray<TSharedRef<IAnimBlueprintCompilationContext::FFoldedPropertyRecord>> ConstantPropertyRecords;
+	TArray<TSharedRef<IAnimBlueprintCompilationContext::FFoldedPropertyRecord>> MutablePropertyRecords;
 
-	/** Delegate fired when the class has finished compiling */
-	FOnFinishCompilingClass OnFinishCompilingClassDelegate;
+	// Allows lookups to see if a node participates in constant folding
+	TMap<UAnimGraphNode_Base*, TArray<TSharedRef<IAnimBlueprintCompilationContext::FFoldedPropertyRecord>>> NodeToFoldedPropertyRecordMap;
 
-	/** Delegate fired when data is being copied to the CDO */
-	FOnCopyTermDefaultsToDefaultObject OnCopyTermDefaultsToDefaultObjectDelegate;
+	// Maps of extension <-> generated property on the instance
+	TMap<UAnimBlueprintExtension*, FStructProperty*> ExtensionToInstancePropertyMap;
+	TMap<FStructProperty*, UAnimBlueprintExtension*> InstancePropertyToExtensionMap;
 
-	// Expose compile options to handlers
-	using FKismetCompilerContext::CompileOptions;
+	// Maps of extension <-> generated property on the class
+	TMap<UAnimBlueprintExtension*, FStructProperty*> ExtensionToClassPropertyMap;
+	TMap<FStructProperty*, UAnimBlueprintExtension*> ClassPropertyToExtensionMap;
 
 private:
+	// Get the generated class as an anim blueprint generated class
+	UAnimBlueprintGeneratedClass* GetNewAnimBlueprintClass() const { return CastChecked<UAnimBlueprintGeneratedClass>(NewClass); };
+	
 	// Run a function on the passed-in graph and each subgraph of it
 	void ForAllSubGraphs(UEdGraph* InGraph, TFunctionRef<void(UEdGraph*)> InPerGraphFunction);
 
@@ -144,8 +143,8 @@ private:
 	// Compiles one animation node
 	void ProcessAnimationNode(UAnimGraphNode_Base* VisualAnimNode);
 
-	// Compiles one root node
-	void ProcessRoot(UAnimGraphNode_Root* Root);
+	// Called during ProcessAnimationNode - gather property folding records for the node
+	void GatherFoldRecordsForAnimationNode(const UScriptStruct* InNodeType, FStructProperty* InNodeProperty, UAnimGraphNode_Base* InVisualAnimNode);
 
 	// Compiles an entire animation graph
 	void ProcessAllAnimationNodes();
@@ -153,6 +152,9 @@ private:
 	// Processes all the supplied anim nodes
 	void ProcessAnimationNodes(TArray<UAnimGraphNode_Base*>& AnimNodeList);
 
+	// Process all the requested extensions
+	void ProcessExtensions();
+	
 	// Gets all anim graph nodes that are piped into the provided node (traverses input pins)
 	void GetLinkedAnimNodes(UAnimGraphNode_Base* InGraphNode, TArray<UAnimGraphNode_Base*>& LinkedAnimNodes) const;
 	void GetLinkedAnimNodes_TraversePin(UEdGraphPin* InPin, TArray<UAnimGraphNode_Base*>& LinkedAnimNodes) const;
@@ -175,8 +177,32 @@ private:
 
 	// Get the current compiled-in attributes uniquely assigned to the specified node
 	TArrayView<const FName> GetAttributesFromNode(UAnimGraphNode_Base* InNode) const;
+
+	// Called at the start of compilation to (re-) create the mutable struct
+	void RecreateMutables();
+	
+	// (Re-)creates sparse class data structure. Called at the start of compilation to (re-) create the internal sparse 
+	// class data. For derived anim BPs this is called just-in-time before CDO copy. This is to ensure that the sparse 
+	// class data is always updated, as in some cases as the layout does not change a full compilation of a (data-only) 
+	// derived anim BP may be skipped.  
+	void RecreateSparseClassData();
 	
 	// Create a uniquely named variable corresponding to an object in the current class
 	FProperty* CreateUniqueVariable(UObject* InForObject, const FEdGraphPinType& Type);
+
+	/** Creates a variable on the specified struct */
+	FProperty* CreateStructVariable(UScriptStruct* InStruct, const FName VarName, const FEdGraphPinType& VarType);
+
+	/** Adds a record for a potentially-folded anim node property */
+	void AddFoldedPropertyRecord(UAnimGraphNode_Base* InAnimGraphNode, FStructProperty* InAnimNodeProperty, FProperty* InProperty, bool bInExposedOnPin, bool bInPinConnected);
+
+	/** Process any anim node properties that are 'foldable' */
+	void ProcessFoldedPropertyRecords();
+
+	// Check whether an anim node participates in constant folding
+	bool IsAnimGraphNodeFolded(UAnimGraphNode_Base* InNode) const;
+
+	// Get the folded property record, if any, for the supplied node & named property
+	const IAnimBlueprintCompilationContext::FFoldedPropertyRecord* GetFoldedPropertyRecord(UAnimGraphNode_Base* InNode, FName InPropertyName) const;
 };
 
diff --git a/Engine/Source/Editor/AnimGraph/Private/AnimBlueprintCompilerCreationContext.cpp b/Engine/Source/Editor/AnimGraph/Private/AnimBlueprintCompilerCreationContext.cpp
--- a/Engine/Source/Editor/AnimGraph/Private/AnimBlueprintCompilerCreationContext.cpp
+++ b/Engine/Source/Editor/AnimGraph/Private/AnimBlueprintCompilerCreationContext.cpp
@@ -4,36 +4,6 @@
 #include "AnimBlueprintCompiler.h"
 #include "EdGraph/EdGraphSchema.h"
 
-FOnStartCompilingClass& FAnimBlueprintCompilerCreationContext::OnStartCompilingClass()
-{
-	return CompilerContext->OnStartCompilingClassDelegate;
-}
-
-FOnPreProcessAnimationNodes& FAnimBlueprintCompilerCreationContext::OnPreProcessAnimationNodes()
-{
-	return CompilerContext->OnPreProcessAnimationNodesDelegate;
-}
-
-FOnPostProcessAnimationNodes& FAnimBlueprintCompilerCreationContext::OnPostProcessAnimationNodes()
-{
-	return CompilerContext->OnPostProcessAnimationNodesDelegate;
-}
-
-FOnPostExpansionStep& FAnimBlueprintCompilerCreationContext::OnPostExpansionStep()
-{
-	return CompilerContext->OnPostExpansionStepDelegate;
-}
-
-FOnFinishCompilingClass& FAnimBlueprintCompilerCreationContext::OnFinishCompilingClass()
-{
-	return CompilerContext->OnFinishCompilingClassDelegate;
-}
-
-FOnCopyTermDefaultsToDefaultObject& FAnimBlueprintCompilerCreationContext::OnCopyTermDefaultsToDefaultObject()
-{
-	return CompilerContext->OnCopyTermDefaultsToDefaultObjectDelegate;
-}
-
 void FAnimBlueprintCompilerCreationContext::RegisterKnownGraphSchema(TSubclassOf<UEdGraphSchema> InGraphSchemaClass)
 {
 	CompilerContext->KnownGraphSchemas.AddUnique(InGraphSchemaClass);
diff --git a/Engine/Source/Editor/AnimGraph/Private/AnimBlueprintCompilerCreationContext.h b/Engine/Source/Editor/AnimGraph/Private/AnimBlueprintCompilerCreationContext.h
--- a/Engine/Source/Editor/AnimGraph/Private/AnimBlueprintCompilerCreationContext.h
+++ b/Engine/Source/Editor/AnimGraph/Private/AnimBlueprintCompilerCreationContext.h
@@ -11,19 +11,13 @@ class FAnimBlueprintCompilerContext;
 class FAnimBlueprintCompilerCreationContext : public IAnimBlueprintCompilerCreationContext
 {
 private:
-	friend class FAnimBlueprintCompilerHandlerCollection;
-
+	friend class FAnimBlueprintCompilerContext;
+	
 	FAnimBlueprintCompilerCreationContext(FAnimBlueprintCompilerContext* InCompilerContext)
 		: CompilerContext(InCompilerContext)
 	{}
 
 	// IAnimBlueprintCompilerCreationContext interface
-	virtual FOnStartCompilingClass& OnStartCompilingClass() override;
-	virtual FOnPreProcessAnimationNodes& OnPreProcessAnimationNodes() override;
-	virtual FOnPostProcessAnimationNodes& OnPostProcessAnimationNodes() override;
-	virtual FOnPostExpansionStep& OnPostExpansionStep() override;
-	virtual FOnFinishCompilingClass& OnFinishCompilingClass() override;
-	virtual FOnCopyTermDefaultsToDefaultObject& OnCopyTermDefaultsToDefaultObject() override;
 	virtual void RegisterKnownGraphSchema(TSubclassOf<UEdGraphSchema> InGraphSchemaClass) override;
 
 private:
diff --git a/Engine/Source/Editor/AnimGraph/Private/AnimBlueprintGeneratedClassCompiledData.cpp b/Engine/Source/Editor/AnimGraph/Private/AnimBlueprintGeneratedClassCompiledData.cpp
--- a/Engine/Source/Editor/AnimGraph/Private/AnimBlueprintGeneratedClassCompiledData.cpp
+++ b/Engine/Source/Editor/AnimGraph/Private/AnimBlueprintGeneratedClassCompiledData.cpp
@@ -51,16 +51,6 @@ TArray<FAnimNotifyEvent>& FAnimBlueprintGeneratedClassCompiledData::GetAnimNotif
 	return Class->AnimNotifies;
 }
 
-TArray<FExposedValueHandler>& FAnimBlueprintGeneratedClassCompiledData::GetExposedValueHandlers() const
-{
-	return Class->EvaluateGraphExposedInputs;
-}
-
-FPropertyAccessLibrary& FAnimBlueprintGeneratedClassCompiledData::GetPropertyAccessLibrary() const
-{
-	return Class->PropertyAccessLibrary;
-}
-
 FAnimBlueprintDebugData& FAnimBlueprintGeneratedClassCompiledData::GetAnimBlueprintDebugData() const
 {
 	return Class->AnimBlueprintDebugData;
@@ -73,7 +63,5 @@ TMap<FName, FGraphAssetPlayerInformation>& FAnimBlueprintGeneratedClassCompiledD
 
 UBlendSpace* FAnimBlueprintGeneratedClassCompiledData::AddBlendSpace(UBlendSpace* InSourceBlendSpace)
 {
-	UBlendSpace* CopiedBlendSpace = DuplicateObject(InSourceBlendSpace, Class);
-	Class->BlendSpaces.Add(CopiedBlendSpace);
-	return CopiedBlendSpace;
-}
\ No newline at end of file
+	return nullptr;
+}
diff --git a/Engine/Source/Editor/AnimGraph/Private/AnimBlueprintGeneratedClassCompiledData.h b/Engine/Source/Editor/AnimGraph/Private/AnimBlueprintGeneratedClassCompiledData.h
--- a/Engine/Source/Editor/AnimGraph/Private/AnimBlueprintGeneratedClassCompiledData.h
+++ b/Engine/Source/Editor/AnimGraph/Private/AnimBlueprintGeneratedClassCompiledData.h
@@ -20,8 +20,6 @@ private:
 	virtual FBlueprintDebugData& GetBlueprintDebugData() const override;
 	virtual TArray<FAnimNotifyEvent>& GetAnimNotifies() const override;
 	virtual int32 FindOrAddNotify(FAnimNotifyEvent& Notify) const override;
-	virtual TArray<FExposedValueHandler>& GetExposedValueHandlers() const override;
-	virtual FPropertyAccessLibrary& GetPropertyAccessLibrary() const override;
 	virtual FAnimBlueprintDebugData& GetAnimBlueprintDebugData() const override;
 	virtual TMap<FName, FGraphAssetPlayerInformation>& GetGraphAssetPlayerInformation() const override;
 	virtual UBlendSpace* AddBlendSpace(UBlendSpace* InSourceBlendSpace) override;
diff --git a/Engine/Source/Editor/AnimGraph/Private/AnimBlueprintPostCompileValidation.cpp b/Engine/Source/Editor/AnimGraph/Private/AnimBlueprintPostCompileValidation.cpp
--- a/Engine/Source/Editor/AnimGraph/Private/AnimBlueprintPostCompileValidation.cpp
+++ b/Engine/Source/Editor/AnimGraph/Private/AnimBlueprintPostCompileValidation.cpp
@@ -145,11 +145,11 @@ void UAnimBlueprintPostCompileValidation::PCV_GatherAnimSequencesFromGraph(TArra
 		{
 			if (const FAnimNode_BlendSpacePlayer* const BlendSpacePlayer = Property->ContainerPtrToValuePtr<FAnimNode_BlendSpacePlayer>(PCV_Params.DefaultAnimInstance))
 			{
-				const bool bPassSyncGroupFilter = !GatherParams.bFilterBySyncGroup || (BlendSpacePlayer->GroupName == GatherParams.SyncGroupName);
-				const bool bPassLoopingFilter = !GatherParams.bFilterByLoopingCondition || (BlendSpacePlayer->bLoop == GatherParams.bLoopingCondition);
+				const bool bPassSyncGroupFilter = !GatherParams.bFilterBySyncGroup || (BlendSpacePlayer->GetGroupName() == GatherParams.SyncGroupName);
+				const bool bPassLoopingFilter = !GatherParams.bFilterByLoopingCondition || (BlendSpacePlayer->GetLoop() == GatherParams.bLoopingCondition);
 				if (bPassSyncGroupFilter && bPassLoopingFilter)
 				{
-					PCV_GatherAnimSequences(OutAnimSequences, BlendSpacePlayer->BlendSpace);
+					PCV_GatherAnimSequences(OutAnimSequences, BlendSpacePlayer->GetBlendSpace());
 				}
 			}
 		}
@@ -157,11 +157,11 @@ void UAnimBlueprintPostCompileValidation::PCV_GatherAnimSequencesFromGraph(TArra
 		{
 			if (const FAnimNode_SequencePlayer* const SequencePlayer = Property->ContainerPtrToValuePtr<FAnimNode_SequencePlayer>(PCV_Params.DefaultAnimInstance))
 			{
-				const bool bPassSyncGroupFilter = !GatherParams.bFilterBySyncGroup || (SequencePlayer->GroupName == GatherParams.SyncGroupName);
-				const bool bPassLoopingFilter = !GatherParams.bFilterByLoopingCondition || (SequencePlayer->bLoopAnimation == GatherParams.bLoopingCondition);
+				const bool bPassSyncGroupFilter = !GatherParams.bFilterBySyncGroup || (SequencePlayer->GetGroupName() == GatherParams.SyncGroupName);
+				const bool bPassLoopingFilter = !GatherParams.bFilterByLoopingCondition || (SequencePlayer->GetLoopAnimation() == GatherParams.bLoopingCondition);
 				if (bPassSyncGroupFilter && bPassLoopingFilter)
 				{
-					PCV_GatherAnimSequences(OutAnimSequences, SequencePlayer->Sequence);
+					PCV_GatherAnimSequences(OutAnimSequences, SequencePlayer->GetSequence());
 				}
 			}
 		}
@@ -176,7 +176,7 @@ void UAnimBlueprintPostCompileValidation::PCV_GatherBlendSpacesFromGraph(TArray<
 		{
 			if (const FAnimNode_BlendSpacePlayer* const BlendSpacePlayer = Property->ContainerPtrToValuePtr<FAnimNode_BlendSpacePlayer>(PCV_Params.DefaultAnimInstance))
 			{
-				OutBlendSpaces.AddUnique(BlendSpacePlayer->BlendSpace);
+				OutBlendSpaces.AddUnique(BlendSpacePlayer->GetBlendSpace());
 			}
 		}
 	}
diff --git a/Engine/Source/Editor/AnimGraph/Private/AnimGraphDetails.cpp b/Engine/Source/Editor/AnimGraph/Private/AnimGraphDetails.cpp
--- a/Engine/Source/Editor/AnimGraph/Private/AnimGraphDetails.cpp
+++ b/Engine/Source/Editor/AnimGraph/Private/AnimGraphDetails.cpp
@@ -321,12 +321,12 @@ FReply FAnimGraphDetails::OnRemoveInputPoseClicked(UAnimGraphNode_LinkedInputPos
 FText FAnimGraphDetails::OnGetGroupText() const
 {
 	UAnimGraphNode_Root* Root = FBlueprintEditorUtils::GetAnimGraphRoot(Graph);
-	if(Root->Node.Group == NAME_None)
+	if(Root->Node.GetGroup() == NAME_None)
 	{
 		return LOCTEXT("DefaultGroup", "Default");
 	}
 		
-	return FText::FromName(Root->Node.Group);
+	return FText::FromName(Root->Node.GetGroup());
 }
 
 void FAnimGraphDetails::OnGroupTextCommitted(const FText& NewText, ETextCommit::Type InTextCommit)
@@ -392,7 +392,7 @@ void FAnimGraphDetails::RefreshGroupSource()
 	GroupSource.Empty();
 
 	UAnimGraphNode_Root* Root = FBlueprintEditorUtils::GetAnimGraphRoot(Graph);
-	if(Root->Node.Group != NAME_None)
+	if(Root->Node.GetGroup() != NAME_None)
 	{
 		GroupSource.Add(MakeShared<FText>(LOCTEXT("DefaultGroup", "Default")));
 	}
diff --git a/Engine/Source/Editor/AnimGraph/Private/AnimGraphModule.cpp b/Engine/Source/Editor/AnimGraph/Private/AnimGraphModule.cpp
--- a/Engine/Source/Editor/AnimGraph/Private/AnimGraphModule.cpp
+++ b/Engine/Source/Editor/AnimGraph/Private/AnimGraphModule.cpp
@@ -24,11 +24,6 @@
 #include "AnimGraphDetails.h"
 #include "AnimationGraphSchema.h"
 #include "AnimBlueprintCompiler.h"
-#include "AnimBlueprintCompilerHandler_Base.h"
-#include "AnimBlueprintCompilerHandler_CachedPose.h"
-#include "AnimBlueprintCompilerHandler_LinkedAnimGraph.h"
-#include "AnimBlueprintCompilerHandler_StateMachine.h"
-#include "AnimBlueprintCompilerHandler_Attributes.h"
 #include "Animation/AnimAttributes.h"
 #include "AnimGraphAttributes.h"
 #include "Animation/AnimSync.h"
@@ -47,32 +42,6 @@ void FAnimGraphModule::StartupModule()
 		return MakeShared<FAnimBlueprintCompilerContext>(CastChecked<UAnimBlueprint>(InBlueprint), InMessageLog, InCompileOptions);
 	});
 
-	// Register node compilation handlers
-	IAnimBlueprintCompilerHandlerCollection::RegisterHandler("AnimBlueprintCompilerHandler_Base", [](IAnimBlueprintCompilerCreationContext& InCreationContext)
-	{
-		return MakeUnique<FAnimBlueprintCompilerHandler_Base>(InCreationContext);
-	});
-
-	IAnimBlueprintCompilerHandlerCollection::RegisterHandler("AnimBlueprintCompilerHandler_CachedPose", [](IAnimBlueprintCompilerCreationContext& InCreationContext)
-	{
-		return MakeUnique<FAnimBlueprintCompilerHandler_CachedPose>(InCreationContext);
-	});
-
-	IAnimBlueprintCompilerHandlerCollection::RegisterHandler("AnimBlueprintCompilerHandler_LinkedAnimGraph", [](IAnimBlueprintCompilerCreationContext& InCreationContext)
-	{
-		return MakeUnique<FAnimBlueprintCompilerHandler_LinkedAnimGraph>(InCreationContext);
-	});
-
-	IAnimBlueprintCompilerHandlerCollection::RegisterHandler("AnimBlueprintCompilerHandler_StateMachine", [](IAnimBlueprintCompilerCreationContext& InCreationContext)
-	{
-		return MakeUnique<FAnimBlueprintCompilerHandler_StateMachine>(InCreationContext);
-	});
-
-	IAnimBlueprintCompilerHandlerCollection::RegisterHandler("AnimBlueprintCompilerHandler_Attributes", [](IAnimBlueprintCompilerCreationContext& InCreationContext)
-	{
-		return MakeUnique<FAnimBlueprintCompilerHandler_Attributes>(InCreationContext);
-	});
-
 	// Register the editor modes
 	FEditorModeRegistry::Get().RegisterMode<FAnimNodeEditMode>(AnimNodeEditModes::AnimNode, LOCTEXT("AnimNodeEditMode", "Anim Node"), FSlateIcon(), false);
 	FEditorModeRegistry::Get().RegisterMode<FTwoBoneIKEditMode>(AnimNodeEditModes::TwoBoneIK, LOCTEXT("TwoBoneIKEditMode", "2-Bone IK"), FSlateIcon(), false);
@@ -170,12 +139,6 @@ void FAnimGraphModule::StartupModule()
 
 void FAnimGraphModule::ShutdownModule()
 {
-	IAnimBlueprintCompilerHandlerCollection::UnregisterHandler("AnimBlueprintCompilerHandler_Base");
-	IAnimBlueprintCompilerHandlerCollection::UnregisterHandler("AnimBlueprintCompilerHandler_CachedPose");
-	IAnimBlueprintCompilerHandlerCollection::UnregisterHandler("AnimBlueprintCompilerHandler_LinkedAnimGraph");
-	IAnimBlueprintCompilerHandlerCollection::UnregisterHandler("AnimBlueprintCompilerHandler_StateMachine");
-	IAnimBlueprintCompilerHandlerCollection::UnregisterHandler("AnimBlueprintCompilerHandler_Attributes");
-
 	// Unregister the editor modes
 	FEditorModeRegistry::Get().UnregisterMode(AnimNodeEditModes::CCDIK);
 	FEditorModeRegistry::Get().UnregisterMode(AnimNodeEditModes::SplineIK);
diff --git a/Engine/Source/Editor/AnimGraph/Private/AnimGraphNode_AimOffsetLookAt.cpp b/Engine/Source/Editor/AnimGraph/Private/AnimGraphNode_AimOffsetLookAt.cpp
--- a/Engine/Source/Editor/AnimGraph/Private/AnimGraphNode_AimOffsetLookAt.cpp
+++ b/Engine/Source/Editor/AnimGraph/Private/AnimGraphNode_AimOffsetLookAt.cpp
@@ -31,7 +31,7 @@ FText UAnimGraphNode_AimOffsetLookAt::GetTooltipText() const
 
 FText UAnimGraphNode_AimOffsetLookAt::GetNodeTitle(ENodeTitleType::Type TitleType) const
 {
-	UBlendSpace* BlendSpaceToCheck = Node.BlendSpace;
+	UBlendSpace* BlendSpaceToCheck = Node.GetBlendSpace();
 	UEdGraphPin* BlendSpacePin = FindPin(GET_MEMBER_NAME_STRING_CHECKED(FAnimNode_AimOffsetLookAt, BlendSpace));
 	if (BlendSpacePin != nullptr && BlendSpaceToCheck == nullptr)
 	{
@@ -78,7 +78,7 @@ void UAnimGraphNode_AimOffsetLookAt::GetMenuActions(FBlueprintActionDatabaseRegi
 		static void SetNodeBlendSpace(UEdGraphNode* NewNode, bool /*bIsTemplateNode*/, TWeakObjectPtr<UBlendSpace> BlendSpace)
 		{
 			UAnimGraphNode_AimOffsetLookAt* BlendSpaceNode = CastChecked<UAnimGraphNode_AimOffsetLookAt>(NewNode);
-			BlendSpaceNode->Node.BlendSpace = BlendSpace.Get();
+			BlendSpaceNode->Node.SetBlendSpace(BlendSpace.Get());
 		}
 
 		static UBlueprintNodeSpawner* MakeBlendSpaceAction(TSubclassOf<UEdGraphNode> const NodeClass, const UBlendSpace* BlendSpace)
@@ -133,7 +133,7 @@ void UAnimGraphNode_AimOffsetLookAt::GetMenuActions(FBlueprintActionDatabaseRegi
 FBlueprintNodeSignature UAnimGraphNode_AimOffsetLookAt::GetSignature() const
 {
 	FBlueprintNodeSignature NodeSignature = Super::GetSignature();
-	NodeSignature.AddSubObject(Node.BlendSpace);
+	NodeSignature.AddSubObject(Node.GetBlendSpace());
 
 	return NodeSignature;
 }
@@ -142,7 +142,7 @@ void UAnimGraphNode_AimOffsetLookAt::SetAnimationAsset(UAnimationAsset* Asset)
 {
 	if (UBlendSpace* BlendSpace = Cast<UBlendSpace>(Asset))
 	{
-		Node.BlendSpace = BlendSpace;
+		Node.SetBlendSpace(BlendSpace);
 	}
 }
 
@@ -150,7 +150,7 @@ void UAnimGraphNode_AimOffsetLookAt::ValidateAnimNodeDuringCompilation(class USk
 {
 	Super::ValidateAnimNodeDuringCompilation(ForSkeleton, MessageLog);
 
-	UBlendSpace* BlendSpaceToCheck = Node.BlendSpace;
+	UBlendSpace* BlendSpaceToCheck = Node.GetBlendSpace();
 	UEdGraphPin* BlendSpacePin = FindPin(GET_MEMBER_NAME_STRING_CHECKED(FAnimNode_AimOffsetLookAt, BlendSpace));
 	if (BlendSpacePin != nullptr && BlendSpaceToCheck == nullptr)
 	{
@@ -230,7 +230,7 @@ void UAnimGraphNode_AimOffsetLookAt::GetNodeContextMenuActions(UToolMenu* Menu,
 
 void UAnimGraphNode_AimOffsetLookAt::GetAllAnimationSequencesReferred(TArray<UAnimationAsset*>& AnimationAssets) const
 {
-	if (Node.BlendSpace)
+	if (Node.GetBlendSpace())
 	{
 		HandleAnimReferenceCollection(Node.BlendSpace, AnimationAssets);
 	}
