Anim node data/compiler refactor

diff --git a/Engine/Source/Editor/AnimGraph/Private/AnimGraphNode_AssetPlayerBase.cpp b/Engine/Source/Editor/AnimGraph/Private/AnimGraphNode_AssetPlayerBase.cpp
--- a/Engine/Source/Editor/AnimGraph/Private/AnimGraphNode_AssetPlayerBase.cpp
+++ b/Engine/Source/Editor/AnimGraph/Private/AnimGraphNode_AssetPlayerBase.cpp
@@ -17,8 +17,6 @@
 #include "AnimGraphNode_PoseDriver.h"
 #include "UObject/UObjectIterator.h"
 #include "Animation/AnimLayerInterface.h"
-#include "IAnimBlueprintCompilerHandlerCollection.h"
-#include "AnimBlueprintCompilerHandler_Base.h"
 #include "IAnimBlueprintGeneratedClassCompiledData.h"
 #include "IAnimBlueprintCompilationContext.h"
 #include "Animation/AnimSync.h"
@@ -35,9 +33,21 @@ void UAnimGraphNode_AssetPlayerBase::Serialize(FArchive& Ar)
 
 	if(Ar.CustomVer(FUE5MainStreamObjectVersion::GUID) < FUE5MainStreamObjectVersion::AnimSyncGroupsExplicitSyncMethod)
 	{
-		if(SyncGroup.GroupName != NAME_None)
+		if(SyncGroup_DEPRECATED.GroupName != NAME_None)
 		{
-			SyncGroup.Method = EAnimSyncMethod::SyncGroup;
+			SyncGroup_DEPRECATED.Method = EAnimSyncMethod::SyncGroup;
+		}
+	}
+
+	if(Ar.CustomVer(FUE5MainStreamObjectVersion::GUID) < FUE5MainStreamObjectVersion::AnimNodeConstantDataRefactorPhase0)
+	{
+		FStructProperty* NodeProperty = GetFNodeProperty();
+		if(NodeProperty->Struct->IsChildOf(FAnimNode_AssetPlayerBase::StaticStruct()))
+		{
+			FAnimNode_AssetPlayerBase* Node = NodeProperty->ContainerPtrToValuePtr<FAnimNode_AssetPlayerBase>(this);
+			Node->SetGroupName(SyncGroup_DEPRECATED.GroupName);
+			Node->GetGroupRole(SyncGroup_DEPRECATED.GroupRole);
+			Node->SetGroupMethod(SyncGroup_DEPRECATED.Method);
 		}
 	}
 }
@@ -46,12 +56,17 @@ void UAnimGraphNode_AssetPlayerBase::PostEditChangeProperty(FPropertyChangedEven
 {
 	Super::PostEditChangeProperty(PropertyChangedEvent);
 
-	if(PropertyChangedEvent.GetPropertyName() == GET_MEMBER_NAME_CHECKED(FAnimationGroupReference, Method))
+	FStructProperty* NodeProperty = GetFNodeProperty();
+	if(NodeProperty->Struct->IsChildOf(FAnimNode_AssetPlayerBase::StaticStruct()))
 	{
-		if(SyncGroup.Method != EAnimSyncMethod::SyncGroup)
+		if(PropertyChangedEvent.GetPropertyName() == GET_MEMBER_NAME_CHECKED(FAnimNode_AssetPlayerBase, Method))
 		{
-			SyncGroup.GroupName = NAME_None;
-			SyncGroup.GroupRole = EAnimGroupRole::CanBeLeader;
+			FAnimNode_AssetPlayerBase* Node = NodeProperty->ContainerPtrToValuePtr<FAnimNode_AssetPlayerBase>(this);
+			if(Node->GetGroupMethod() != EAnimSyncMethod::SyncGroup)
+			{
+				Node->SetGroupName(NAME_None);
+				Node->GetGroupRole(EAnimGroupRole::CanBeLeader);
+			}
 		}
 	}
 }
@@ -135,9 +150,15 @@ void UAnimGraphNode_AssetPlayerBase::ValidateAnimNodeDuringCompilation(USkeleton
 {
 	Super::ValidateAnimNodeDuringCompilation(ForSkeleton, MessageLog);
 
-	if(SyncGroup.Method == EAnimSyncMethod::SyncGroup && SyncGroup.GroupName == NAME_None)
+	FStructProperty* NodeProperty = GetFNodeProperty();
+	if(NodeProperty->Struct->IsChildOf(FAnimNode_AssetPlayerBase::StaticStruct()))
 	{
-		MessageLog.Error(*LOCTEXT("NoSyncGroupSupplied", "Node @@ is set to use sync groups, but no sync group has been supplied").ToString(), this);
+		FAnimNode_AssetPlayerBase* Node = NodeProperty->ContainerPtrToValuePtr<FAnimNode_AssetPlayerBase>(this);
+
+		if(Node->GetGroupMethod() == EAnimSyncMethod::SyncGroup && Node->GetGroupName() == NAME_None)
+		{
+			MessageLog.Error(*LOCTEXT("NoSyncGroupSupplied", "Node @@ is set to use sync groups, but no sync group has been supplied").ToString(), this);
+		}
 	}
 }
 
@@ -145,9 +166,15 @@ void UAnimGraphNode_AssetPlayerBase::GetOutputLinkAttributes(FNodeAttributeArray
 {
 	OutAttributes.Add(UE::Anim::FAttributes::Curves);
 	OutAttributes.Add(UE::Anim::FAttributes::Attributes);
-	if(SyncGroup.Method == EAnimSyncMethod::Graph)
+
+	FStructProperty* NodeProperty = GetFNodeProperty();
+	if(NodeProperty->Struct->IsChildOf(FAnimNode_AssetPlayerBase::StaticStruct()))
 	{
-		OutAttributes.Add(UE::Anim::FAnimSync::Attribute);
+		const FAnimNode_AssetPlayerBase* Node = NodeProperty->ContainerPtrToValuePtr<FAnimNode_AssetPlayerBase>(this);
+		if(Node->GetGroupMethod() == EAnimSyncMethod::Graph)
+		{
+			OutAttributes.Add(UE::Anim::FAnimSync::Attribute);
+		}
 	}
 }
 
diff --git a/Engine/Source/Editor/AnimGraph/Private/AnimGraphNode_Base.cpp b/Engine/Source/Editor/AnimGraph/Private/AnimGraphNode_Base.cpp
--- a/Engine/Source/Editor/AnimGraph/Private/AnimGraphNode_Base.cpp
+++ b/Engine/Source/Editor/AnimGraph/Private/AnimGraphNode_Base.cpp
@@ -18,9 +18,12 @@
 #include "Kismet2/CompilerResultsLog.h"
 #include "Subsystems/AssetEditorSubsystem.h"
 #include "AnimBlueprintCompiler.h"
-#include "AnimBlueprintCompilerHandler_Base.h"
+#include "AnimBlueprintExtension_Base.h"
+#include "AnimBlueprintExtension_Attributes.h"
+#include "AnimBlueprintExtension_PropertyAccess.h"
 #include "IAnimBlueprintCompilationContext.h"
 #include "AnimBlueprintCompilationContext.h"
+#include "AnimBlueprintExtension.h"
 #include "FindInBlueprintManager.h"
 #include "UObject/ReleaseObjectVersion.h"
 
@@ -37,8 +40,8 @@ UAnimGraphNode_Base::UAnimGraphNode_Base(const FObjectInitializer& ObjectInitial
 void UAnimGraphNode_Base::ExpandNode(FKismetCompilerContext& CompilerContext, UEdGraph* SourceGraph)
 {
 	TUniquePtr<IAnimBlueprintCompilationContext> CompilationContext = IAnimBlueprintCompilationContext::Get(CompilerContext);
-	FAnimBlueprintCompilerHandler_Base* Handler = CompilationContext->GetHandler<FAnimBlueprintCompilerHandler_Base>("AnimBlueprintCompilerHandler_Base");
-	Handler->CreateEvaluationHandlerForNode(*CompilationContext.Get(), this);
+	UAnimBlueprintExtension_Base* Extension = UAnimBlueprintExtension_Base::GetExtension<UAnimBlueprintExtension_Base>(GetAnimBlueprint());
+	Extension->CreateEvaluationHandlerForNode(*CompilationContext.Get(), this);
 }
 
 void UAnimGraphNode_Base::PreEditChange(FProperty* PropertyThatWillChange)
@@ -98,6 +101,22 @@ void UAnimGraphNode_Base::Serialize(FArchive& Ar)
 	}
 }
 
+void UAnimGraphNode_Base::PostPlacedNewNode()
+{
+	Super::PostPlacedNewNode();
+
+	// This makes sure that all anim BP extensions are registered that this node needs
+	UAnimBlueprintExtension::RequestExtensionsForNode(this);
+}
+
+void UAnimGraphNode_Base::DestroyNode()
+{
+	// This node may have been the last using its extension, so refresh
+	GetAnimBlueprint()->RequestRefreshExtensions();
+
+	Super::DestroyNode();
+}
+
 void UAnimGraphNode_Base::CreateOutputPins()
 {
 	if (!IsSinkNode())
@@ -140,6 +159,21 @@ void UAnimGraphNode_Base::CopyTermDefaultsToDefaultObject(IAnimBlueprintCopyTerm
 	OnCopyTermDefaultsToDefaultObject(InCompilationContext, InPerNodeContext, OutCompiledData);
 }
 
+void UAnimGraphNode_Base::OverrideAssets(IAnimBlueprintNodeOverrideAssetsContext& InContext) const
+{
+	if(InContext.GetAssets().Num() > 0)
+	{
+		if(UAnimationAsset* AnimationAsset = Cast<UAnimationAsset>(InContext.GetAssets()[0]))
+		{
+			// Call the legacy implementation
+PRAGMA_DISABLE_DEPRECATION_WARNINGS
+			InContext.GetAnimNode<FAnimNode_Base>().OverrideAsset(AnimationAsset);
+PRAGMA_ENABLE_DEPRECATION_WARNINGS
+		}
+	}
+	OnOverrideAssets(InContext);
+}
+
 void UAnimGraphNode_Base::InternalPinCreation(TArray<UEdGraphPin*>* OldPins)
 {
 	// preload required assets first before creating pins
@@ -383,10 +417,10 @@ void UAnimGraphNode_Base::GetPinHoverText(const UEdGraphPin& Pin, FString& Hover
 
 void UAnimGraphNode_Base::ProcessDuringCompilation(IAnimBlueprintCompilationContext& InCompilationContext, IAnimBlueprintGeneratedClassCompiledData& OutCompiledData)
 {
-	FAnimBlueprintCompilerHandler_Base* HandlerBase = InCompilationContext.GetHandler<FAnimBlueprintCompilerHandler_Base>("AnimBlueprintCompilerHandler_Base");
+	UAnimBlueprintExtension_Base* Extension = UAnimBlueprintExtension::GetExtension<UAnimBlueprintExtension_Base>(GetAnimBlueprint());
 
 	// Record pose pins for later patchup and gather pins that have an associated evaluation handler
-	HandlerBase->AddStructEvalHandlers(this, InCompilationContext, OutCompiledData);
+	Extension->AddStructEvalHandlers(this, InCompilationContext, OutCompiledData);
 
 	// Call the override point
 	OnProcessDuringCompilation(InCompilationContext, OutCompiledData);
@@ -512,6 +546,12 @@ bool UAnimGraphNode_Base::IsPinExposedAndLinked(const FString& InPinName, const
 	return Pin != nullptr && Pin->LinkedTo.Num() > 0 && Pin->LinkedTo[0] != nullptr;
 }
 
+bool UAnimGraphNode_Base::IsPinExposedAndBound(const FString& InPinName, const EEdGraphPinDirection InDirection) const
+{
+	UEdGraphPin* Pin = FindPin(InPinName, InDirection);
+	return Pin != nullptr && Pin->LinkedTo.Num() == 0 && PropertyBindings.Find(Pin->GetFName()) != nullptr;
+}
+
 void UAnimGraphNode_Base::PinConnectionListChanged(UEdGraphPin* Pin)
 {
 	if(Pin->LinkedTo.Num() > 0)
diff --git a/Engine/Source/Editor/AnimGraph/Private/AnimGraphNode_BlendListBase.cpp b/Engine/Source/Editor/AnimGraph/Private/AnimGraphNode_BlendListBase.cpp
--- a/Engine/Source/Editor/AnimGraph/Private/AnimGraphNode_BlendListBase.cpp
+++ b/Engine/Source/Editor/AnimGraph/Private/AnimGraphNode_BlendListBase.cpp
@@ -113,7 +113,7 @@ void UAnimGraphNode_BlendListBase::GetOutputLinkAttributes(FNodeAttributeArray&
 	{
 		const FAnimNode_BlendListBase* BaseNode = GetFNodeProperty()->ContainerPtrToValuePtr<FAnimNode_BlendListBase>(this);
 
-		if(BaseNode->TransitionType == EBlendListTransitionType::Inertialization)
+		if(BaseNode->GetTransitionType() == EBlendListTransitionType::Inertialization)
 		{
 			OutAttributes.Add(UE::Anim::IInertializationRequester::Attribute);
 		}
diff --git a/Engine/Source/Editor/AnimGraph/Private/AnimGraphNode_BlendListByEnum.cpp b/Engine/Source/Editor/AnimGraph/Private/AnimGraphNode_BlendListByEnum.cpp
--- a/Engine/Source/Editor/AnimGraph/Private/AnimGraphNode_BlendListByEnum.cpp
+++ b/Engine/Source/Editor/AnimGraph/Private/AnimGraphNode_BlendListByEnum.cpp
@@ -307,8 +307,9 @@ void UAnimGraphNode_BlendListByEnum::BakeDataDuringCompilation(class FCompilerRe
 		BoundEnum->ConditionalPostLoad();
 
 		// Zero the array out so it looks up the default value, and stat counting at index 1
-		Node.EnumToPoseIndex.Empty();
-		Node.EnumToPoseIndex.AddZeroed(BoundEnum->NumEnums());
+		TArray<int32> EnumToPoseIndex;
+		EnumToPoseIndex.Empty();
+		EnumToPoseIndex.AddZeroed(BoundEnum->NumEnums());
 		int32 PinIndex = 1;
 
 		// Run thru the enum entries
@@ -319,7 +320,7 @@ void UAnimGraphNode_BlendListByEnum::BakeDataDuringCompilation(class FCompilerRe
 
 			if (EnumIndex != INDEX_NONE)
 			{
-				Node.EnumToPoseIndex[EnumIndex] = PinIndex;
+				EnumToPoseIndex[EnumIndex] = PinIndex;
 			}
 			else
 			{
@@ -328,6 +329,8 @@ void UAnimGraphNode_BlendListByEnum::BakeDataDuringCompilation(class FCompilerRe
 
 			++PinIndex;
 		}
+
+		Node.SetEnumToPoseIndex(EnumToPoseIndex);
 	}
 }
 
diff --git a/Engine/Source/Editor/AnimGraph/Private/AnimGraphNode_BlendListByInt.cpp b/Engine/Source/Editor/AnimGraph/Private/AnimGraphNode_BlendListByInt.cpp
--- a/Engine/Source/Editor/AnimGraph/Private/AnimGraphNode_BlendListByInt.cpp
+++ b/Engine/Source/Editor/AnimGraph/Private/AnimGraphNode_BlendListByInt.cpp
@@ -32,6 +32,8 @@ FText UAnimGraphNode_BlendListByInt::GetNodeTitle(ENodeTitleType::Type TitleType
 
 void UAnimGraphNode_BlendListByInt::PostPlacedNewNode()
 {
+	Super::PostPlacedNewNode();
+
 	Node.AddPose();
 	ReconstructNode();
 }
diff --git a/Engine/Source/Editor/AnimGraph/Private/AnimGraphNode_BlendSpaceEvaluator.cpp b/Engine/Source/Editor/AnimGraph/Private/AnimGraphNode_BlendSpaceEvaluator.cpp
--- a/Engine/Source/Editor/AnimGraph/Private/AnimGraphNode_BlendSpaceEvaluator.cpp
+++ b/Engine/Source/Editor/AnimGraph/Private/AnimGraphNode_BlendSpaceEvaluator.cpp
@@ -4,6 +4,7 @@
 #include "ToolMenus.h"
 #include "AnimGraphCommands.h"
 #include "Kismet2/CompilerResultsLog.h"
+#include "IAnimBlueprintNodeOverrideAssetsContext.h"
 
 /////////////////////////////////////////////////////
 // UAnimGraphNode_BlendSpaceEvaluator
@@ -39,11 +40,11 @@ FText UAnimGraphNode_BlendSpaceEvaluator::GetNodeTitleForBlendSpace(ENodeTitleTy
 		TitleArgs.Add(TEXT("BlendSpaceName"), BlendSpaceName);
 		FText Title = FText::Format(LOCTEXT("BlendSpaceEvaluatorFullTitle", "{BlendSpaceName}\nBlendspace Evaluator"), TitleArgs);
 
-		if ((TitleType == ENodeTitleType::FullTitle) && (SyncGroup.GroupName != NAME_None))
+		if ((TitleType == ENodeTitleType::FullTitle) && (Node.GetGroupName() != NAME_None))
 		{
 			FFormatNamedArguments Args;
 			Args.Add(TEXT("Title"), Title);
-			Args.Add(TEXT("SyncGroupName"), FText::FromName(SyncGroup.GroupName));
+			Args.Add(TEXT("SyncGroupName"), FText::FromName(Node.GetGroupName()));
 			Title = FText::Format(LOCTEXT("BlendSpaceNodeGroupSubtitle", "{Title}\nSync group {SyncGroupName}"), Args);
 		}
 		// FText::Format() is slow, so we cache this to save on performance
@@ -55,7 +56,7 @@ FText UAnimGraphNode_BlendSpaceEvaluator::GetNodeTitleForBlendSpace(ENodeTitleTy
 
 FText UAnimGraphNode_BlendSpaceEvaluator::GetNodeTitle(ENodeTitleType::Type TitleType) const
 {
-	if (Node.BlendSpace == nullptr)
+	if (Node.GetBlendSpace() == nullptr)
 	{
 		// we may have a valid variable connected or default pin value
 		UEdGraphPin* BlendSpacePin = FindPin(GET_MEMBER_NAME_STRING_CHECKED(FAnimNode_BlendSpacePlayer, BlendSpace));
@@ -83,7 +84,7 @@ FText UAnimGraphNode_BlendSpaceEvaluator::GetNodeTitle(ENodeTitleType::Type Titl
 	//        choose to mark this dirty when that happens for this to properly work
 	else //if (!CachedNodeTitles.IsTitleCached(TitleType, this))
 	{
-		return GetNodeTitleForBlendSpace(TitleType, Node.BlendSpace);
+		return GetNodeTitleForBlendSpace(TitleType, Node.GetBlendSpace());
 	}
 }
 
@@ -97,7 +98,7 @@ void UAnimGraphNode_BlendSpaceEvaluator::ValidateAnimNodeDuringCompilation(class
 {
 	Super::ValidateAnimNodeDuringCompilation(ForSkeleton, MessageLog);
 
-	UBlendSpace* BlendSpaceToCheck = Node.BlendSpace;
+	UBlendSpace* BlendSpaceToCheck = Node.GetBlendSpace();
 	UEdGraphPin* BlendSpacePin = FindPin(GET_MEMBER_NAME_STRING_CHECKED(FAnimNode_BlendSpaceEvaluator, BlendSpace));
 	if (BlendSpacePin != nullptr && BlendSpaceToCheck == nullptr)
 	{
@@ -136,10 +137,7 @@ void UAnimGraphNode_BlendSpaceEvaluator::ValidateAnimNodeDuringCompilation(class
 void UAnimGraphNode_BlendSpaceEvaluator::BakeDataDuringCompilation(class FCompilerResultsLog& MessageLog)
 {
 	UAnimBlueprint* AnimBlueprint = GetAnimBlueprint();
-	AnimBlueprint->FindOrAddGroup(SyncGroup.GroupName);
-	Node.GroupName = SyncGroup.GroupName;
-	Node.GroupRole = SyncGroup.GroupRole;
-	Node.Method = SyncGroup.Method;
+	AnimBlueprint->FindOrAddGroup(Node.GetGroupName());
 }
 
 void UAnimGraphNode_BlendSpaceEvaluator::GetNodeContextMenuActions(UToolMenu* Menu, UGraphNodeContextMenuContext* Context) const
@@ -159,7 +157,19 @@ void UAnimGraphNode_BlendSpaceEvaluator::SetAnimationAsset(UAnimationAsset* Asse
 {
 	if (UBlendSpace* BlendSpace = Cast<UBlendSpace>(Asset))
 	{
-		Node.BlendSpace = BlendSpace;
+		Node.SetBlendSpace(BlendSpace);
+	}
+}
+
+void UAnimGraphNode_BlendSpaceEvaluator::OnOverrideAssets(IAnimBlueprintNodeOverrideAssetsContext& InContext) const
+{
+	if(InContext.GetAssets().Num() > 0)
+	{
+		if (UBlendSpace* BlendSpace = Cast<UBlendSpace>(InContext.GetAssets()[0]))
+		{
+			FAnimNode_BlendSpaceEvaluator& AnimNode = InContext.GetAnimNode<FAnimNode_BlendSpaceEvaluator>();
+			AnimNode.SetBlendSpace(BlendSpace);
+		}
 	}
 }
 
@@ -170,7 +180,7 @@ bool UAnimGraphNode_BlendSpaceEvaluator::DoesSupportTimeForTransitionGetter() co
 
 UAnimationAsset* UAnimGraphNode_BlendSpaceEvaluator::GetAnimationAsset() const 
 {
-	UBlendSpace* BlendSpace = Node.BlendSpace;
+	UBlendSpace* BlendSpace = Node.GetBlendSpace();
 	UEdGraphPin* BlendSpacePin = FindPin(GET_MEMBER_NAME_STRING_CHECKED(FAnimNode_BlendSpaceEvaluator, BlendSpace));
 	if (BlendSpacePin != nullptr && BlendSpace == nullptr)
 	{
diff --git a/Engine/Source/Editor/AnimGraph/Private/AnimGraphNode_BlendSpaceGraphBase.cpp b/Engine/Source/Editor/AnimGraph/Private/AnimGraphNode_BlendSpaceGraphBase.cpp
--- a/Engine/Source/Editor/AnimGraph/Private/AnimGraphNode_BlendSpaceGraphBase.cpp
+++ b/Engine/Source/Editor/AnimGraph/Private/AnimGraphNode_BlendSpaceGraphBase.cpp
@@ -20,12 +20,14 @@
 #include "EdGraphUtilities.h"
 #include "IAnimBlueprintCompilationContext.h"
 #include "AnimationBlendSpaceSampleGraph.h"
+#include "AnimBlueprintExtension.h"
 #include "AnimGraphNode_BlendSpaceSampleResult.h"
 #include "AnimGraphNode_SequencePlayer.h"
 #include "Kismet2/Kismet2NameValidators.h"
 #include "AnimNodes/AnimNode_BlendSpaceGraphBase.h"
 #include "BlueprintEditor.h"
 #include "Animation/AnimSequence.h"
+#include "AnimBlueprintExtension_BlendSpaceGraph.h"
 
 #define LOCTEXT_NAMESPACE "UAnimGraphNode_BlendSpaceGraphBase"
 
@@ -122,10 +124,9 @@ void UAnimGraphNode_BlendSpaceGraphBase::OnProcessDuringCompilation(IAnimBluepri
 
 void UAnimGraphNode_BlendSpaceGraphBase::OnCopyTermDefaultsToDefaultObject(IAnimBlueprintCopyTermDefaultsContext& InCompilationContext, IAnimBlueprintNodeCopyTermDefaultsContext& InPerNodeContext, IAnimBlueprintGeneratedClassCompiledData& OutCompiledData)
 {
-	UAnimGraphNode_BlendSpaceGraphBase* TrueNode = InCompilationContext.GetMessageLog().FindSourceObjectTypeChecked<UAnimGraphNode_BlendSpaceGraphBase>(this);
-
 	FAnimNode_BlendSpaceGraphBase* DestinationNode = reinterpret_cast<FAnimNode_BlendSpaceGraphBase*>(InPerNodeContext.GetDestinationPtr());
-	DestinationNode->BlendSpace = OutCompiledData.AddBlendSpace(BlendSpace);
+	UAnimBlueprintExtension_BlendSpaceGraph* Extension = UAnimBlueprintExtension::GetExtension<UAnimBlueprintExtension_BlendSpaceGraph>(GetAnimBlueprint());
+	DestinationNode->BlendSpace = Extension->AddBlendSpace(BlendSpace);
 }
 
 void UAnimGraphNode_BlendSpaceGraphBase::SetupFromAsset(UBlendSpace* InBlendSpace, bool bInIsTemplateNode)
@@ -322,6 +323,8 @@ void UAnimGraphNode_BlendSpaceGraphBase::PostPasteNode()
 
 void UAnimGraphNode_BlendSpaceGraphBase::PostPlacedNewNode()
 {
+	Super::PostPlacedNewNode();
+	
 	// Create a new graph & blendspace if we havent been set up already
 	if(BlendSpaceGraph == nullptr)
 	{
@@ -421,7 +424,7 @@ UAnimationBlendSpaceSampleGraph* UAnimGraphNode_BlendSpaceGraphBase::AddGraphInt
 		FGraphNodeCreator<UAnimGraphNode_SequencePlayer> SequencePlayerNodeCreator(*SampleGraph);
 		UAnimGraphNode_SequencePlayer* SequencePlayer = SequencePlayerNodeCreator.CreateNode();
 		SequencePlayer->SetAnimationAsset(InSequence);
-		SequencePlayer->SyncGroup.Method = EAnimSyncMethod::Graph;
+		SequencePlayer->Node.SetGroupMethod(EAnimSyncMethod::Graph);
 		SequencePlayerNodeCreator.Finalize();
 
 		// Offset node in X
@@ -493,4 +496,9 @@ void UAnimGraphNode_BlendSpaceGraphBase::GetInputLinkAttributes(FNodeAttributeAr
 	}
 }
 
+void UAnimGraphNode_BlendSpaceGraphBase::GetRequiredExtensions(TArray<TSubclassOf<UAnimBlueprintExtension>>& OutExtensions) const
+{
+	OutExtensions.Add(UAnimBlueprintExtension_BlendSpaceGraph::StaticClass());
+}
+
 #undef LOCTEXT_NAMESPACE
diff --git a/Engine/Source/Editor/AnimGraph/Private/AnimGraphNode_BlendSpacePlayer.cpp b/Engine/Source/Editor/AnimGraph/Private/AnimGraphNode_BlendSpacePlayer.cpp
--- a/Engine/Source/Editor/AnimGraph/Private/AnimGraphNode_BlendSpacePlayer.cpp
+++ b/Engine/Source/Editor/AnimGraph/Private/AnimGraphNode_BlendSpacePlayer.cpp
@@ -11,6 +11,8 @@
 #include "Animation/AimOffsetBlendSpace.h"
 #include "Animation/AimOffsetBlendSpace1D.h"
 #include "AnimGraphCommands.h"
+#include "IAnimBlueprintCopyTermDefaultsContext.h"
+#include "IAnimBlueprintNodeOverrideAssetsContext.h"
 
 /////////////////////////////////////////////////////
 // UAnimGraphNode_BlendSpacePlayer
@@ -50,12 +52,12 @@ FText UAnimGraphNode_BlendSpacePlayer::GetNodeTitleForBlendSpace(ENodeTitleType:
 			FFormatNamedArguments Args;
 			Args.Add(TEXT("Title"), Title);
 
-			if(SyncGroup.Method == EAnimSyncMethod::SyncGroup)
+			if(Node.GetGroupMethod() == EAnimSyncMethod::SyncGroup)
 			{
-				Args.Add(TEXT("SyncGroupName"), FText::FromName(SyncGroup.GroupName));
+				Args.Add(TEXT("SyncGroupName"), FText::FromName(Node.GetGroupName()));
 				Title = FText::Format(LOCTEXT("BlendSpaceNodeGroupSubtitle", "{Title}\nSync group {SyncGroupName}"), Args);
 			}
-			else if(SyncGroup.Method == EAnimSyncMethod::Graph)
+			else if(Node.GetGroupMethod() == EAnimSyncMethod::Graph)
 			{
 				Title = FText::Format(LOCTEXT("BlendSpaceNodeGroupSubtitle", "{Title}\nGraph sync group"), Args);
 
@@ -84,7 +86,7 @@ FText UAnimGraphNode_BlendSpacePlayer::GetNodeTitleForBlendSpace(ENodeTitleType:
 
 FText UAnimGraphNode_BlendSpacePlayer::GetNodeTitle(ENodeTitleType::Type TitleType) const
 {	
-	if (Node.BlendSpace == nullptr)
+	if (Node.GetBlendSpace() == nullptr)
 	{
 		// we may have a valid variable connected or default pin value
 		UEdGraphPin* BlendSpacePin = FindPin(GET_MEMBER_NAME_STRING_CHECKED(FAnimNode_BlendSpacePlayer, BlendSpace));
@@ -112,7 +114,7 @@ FText UAnimGraphNode_BlendSpacePlayer::GetNodeTitle(ENodeTitleType::Type TitleTy
 	//        choose to mark this dirty when that happens for this to properly work
 	else //if (!CachedNodeTitles.IsTitleCached(TitleType, this))
 	{
-		return GetNodeTitleForBlendSpace(TitleType, Node.BlendSpace);
+		return GetNodeTitleForBlendSpace(TitleType, Node.GetBlendSpace());
 	}
 }
 
@@ -120,7 +122,7 @@ void UAnimGraphNode_BlendSpacePlayer::ValidateAnimNodeDuringCompilation(class US
 {
 	Super::ValidateAnimNodeDuringCompilation(ForSkeleton, MessageLog);
 
-	UBlendSpace* BlendSpaceToCheck = Node.BlendSpace;
+	UBlendSpace* BlendSpaceToCheck = Node.GetBlendSpace();
 	UEdGraphPin* BlendSpacePin = FindPin(GET_MEMBER_NAME_STRING_CHECKED(FAnimNode_BlendSpacePlayer, BlendSpace));
 	if (BlendSpacePin != nullptr && BlendSpaceToCheck == nullptr)
 	{
@@ -159,10 +161,7 @@ void UAnimGraphNode_BlendSpacePlayer::ValidateAnimNodeDuringCompilation(class US
 void UAnimGraphNode_BlendSpacePlayer::BakeDataDuringCompilation(class FCompilerResultsLog& MessageLog)
 {
 	UAnimBlueprint* AnimBlueprint = GetAnimBlueprint();
-	AnimBlueprint->FindOrAddGroup(SyncGroup.GroupName);
-	Node.GroupName = SyncGroup.GroupName;
-	Node.GroupRole = SyncGroup.GroupRole;
-	Node.Method = SyncGroup.Method;
+	AnimBlueprint->FindOrAddGroup(Node.GetGroupName());
 }
 
 void UAnimGraphNode_BlendSpacePlayer::GetNodeContextMenuActions(UToolMenu* Menu, UGraphNodeContextMenuContext* Context) const
@@ -186,7 +185,7 @@ void UAnimGraphNode_BlendSpacePlayer::GetMenuActions(FBlueprintActionDatabaseReg
 		static void SetNodeBlendSpace(UEdGraphNode* NewNode, bool /*bIsTemplateNode*/, TWeakObjectPtr<UBlendSpace> BlendSpace)
 		{
 			UAnimGraphNode_BlendSpacePlayer* BlendSpaceNode = CastChecked<UAnimGraphNode_BlendSpacePlayer>(NewNode);
-			BlendSpaceNode->Node.BlendSpace = BlendSpace.Get();
+			BlendSpaceNode->Node.SetBlendSpace(BlendSpace.Get());
 		}
 
 		static UBlueprintNodeSpawner* MakeBlendSpaceAction(TSubclassOf<UEdGraphNode> const NodeClass, const UBlendSpace* BlendSpace)
@@ -241,7 +240,7 @@ void UAnimGraphNode_BlendSpacePlayer::GetMenuActions(FBlueprintActionDatabaseReg
 FBlueprintNodeSignature UAnimGraphNode_BlendSpacePlayer::GetSignature() const
 {
 	FBlueprintNodeSignature NodeSignature = Super::GetSignature();
-	NodeSignature.AddSubObject(Node.BlendSpace);
+	NodeSignature.AddSubObject(Node.GetBlendSpace());
 
 	return NodeSignature;
 }
@@ -250,13 +249,25 @@ void UAnimGraphNode_BlendSpacePlayer::SetAnimationAsset(UAnimationAsset* Asset)
 {
 	if (UBlendSpace* BlendSpace = Cast<UBlendSpace>(Asset))
 	{
-		Node.BlendSpace = BlendSpace;
+		Node.SetBlendSpace(BlendSpace);
+	}
+}
+
+void UAnimGraphNode_BlendSpacePlayer::OnOverrideAssets(IAnimBlueprintNodeOverrideAssetsContext& InContext) const
+{
+	if(InContext.GetAssets().Num() > 0)
+	{
+		if (UBlendSpace* BlendSpace = Cast<UBlendSpace>(InContext.GetAssets()[0]))
+		{
+			FAnimNode_BlendSpacePlayer& AnimNode = InContext.GetAnimNode<FAnimNode_BlendSpacePlayer>();
+			AnimNode.SetBlendSpace(BlendSpace);
+		}
 	}
 }
 
 void UAnimGraphNode_BlendSpacePlayer::GetAllAnimationSequencesReferred(TArray<UAnimationAsset*>& AnimationAssets) const
 {
-	if(Node.BlendSpace)
+	if(Node.GetBlendSpace())
 	{
 		HandleAnimReferenceCollection(Node.BlendSpace, AnimationAssets);
 	}
@@ -274,7 +285,7 @@ bool UAnimGraphNode_BlendSpacePlayer::DoesSupportTimeForTransitionGetter() const
 
 UAnimationAsset* UAnimGraphNode_BlendSpacePlayer::GetAnimationAsset() const 
 {
-	UBlendSpace* BlendSpace = Node.BlendSpace;
+	UBlendSpace* BlendSpace = Node.GetBlendSpace();
 	UEdGraphPin* BlendSpacePin = FindPin(GET_MEMBER_NAME_STRING_CHECKED(FAnimNode_BlendSpacePlayer, BlendSpace));
 	if (BlendSpacePin != nullptr && BlendSpace == nullptr)
 	{
diff --git a/Engine/Source/Editor/AnimGraph/Private/AnimGraphNode_CustomProperty.cpp b/Engine/Source/Editor/AnimGraph/Private/AnimGraphNode_CustomProperty.cpp
--- a/Engine/Source/Editor/AnimGraph/Private/AnimGraphNode_CustomProperty.cpp
+++ b/Engine/Source/Editor/AnimGraph/Private/AnimGraphNode_CustomProperty.cpp
@@ -6,18 +6,8 @@
 #include "Kismet2/BlueprintEditorUtils.h"
 #include "DetailLayoutBuilder.h"
 #include "AnimationGraphSchema.h"
-#include "Animation/AnimInstance.h"
-#include "Animation/AnimLayerInterface.h"
 #include "DetailCategoryBuilder.h"
-#include "Widgets/Text/STextBlock.h"
-#include "Widgets/Input/SCheckBox.h"
-#include "Widgets/SBoxPanel.h"
-#include "DetailWidgetRow.h"
 #include "PropertyCustomizationHelpers.h"
-#include "KismetCompilerMisc.h"
-#include "AnimBlueprintCompilerHandler_Base.h"
-#include "IAnimBlueprintCompilerHandlerCollection.h"
-#include "IAnimBlueprintCompilerHandler.h"
 #include "KismetCompiler.h"
 #include "IAnimBlueprintCompilationContext.h"
 
diff --git a/Engine/Source/Editor/AnimGraph/Private/AnimGraphNode_IdentityPose.cpp b/Engine/Source/Editor/AnimGraph/Private/AnimGraphNode_IdentityPose.cpp
--- a/Engine/Source/Editor/AnimGraph/Private/AnimGraphNode_IdentityPose.cpp
+++ b/Engine/Source/Editor/AnimGraph/Private/AnimGraphNode_IdentityPose.cpp
@@ -10,7 +10,7 @@
 UAnimGraphNode_IdentityPose::UAnimGraphNode_IdentityPose(const FObjectInitializer& ObjectInitializer)
 	: Super(ObjectInitializer)
 {
-	Node.RefPoseType = EIT_Additive;
+	Node.SetRefPoseType(EIT_Additive);
 }
 
 FText UAnimGraphNode_IdentityPose::GetTooltipText() const
diff --git a/Engine/Source/Editor/AnimGraph/Private/AnimGraphNode_LinkedAnimGraphBase.cpp b/Engine/Source/Editor/AnimGraph/Private/AnimGraphNode_LinkedAnimGraphBase.cpp
--- a/Engine/Source/Editor/AnimGraph/Private/AnimGraphNode_LinkedAnimGraphBase.cpp
+++ b/Engine/Source/Editor/AnimGraph/Private/AnimGraphNode_LinkedAnimGraphBase.cpp
@@ -18,6 +18,7 @@
 #include "Animation/AnimNode_LinkedAnimGraph.h"
 #include "AnimGraphAttributes.h"
 #include "IAnimBlueprintCopyTermDefaultsContext.h"
+#include "AnimBlueprintExtension_LinkedAnimGraph.h"
 
 #define LOCTEXT_NAMESPACE "LinkedAnimGraph"
 
@@ -489,4 +490,9 @@ void UAnimGraphNode_LinkedAnimGraphBase::OnCopyTermDefaultsToDefaultObject(IAnim
 	DestinationNode->NodeIndex = InPerNodeContext.GetNodePropertyIndex();
 }
 
+void UAnimGraphNode_LinkedAnimGraphBase::GetRequiredExtensions(TArray<TSubclassOf<UAnimBlueprintExtension>>& OutExtensions) const
+{
+	OutExtensions.Add(UAnimBlueprintExtension_LinkedAnimGraph::StaticClass());
+}
+
 #undef LOCTEXT_NAMESPACE
diff --git a/Engine/Source/Editor/AnimGraph/Private/AnimGraphNode_LocalRefPose.cpp b/Engine/Source/Editor/AnimGraph/Private/AnimGraphNode_LocalRefPose.cpp
--- a/Engine/Source/Editor/AnimGraph/Private/AnimGraphNode_LocalRefPose.cpp
+++ b/Engine/Source/Editor/AnimGraph/Private/AnimGraphNode_LocalRefPose.cpp
@@ -10,7 +10,7 @@
 UAnimGraphNode_LocalRefPose::UAnimGraphNode_LocalRefPose(const FObjectInitializer& ObjectInitializer)
 	: Super(ObjectInitializer)
 {
-	Node.RefPoseType = EIT_LocalSpace;
+	Node.SetRefPoseType(EIT_LocalSpace);
 }
 
 FText UAnimGraphNode_LocalRefPose::GetTooltipText() const
diff --git a/Engine/Source/Editor/AnimGraph/Private/AnimGraphNode_MultiWayBlend.cpp b/Engine/Source/Editor/AnimGraph/Private/AnimGraphNode_MultiWayBlend.cpp
--- a/Engine/Source/Editor/AnimGraph/Private/AnimGraphNode_MultiWayBlend.cpp
+++ b/Engine/Source/Editor/AnimGraph/Private/AnimGraphNode_MultiWayBlend.cpp
@@ -94,6 +94,8 @@ void UAnimGraphNode_MultiWayBlend::RemovePinFromBlendNode(UEdGraphPin* Pin)
 
 void UAnimGraphNode_MultiWayBlend::PostPlacedNewNode()
 {
+	Super::PostPlacedNewNode();
+	
 	// Make sure we start out with two inputs
 	Node.AddPose();
 	Node.AddPose();
diff --git a/Engine/Source/Editor/AnimGraph/Private/AnimGraphNode_PoseHandler.cpp b/Engine/Source/Editor/AnimGraph/Private/AnimGraphNode_PoseHandler.cpp
--- a/Engine/Source/Editor/AnimGraph/Private/AnimGraphNode_PoseHandler.cpp
+++ b/Engine/Source/Editor/AnimGraph/Private/AnimGraphNode_PoseHandler.cpp
@@ -4,6 +4,7 @@
 #include "Animation/PoseAsset.h"
 #include "AnimNodes/AnimNode_PoseHandler.h"
 #include "Kismet2/CompilerResultsLog.h"
+#include "IAnimBlueprintNodeOverrideAssetsContext.h"
 
 UAnimGraphNode_PoseHandler::UAnimGraphNode_PoseHandler(const FObjectInitializer& ObjectInitializer)
 	:Super(ObjectInitializer)
@@ -48,6 +49,18 @@ void UAnimGraphNode_PoseHandler::SetAnimationAsset(UAnimationAsset* Asset)
 	}
 }
 
+void UAnimGraphNode_PoseHandler::OnOverrideAssets(IAnimBlueprintNodeOverrideAssetsContext& InContext) const
+{
+	if(InContext.GetAssets().Num() > 0)
+	{
+		if (UPoseAsset* PoseAsset = Cast<UPoseAsset>(InContext.GetAssets()[0]))
+		{
+			FAnimNode_PoseHandler& AnimNode = InContext.GetAnimNode<FAnimNode_PoseHandler>();
+			AnimNode.SetPoseAsset(PoseAsset);
+		}
+	}
+}
+
 void UAnimGraphNode_PoseHandler::PreloadRequiredAssets()
 {
 	PreloadObject(GetPoseHandlerNode()->PoseAsset);
diff --git a/Engine/Source/Editor/AnimGraph/Private/AnimGraphNode_RefPoseBase.cpp b/Engine/Source/Editor/AnimGraph/Private/AnimGraphNode_RefPoseBase.cpp
--- a/Engine/Source/Editor/AnimGraph/Private/AnimGraphNode_RefPoseBase.cpp
+++ b/Engine/Source/Editor/AnimGraph/Private/AnimGraphNode_RefPoseBase.cpp
@@ -17,7 +17,7 @@ FString UAnimGraphNode_RefPoseBase::GetNodeCategory() const
 
 FLinearColor UAnimGraphNode_RefPoseBase::GetNodeTitleColor() const
 {
-	if ( Node.RefPoseType == EIT_Additive )
+	if ( Node.GetRefPoseType() == EIT_Additive )
 	{
 		return FLinearColor(0.10f, 0.60f, 0.12f);
 	}
diff --git a/Engine/Source/Editor/AnimGraph/Private/AnimGraphNode_Root.cpp b/Engine/Source/Editor/AnimGraph/Private/AnimGraphNode_Root.cpp
--- a/Engine/Source/Editor/AnimGraph/Private/AnimGraphNode_Root.cpp
+++ b/Engine/Source/Editor/AnimGraph/Private/AnimGraphNode_Root.cpp
@@ -3,7 +3,6 @@
 #include "AnimGraphNode_Root.h"
 #include "GraphEditorSettings.h"
 #include "AnimBlueprintCompiler.h"
-#include "AnimBlueprintCompilerHandler_Base.h"
 #include "IAnimBlueprintCompilationContext.h"
 #include "IAnimBlueprintCopyTermDefaultsContext.h"
 
@@ -92,7 +91,7 @@ void UAnimGraphNode_Root::OnProcessDuringCompilation(IAnimBlueprintCompilationCo
 {
 	UAnimGraphNode_Root* TrueNode = InCompilationContext.GetMessageLog().FindSourceObjectTypeChecked<UAnimGraphNode_Root>(this);
 
-	Node.Name = TrueNode->GetGraph()->GetFName();
+	Node.SetName(TrueNode->GetGraph()->GetFName());
 }
 
 void UAnimGraphNode_Root::OnCopyTermDefaultsToDefaultObject(IAnimBlueprintCopyTermDefaultsContext& InCompilationContext, IAnimBlueprintNodeCopyTermDefaultsContext& InPerNodeContext, IAnimBlueprintGeneratedClassCompiledData& OutCompiledData)
@@ -100,7 +99,7 @@ void UAnimGraphNode_Root::OnCopyTermDefaultsToDefaultObject(IAnimBlueprintCopyTe
 	UAnimGraphNode_Root* TrueNode = InCompilationContext.GetMessageLog().FindSourceObjectTypeChecked<UAnimGraphNode_Root>(this);
 
 	FAnimNode_Root* DestinationNode = reinterpret_cast<FAnimNode_Root*>(InPerNodeContext.GetDestinationPtr());
-	DestinationNode->Name = TrueNode->GetGraph()->GetFName();
+	DestinationNode->SetName(TrueNode->GetGraph()->GetFName());
 }
 
 #undef LOCTEXT_NAMESPACE
diff --git a/Engine/Source/Editor/AnimGraph/Private/AnimGraphNode_RotationOffsetBlendSpace.cpp b/Engine/Source/Editor/AnimGraph/Private/AnimGraphNode_RotationOffsetBlendSpace.cpp
--- a/Engine/Source/Editor/AnimGraph/Private/AnimGraphNode_RotationOffsetBlendSpace.cpp
+++ b/Engine/Source/Editor/AnimGraph/Private/AnimGraphNode_RotationOffsetBlendSpace.cpp
@@ -34,7 +34,7 @@ FText UAnimGraphNode_RotationOffsetBlendSpace::GetTooltipText() const
 
 FText UAnimGraphNode_RotationOffsetBlendSpace::GetNodeTitle(ENodeTitleType::Type TitleType) const
 {
-	UBlendSpace* BlendSpaceToCheck = Node.BlendSpace;
+	UBlendSpace* BlendSpaceToCheck = Node.GetBlendSpace();
 	UEdGraphPin* BlendSpacePin = FindPin(GET_MEMBER_NAME_STRING_CHECKED(FAnimNode_RotationOffsetBlendSpace, BlendSpace));
 	if (BlendSpacePin != nullptr && BlendSpaceToCheck == nullptr)
 	{
@@ -81,7 +81,7 @@ void UAnimGraphNode_RotationOffsetBlendSpace::GetMenuActions(FBlueprintActionDat
 		static void SetNodeBlendSpace(UEdGraphNode* NewNode, bool /*bIsTemplateNode*/, TWeakObjectPtr<UBlendSpace> BlendSpace)
 		{
 			UAnimGraphNode_RotationOffsetBlendSpace* BlendSpaceNode = CastChecked<UAnimGraphNode_RotationOffsetBlendSpace>(NewNode);
-			BlendSpaceNode->Node.BlendSpace = BlendSpace.Get();
+			BlendSpaceNode->Node.SetBlendSpace(BlendSpace.Get());
 		}
 
 		static UBlueprintNodeSpawner* MakeBlendSpaceAction(TSubclassOf<UEdGraphNode> const NodeClass, const UBlendSpace* BlendSpace)
@@ -136,7 +136,7 @@ void UAnimGraphNode_RotationOffsetBlendSpace::GetMenuActions(FBlueprintActionDat
 FBlueprintNodeSignature UAnimGraphNode_RotationOffsetBlendSpace::GetSignature() const
 {
 	FBlueprintNodeSignature NodeSignature = Super::GetSignature();
-	NodeSignature.AddSubObject(Node.BlendSpace);
+	NodeSignature.AddSubObject(Node.GetBlendSpace());
 
 	return NodeSignature;
 }
@@ -145,7 +145,7 @@ void UAnimGraphNode_RotationOffsetBlendSpace::SetAnimationAsset(UAnimationAsset*
 {
 	if (UBlendSpace* BlendSpace = Cast<UBlendSpace>(Asset))
 	{
-		Node.BlendSpace = BlendSpace;
+		Node.SetBlendSpace(BlendSpace);
 	}
 }
 
@@ -153,7 +153,7 @@ void UAnimGraphNode_RotationOffsetBlendSpace::ValidateAnimNodeDuringCompilation(
 {
 	Super::ValidateAnimNodeDuringCompilation(ForSkeleton, MessageLog);
 
-	UBlendSpace* BlendSpaceToCheck = Node.BlendSpace;
+	UBlendSpace* BlendSpaceToCheck = Node.GetBlendSpace();
 	UEdGraphPin* BlendSpacePin = FindPin(GET_MEMBER_NAME_STRING_CHECKED(FAnimNode_RotationOffsetBlendSpace, BlendSpace));
 	if (BlendSpacePin != nullptr && BlendSpaceToCheck == nullptr)
 	{
@@ -208,7 +208,7 @@ void UAnimGraphNode_RotationOffsetBlendSpace::GetNodeContextMenuActions(UToolMen
 
 void UAnimGraphNode_RotationOffsetBlendSpace::GetAllAnimationSequencesReferred(TArray<UAnimationAsset*>& AnimationAssets) const
 {
-	if(Node.BlendSpace)
+	if(Node.GetBlendSpace())
 	{
 		HandleAnimReferenceCollection(Node.BlendSpace, AnimationAssets);
 	}
diff --git a/Engine/Source/Editor/AnimGraph/Private/AnimGraphNode_SaveCachedPose.cpp b/Engine/Source/Editor/AnimGraph/Private/AnimGraphNode_SaveCachedPose.cpp
--- a/Engine/Source/Editor/AnimGraph/Private/AnimGraphNode_SaveCachedPose.cpp
+++ b/Engine/Source/Editor/AnimGraph/Private/AnimGraphNode_SaveCachedPose.cpp
@@ -5,6 +5,7 @@
 
 #include "BlueprintNodeSpawner.h"
 #include "BlueprintActionDatabaseRegistrar.h"
+#include "AnimBlueprintExtension_CachedPose.h"
 
 /////////////////////////////////////////////////////
 // FCachedPoseNameValidator
@@ -116,4 +117,9 @@ TSharedPtr<class INameValidatorInterface> UAnimGraphNode_SaveCachedPose::MakeNam
 	return MakeShareable(new FCachedPoseNameValidator(GetBlueprint(), CacheName));
 }
 
+void UAnimGraphNode_SaveCachedPose::GetRequiredExtensions(TArray<TSubclassOf<UAnimBlueprintExtension>>& OutExtensions) const
+{
+	OutExtensions.Add(UAnimBlueprintExtension_CachedPose::StaticClass());
+}
+
 #undef LOCTEXT_NAMESPACE
diff --git a/Engine/Source/Editor/AnimGraph/Private/AnimGraphNode_SequenceEvaluator.cpp b/Engine/Source/Editor/AnimGraph/Private/AnimGraphNode_SequenceEvaluator.cpp
--- a/Engine/Source/Editor/AnimGraph/Private/AnimGraphNode_SequenceEvaluator.cpp
+++ b/Engine/Source/Editor/AnimGraph/Private/AnimGraphNode_SequenceEvaluator.cpp
@@ -8,6 +8,7 @@
 #include "Animation/AnimComposite.h"
 #include "Animation/AnimSequence.h"
 #include "AnimGraphNode_SequenceEvaluator.h"
+#include "IAnimBlueprintNodeOverrideAssetsContext.h"
 
 /////////////////////////////////////////////////////
 // UAnimGraphNode_SequenceEvaluator
@@ -21,7 +22,7 @@ UAnimGraphNode_SequenceEvaluator::UAnimGraphNode_SequenceEvaluator(const FObject
 
 void UAnimGraphNode_SequenceEvaluator::PreloadRequiredAssets()
 {
-	PreloadObject(Node.Sequence);
+	PreloadObject(Node.GetSequence());
 
 	Super::PreloadRequiredAssets();
 }
@@ -29,15 +30,12 @@ void UAnimGraphNode_SequenceEvaluator::PreloadRequiredAssets()
 void UAnimGraphNode_SequenceEvaluator::BakeDataDuringCompilation(class FCompilerResultsLog& MessageLog)
 {
 	UAnimBlueprint* AnimBlueprint = GetAnimBlueprint();
-	AnimBlueprint->FindOrAddGroup(SyncGroup.GroupName);
-	Node.GroupName = SyncGroup.GroupName;
-	Node.GroupRole = SyncGroup.GroupRole;
-	Node.Method = SyncGroup.Method;
+	AnimBlueprint->FindOrAddGroup(Node.GetGroupName());
 }
 
 void UAnimGraphNode_SequenceEvaluator::GetAllAnimationSequencesReferred(TArray<UAnimationAsset*>& AnimationAssets) const
 {
-	if(Node.Sequence)
+	if(Node.GetSequence())
 	{
 		HandleAnimReferenceCollection(Node.Sequence, AnimationAssets);
 	}
@@ -76,7 +74,8 @@ FText UAnimGraphNode_SequenceEvaluator::GetNodeTitleForSequence(ENodeTitleType::
 
 FText UAnimGraphNode_SequenceEvaluator::GetNodeTitle(ENodeTitleType::Type TitleType) const
 {
-	if (Node.Sequence == nullptr)
+	UAnimSequenceBase* Sequence = Node.GetSequence();
+	if (Sequence == nullptr)
 	{
 		// we may have a valid variable connected or default pin value
 		UEdGraphPin* SequencePin = FindPin(GET_MEMBER_NAME_STRING_CHECKED(FAnimNode_SequenceEvaluator, Sequence));
@@ -98,7 +97,7 @@ FText UAnimGraphNode_SequenceEvaluator::GetNodeTitle(ENodeTitleType::Type TitleT
 	//        then, we'll leave this optimization off
 	else //if (CachedNodeTitle.IsOutOfDate(this))
 	{
-		GetNodeTitleForSequence(TitleType, Node.Sequence);
+		GetNodeTitleForSequence(TitleType, Sequence);
 	}
 
 	return CachedNodeTitle;
@@ -113,7 +112,19 @@ void UAnimGraphNode_SequenceEvaluator::SetAnimationAsset(UAnimationAsset* Asset)
 {
 	if (UAnimSequenceBase* Seq =  Cast<UAnimSequence>(Asset))
 	{
-		Node.Sequence = Seq;
+		Node.SetSequence(Seq);
+	}
+}
+
+void UAnimGraphNode_SequenceEvaluator::OnOverrideAssets(IAnimBlueprintNodeOverrideAssetsContext& InContext) const
+{
+	if(InContext.GetAssets().Num() > 0)
+	{
+		if (UAnimSequenceBase* Sequence = Cast<UAnimSequenceBase>(InContext.GetAssets()[0]))
+		{
+			FAnimNode_SequenceEvaluator& AnimNode = InContext.GetAnimNode<FAnimNode_SequenceEvaluator>();
+			AnimNode.SetSequence(Sequence);
+		}
 	}
 }
 
@@ -121,7 +132,7 @@ void UAnimGraphNode_SequenceEvaluator::ValidateAnimNodeDuringCompilation(class U
 {
 	Super::ValidateAnimNodeDuringCompilation(ForSkeleton, MessageLog);
 
-	UAnimSequenceBase* SequenceToCheck = Node.Sequence;
+	UAnimSequenceBase* SequenceToCheck = Node.GetSequence();
 	UEdGraphPin* SequencePin = FindPin(GET_MEMBER_NAME_STRING_CHECKED(FAnimNode_SequenceEvaluator, Sequence));
 	if (SequencePin != nullptr && SequenceToCheck == nullptr)
 	{
@@ -177,14 +188,14 @@ bool UAnimGraphNode_SequenceEvaluator::DoesSupportTimeForTransitionGetter() cons
 
 UAnimationAsset* UAnimGraphNode_SequenceEvaluator::GetAnimationAsset() const 
 {
-	UAnimSequenceBase* Sequence = Node.Sequence;
+	UAnimSequenceBase* Sequence = Node.GetSequence();
 	UEdGraphPin* SequencePin = FindPin(GET_MEMBER_NAME_STRING_CHECKED(FAnimNode_SequenceEvaluator, Sequence));
 	if (SequencePin != nullptr && Sequence == nullptr)
 	{
 		Sequence = Cast<UAnimSequenceBase>(SequencePin->DefaultObject);
 	}
 
-	return Node.Sequence;
+	return Sequence;
 }
 
 const TCHAR* UAnimGraphNode_SequenceEvaluator::GetTimePropertyName() const 
diff --git a/Engine/Source/Editor/AnimGraph/Private/AnimGraphNode_SequencePlayer.cpp b/Engine/Source/Editor/AnimGraph/Private/AnimGraphNode_SequencePlayer.cpp
--- a/Engine/Source/Editor/AnimGraph/Private/AnimGraphNode_SequencePlayer.cpp
+++ b/Engine/Source/Editor/AnimGraph/Private/AnimGraphNode_SequencePlayer.cpp
@@ -17,6 +17,8 @@
 #include "Animation/AnimComposite.h"
 #include "Animation/AnimSequence.h"
 #include "Animation/AnimPoseSearchProvider.h"
+#include "UObject/UE5MainStreamObjectVersion.h"
+#include "IAnimBlueprintNodeOverrideAssetsContext.h"
 
 #define LOCTEXT_NAMESPACE "A3Nodes"
 
@@ -41,7 +43,7 @@ public:
 	virtual UEdGraphNode* PerformAction(class UEdGraph* ParentGraph, UEdGraphPin* FromPin, const FVector2D Location, bool bSelectNewNode = true) override
 	{
 		UAnimGraphNode_SequencePlayer* SpawnedNode = CastChecked<UAnimGraphNode_SequencePlayer>(FEdGraphSchemaAction_K2NewNode::PerformAction(ParentGraph, FromPin, Location, bSelectNewNode));
-		SpawnedNode->Node.Sequence = Cast<UAnimSequence>(AssetInfo.GetAsset());
+		SpawnedNode->Node.SetSequence(Cast<UAnimSequence>(AssetInfo.GetAsset()));
 
 		return SpawnedNode;
 	}
@@ -55,16 +57,29 @@ UAnimGraphNode_SequencePlayer::UAnimGraphNode_SequencePlayer(const FObjectInitia
 {
 }
 
+void UAnimGraphNode_SequencePlayer::Serialize(FArchive& Ar)
+{
+	Super::Serialize(Ar);
+
+	Ar.UsingCustomVersion(FUE5MainStreamObjectVersion::GUID);
+
+	if(Ar.IsLoading() && Ar.CustomVer(FUE5MainStreamObjectVersion::GUID) < FUE5MainStreamObjectVersion::AnimNodeConstantDataRefactorPhase0)
+	{
+		Node.PlayRateScaleBiasClampConstants.CopyFromLegacy(Node.PlayRateScaleBiasClamp_DEPRECATED);
+	}
+}
+
 void UAnimGraphNode_SequencePlayer::PreloadRequiredAssets()
 {
-	PreloadObject(Node.Sequence);
+	PreloadObject(Node.GetSequence());
 
 	Super::PreloadRequiredAssets();
 }
 
 FLinearColor UAnimGraphNode_SequencePlayer::GetNodeTitleColor() const
 {
-	if ((Node.Sequence != NULL) && Node.Sequence->IsValidAdditive())
+	UAnimSequenceBase* Sequence = Node.GetSequence();
+	if ((Sequence != NULL) && Sequence->IsValidAdditive())
 	{
 		return FLinearColor(0.10f, 0.60f, 0.12f);
 	}
@@ -76,13 +91,14 @@ FLinearColor UAnimGraphNode_SequencePlayer::GetNodeTitleColor() const
 
 FText UAnimGraphNode_SequencePlayer::GetTooltipText() const
 {
-	if (!Node.Sequence)
+	UAnimSequenceBase* Sequence = Node.GetSequence();
+	if (!Sequence)
 	{
 		return FText();
 	}
 
-	const bool bAdditive = Node.Sequence->IsValidAdditive();
-	return GetTitleGivenAssetInfo(FText::FromString(Node.Sequence->GetPathName()), bAdditive);
+	const bool bAdditive = Sequence->IsValidAdditive();
+	return GetTitleGivenAssetInfo(FText::FromString(Sequence->GetPathName()), bAdditive);
 }
 
 FText UAnimGraphNode_SequencePlayer::GetNodeTitleForSequence(ENodeTitleType::Type TitleType, UAnimSequenceBase* InSequence) const
@@ -90,9 +106,9 @@ FText UAnimGraphNode_SequencePlayer::GetNodeTitleForSequence(ENodeTitleType::Typ
 	const bool bAdditive = InSequence->IsValidAdditive();
 	const FText BasicTitle = GetTitleGivenAssetInfo(FText::FromName(InSequence->GetFName()), bAdditive);
 
-	if(SyncGroup.Method == EAnimSyncMethod::SyncGroup)
+	if(Node.GetGroupMethod() == EAnimSyncMethod::SyncGroup)
 	{
-		const FText SyncGroupName = FText::FromName(SyncGroup.GroupName);
+		const FText SyncGroupName = FText::FromName(Node.GetGroupName());
 
 		FFormatNamedArguments Args;
 		Args.Add(TEXT("Title"), BasicTitle);
@@ -107,7 +123,7 @@ FText UAnimGraphNode_SequencePlayer::GetNodeTitleForSequence(ENodeTitleType::Typ
 			return FText::Format(LOCTEXT("SequenceNodeGroupWithSubtitleList", "{Title} (Sync group {SyncGroup})"), Args);
 		}
 	}
-	else if(SyncGroup.Method == EAnimSyncMethod::Graph)
+	else if(Node.GetGroupMethod() == EAnimSyncMethod::Graph)
 	{
 		FFormatNamedArguments Args;
 		Args.Add(TEXT("Title"), BasicTitle);
@@ -153,7 +169,8 @@ FText UAnimGraphNode_SequencePlayer::GetNodeTitleForSequence(ENodeTitleType::Typ
 
 FText UAnimGraphNode_SequencePlayer::GetNodeTitle(ENodeTitleType::Type TitleType) const
 {
-	if (Node.Sequence == nullptr)
+	UAnimSequenceBase* Sequence = Node.GetSequence();
+	if (Sequence == nullptr)
 	{
 		// we may have a valid variable connected or default pin value
 		UEdGraphPin* SequencePin = FindPin(GET_MEMBER_NAME_STRING_CHECKED(FAnimNode_SequencePlayer, Sequence));
@@ -172,7 +189,7 @@ FText UAnimGraphNode_SequencePlayer::GetNodeTitle(ENodeTitleType::Type TitleType
 	}
 	else
 	{
-		return GetNodeTitleForSequence(TitleType, Node.Sequence);
+		return GetNodeTitleForSequence(TitleType, Sequence);
 	}
 }
 
@@ -201,7 +218,7 @@ void UAnimGraphNode_SequencePlayer::GetMenuActions(FBlueprintActionDatabaseRegis
 	auto LoadedAssetSetup = [](UEdGraphNode* NewNode, bool /*bIsTemplateNode*/, TWeakObjectPtr<UAnimSequence> SequencePtr)
 	{
 		UAnimGraphNode_SequencePlayer* SequencePlayerNode = CastChecked<UAnimGraphNode_SequencePlayer>(NewNode);
-		SequencePlayerNode->Node.Sequence = SequencePtr.Get();
+		SequencePlayerNode->Node.SetSequence(SequencePtr.Get());
 	};
 
 	auto UnloadedAssetSetup = [](UEdGraphNode* NewNode, bool bIsTemplateNode, const FAssetData AssetData)
@@ -215,7 +232,7 @@ void UAnimGraphNode_SequencePlayer::GetMenuActions(FBlueprintActionDatabaseRegis
 		{
 			UAnimSequence* Sequence = Cast<UAnimSequence>(AssetData.GetAsset());
 			check(Sequence != nullptr);
-			SequencePlayerNode->Node.Sequence = Sequence;
+			SequencePlayerNode->Node.SetSequence(Sequence);
 		}
 	};	
 
@@ -311,9 +328,10 @@ bool UAnimGraphNode_SequencePlayer::IsActionFilteredOut(class FBlueprintActionFi
 		UAnimBlueprint* AnimBlueprint = Cast<UAnimBlueprint>(Blueprint);
 		if (AnimBlueprint && AnimBlueprint->TargetSkeleton)
 		{
-			if(Node.Sequence)
+			UAnimSequenceBase* Sequence = Node.GetSequence();
+			if(Sequence)
 			{
-				if (!AnimBlueprint->TargetSkeleton->IsCompatible(Node.Sequence->GetSkeleton()))
+				if (!AnimBlueprint->TargetSkeleton->IsCompatible(Sequence->GetSkeleton()))
 				{
 					// Asset does not use a compatible skeleton with the Blueprint, cannot use
 					bIsFilteredOut = true;
@@ -355,7 +373,7 @@ void UAnimGraphNode_SequencePlayer::ValidateAnimNodeDuringCompilation(class USke
 {
 	Super::ValidateAnimNodeDuringCompilation(ForSkeleton, MessageLog);
 
-	UAnimSequenceBase* SequenceToCheck = Node.Sequence;
+	UAnimSequenceBase* SequenceToCheck = Node.GetSequence();
 	UEdGraphPin* SequencePin = FindPin(GET_MEMBER_NAME_STRING_CHECKED(FAnimNode_SequencePlayer, Sequence));
 	if (SequencePin != nullptr && SequenceToCheck == nullptr)
 	{
@@ -412,22 +430,31 @@ void UAnimGraphNode_SequencePlayer::SetAnimationAsset(UAnimationAsset* Asset)
 {
 	if (UAnimSequenceBase* Seq = Cast<UAnimSequenceBase>(Asset))
 	{
-		Node.Sequence = Seq;
+		Node.SetSequence(Seq);
+	}
+}
+
+void UAnimGraphNode_SequencePlayer::OnOverrideAssets(IAnimBlueprintNodeOverrideAssetsContext& InContext) const
+{
+	if(InContext.GetAssets().Num() > 0)
+	{
+		if (UAnimSequenceBase* Sequence = Cast<UAnimSequenceBase>(InContext.GetAssets()[0]))
+		{
+			FAnimNode_SequencePlayer& AnimNode = InContext.GetAnimNode<FAnimNode_SequencePlayer>();
+			AnimNode.SetSequence(Sequence);
+		}
 	}
 }
 
 void UAnimGraphNode_SequencePlayer::BakeDataDuringCompilation(class FCompilerResultsLog& MessageLog)
 {
 	UAnimBlueprint* AnimBlueprint = GetAnimBlueprint();
-	AnimBlueprint->FindOrAddGroup(SyncGroup.GroupName);
-	Node.GroupName = SyncGroup.GroupName;
-	Node.GroupRole = SyncGroup.GroupRole;
-	Node.Method = SyncGroup.Method;
+	AnimBlueprint->FindOrAddGroup(Node.GetGroupName());
 }
 
 void UAnimGraphNode_SequencePlayer::GetAllAnimationSequencesReferred(TArray<UAnimationAsset*>& AnimationAssets) const
 {
-	if(Node.Sequence)
+	if(Node.GetSequence())
 	{
 		HandleAnimReferenceCollection(Node.Sequence, AnimationAssets);
 	}
@@ -446,7 +473,7 @@ bool UAnimGraphNode_SequencePlayer::DoesSupportTimeForTransitionGetter() const
 
 UAnimationAsset* UAnimGraphNode_SequencePlayer::GetAnimationAsset() const 
 {
-	UAnimSequenceBase* Sequence = Node.Sequence;
+	UAnimSequenceBase* Sequence = Node.GetSequence();
 	UEdGraphPin* SequencePin = FindPin(GET_MEMBER_NAME_STRING_CHECKED(FAnimNode_SequencePlayer, Sequence));
 	if (SequencePin != nullptr && Sequence == nullptr)
 	{
@@ -488,11 +515,11 @@ void UAnimGraphNode_SequencePlayer::CustomizePinData(UEdGraphPin* Pin, FName Sou
 			UEdGraphPin* PlayRateBasisPin = FindPin(GET_MEMBER_NAME_STRING_CHECKED(FAnimNode_SequencePlayer, PlayRateBasis));
 			if (!PlayRateBasisPin || PlayRateBasisPin->bHidden)
 			{
-				if (Node.PlayRateBasis != 1.f)
+				if (Node.GetPlayRateBasis() != 1.f)
 				{
 					FFormatNamedArguments Args;
 					Args.Add(TEXT("PinFriendlyName"), Pin->PinFriendlyName);
-					Args.Add(TEXT("PlayRateBasis"), FText::AsNumber(Node.PlayRateBasis));
+					Args.Add(TEXT("PlayRateBasis"), FText::AsNumber(Node.GetPlayRateBasis()));
 					Pin->PinFriendlyName = FText::Format(LOCTEXT("FAnimNode_SequencePlayer_PlayRateBasis_Value", "({PinFriendlyName} / {PlayRateBasis})"), Args);
 				}
 			}
@@ -503,7 +530,7 @@ void UAnimGraphNode_SequencePlayer::CustomizePinData(UEdGraphPin* Pin, FName Sou
 				Pin->PinFriendlyName = FText::Format(LOCTEXT("FAnimNode_SequencePlayer_PlayRateBasis_Name", "({PinFriendlyName} / PlayRateBasis)"), Args);
 			}
 
-			Pin->PinFriendlyName = Node.PlayRateScaleBiasClamp.GetFriendlyName(Pin->PinFriendlyName);
+			Pin->PinFriendlyName = Node.GetPlayRateScaleBiasClampConstants().GetFriendlyName(Pin->PinFriendlyName);
 		}
 	}
 }
@@ -514,17 +541,17 @@ void UAnimGraphNode_SequencePlayer::PostEditChangeProperty(struct FPropertyChang
 
 	// Reconstruct node to show updates to PinFriendlyNames.
 	if ((PropertyName == GET_MEMBER_NAME_STRING_CHECKED(FAnimNode_SequencePlayer, PlayRateBasis))
-		|| (PropertyName == GET_MEMBER_NAME_STRING_CHECKED(FInputScaleBiasClamp, bMapRange))
+		|| (PropertyName == GET_MEMBER_NAME_STRING_CHECKED(FInputScaleBiasClampConstants, bMapRange))
 		|| (PropertyName == GET_MEMBER_NAME_STRING_CHECKED(FInputRange, Min))
 		|| (PropertyName == GET_MEMBER_NAME_STRING_CHECKED(FInputRange, Max))
-		|| (PropertyName == GET_MEMBER_NAME_STRING_CHECKED(FInputScaleBiasClamp, Scale))
-		|| (PropertyName == GET_MEMBER_NAME_STRING_CHECKED(FInputScaleBiasClamp, Bias))
-		|| (PropertyName == GET_MEMBER_NAME_STRING_CHECKED(FInputScaleBiasClamp, bClampResult))
-		|| (PropertyName == GET_MEMBER_NAME_STRING_CHECKED(FInputScaleBiasClamp, ClampMin))
-		|| (PropertyName == GET_MEMBER_NAME_STRING_CHECKED(FInputScaleBiasClamp, ClampMax))
-		|| (PropertyName == GET_MEMBER_NAME_STRING_CHECKED(FInputScaleBiasClamp, bInterpResult))
-		|| (PropertyName == GET_MEMBER_NAME_STRING_CHECKED(FInputScaleBiasClamp, InterpSpeedIncreasing))
-		|| (PropertyName == GET_MEMBER_NAME_STRING_CHECKED(FInputScaleBiasClamp, InterpSpeedDecreasing)))
+		|| (PropertyName == GET_MEMBER_NAME_STRING_CHECKED(FInputScaleBiasClampConstants, Scale))
+		|| (PropertyName == GET_MEMBER_NAME_STRING_CHECKED(FInputScaleBiasClampConstants, Bias))
+		|| (PropertyName == GET_MEMBER_NAME_STRING_CHECKED(FInputScaleBiasClampConstants, bClampResult))
+		|| (PropertyName == GET_MEMBER_NAME_STRING_CHECKED(FInputScaleBiasClampConstants, ClampMin))
+		|| (PropertyName == GET_MEMBER_NAME_STRING_CHECKED(FInputScaleBiasClampConstants, ClampMax))
+		|| (PropertyName == GET_MEMBER_NAME_STRING_CHECKED(FInputScaleBiasClampConstants, bInterpResult))
+		|| (PropertyName == GET_MEMBER_NAME_STRING_CHECKED(FInputScaleBiasClampConstants, InterpSpeedIncreasing))
+		|| (PropertyName == GET_MEMBER_NAME_STRING_CHECKED(FInputScaleBiasClampConstants, InterpSpeedDecreasing)))
 	{
 		ReconstructNode();
 	}
diff --git a/Engine/Source/Editor/AnimGraph/Private/AnimGraphNode_StateMachineBase.cpp b/Engine/Source/Editor/AnimGraph/Private/AnimGraphNode_StateMachineBase.cpp
--- a/Engine/Source/Editor/AnimGraph/Private/AnimGraphNode_StateMachineBase.cpp
+++ b/Engine/Source/Editor/AnimGraph/Private/AnimGraphNode_StateMachineBase.cpp
@@ -25,7 +25,7 @@
 #include "AnimGraphNode_LinkedAnimLayer.h"
 #include "AnimGraphNode_TransitionPoseEvaluator.h"
 #include "AnimGraphNode_CustomTransitionResult.h"
-#include "AnimBlueprintCompilerHandler_StateMachine.h"
+#include "AnimBlueprintExtension_StateMachine.h"
 #include "IAnimBlueprintGeneratedClassCompiledData.h"
 #include "IAnimBlueprintCompilationContext.h"
 #include "Animation/AnimNode_Inertialization.h"
@@ -114,6 +114,8 @@ FString UAnimGraphNode_StateMachineBase::GetNodeCategory() const
 
 void UAnimGraphNode_StateMachineBase::PostPlacedNewNode()
 {
+	Super::PostPlacedNewNode();
+
 	// Create a new animation graph
 	check(EditorStateMachineGraph == NULL);
 	EditorStateMachineGraph = CastChecked<UAnimationStateMachineGraph>(FBlueprintEditorUtils::CreateNewGraph(this, NAME_None, UAnimationStateMachineGraph::StaticClass(), UAnimationStateMachineSchema::StaticClass()));
@@ -363,8 +365,8 @@ void UAnimGraphNode_StateMachineBase::OnProcessDuringCompilation(IAnimBlueprintC
 		}
 	};
 	
-	FAnimBlueprintCompilerHandler_StateMachine* CompilerHandler = InCompilationContext.GetHandler<FAnimBlueprintCompilerHandler_StateMachine>("AnimBlueprintCompilerHandler_StateMachine");
-	check(CompilerHandler);
+	UAnimBlueprintExtension_StateMachine* Extension = UAnimBlueprintExtension::GetExtension<UAnimBlueprintExtension_StateMachine>(GetAnimBlueprint());
+	check(Extension);
 
 	if (EditorStateMachineGraph == NULL)
 	{
@@ -431,7 +433,7 @@ void UAnimGraphNode_StateMachineBase::OnProcessDuringCompilation(IAnimBlueprintC
 				{
 					InCompilationContext.ValidateGraphIsWellFormed(StateNode->BoundGraph);
 
-					BakedState.StateRootNodeIndex = CompilerHandler->ExpandGraphAndProcessNodes(StateNode->BoundGraph, AnimGraphResultNode, InCompilationContext, OutCompiledData);
+					BakedState.StateRootNodeIndex = Extension->ExpandGraphAndProcessNodes(StateNode->BoundGraph, AnimGraphResultNode, InCompilationContext, OutCompiledData);
 
 					// See if the state consists of a single sequence player node, and remember the index if so
 					for (UEdGraphPin* TestPin : AnimGraphResultNode->Pins)
@@ -473,7 +475,7 @@ void UAnimGraphNode_StateMachineBase::OnProcessDuringCompilation(IAnimBlueprintC
 			{
 				if (UAnimGraphNode_TransitionResult* EntryRuleResultNode = CastChecked<UAnimationTransitionGraph>(ConduitNode->BoundGraph)->GetResultNode())
 				{
-					BakedState.EntryRuleNodeIndex = CompilerHandler->ExpandGraphAndProcessNodes(ConduitNode->BoundGraph, EntryRuleResultNode, InCompilationContext, OutCompiledData);
+					BakedState.EntryRuleNodeIndex = Extension->ExpandGraphAndProcessNodes(ConduitNode->BoundGraph, EntryRuleResultNode, InCompilationContext, OutCompiledData);
 				}
 			}
 
@@ -650,7 +652,7 @@ void UAnimGraphNode_StateMachineBase::OnProcessDuringCompilation(IAnimBlueprintC
 				}
 				else
 				{
-					Rule.CanTakeDelegateIndex = CompilerHandler->ExpandGraphAndProcessNodes(TransitionNode->BoundGraph, TransitionResultNode, InCompilationContext, OutCompiledData, TransitionNode);
+					Rule.CanTakeDelegateIndex = Extension->ExpandGraphAndProcessNodes(TransitionNode->BoundGraph, TransitionResultNode, InCompilationContext, OutCompiledData, TransitionNode);
 					AlreadyMergedTransitionList.Add(TransitionResultNode, Rule.CanTakeDelegateIndex);
 				}
 			}
@@ -671,7 +673,7 @@ void UAnimGraphNode_StateMachineBase::OnProcessDuringCompilation(IAnimBlueprintC
 				TArray<UEdGraphNode*> ClonedNodes;
 				if (CustomTransitionGraph->GetResultNode())
 				{
-					Rule.CustomResultNodeIndex = CompilerHandler->ExpandGraphAndProcessNodes(TransitionNode->CustomTransitionGraph, CustomTransitionGraph->GetResultNode(), InCompilationContext, OutCompiledData, nullptr, &ClonedNodes);
+					Rule.CustomResultNodeIndex = Extension->ExpandGraphAndProcessNodes(TransitionNode->CustomTransitionGraph, CustomTransitionGraph->GetResultNode(), InCompilationContext, OutCompiledData, nullptr, &ClonedNodes);
 				}
 
 				// Find all the pose evaluators used in this transition, save handles to them because we need to populate some pose data before executing
@@ -714,4 +716,9 @@ void UAnimGraphNode_StateMachineBase::GetOutputLinkAttributes(FNodeAttributeArra
 	}
 }
 
+void UAnimGraphNode_StateMachineBase::GetRequiredExtensions(TArray<TSubclassOf<UAnimBlueprintExtension>>& OutExtensions) const
+{
+	OutExtensions.Add(UAnimBlueprintExtension_StateMachine::StaticClass());
+}
+
 #undef LOCTEXT_NAMESPACE
diff --git a/Engine/Source/Editor/AnimGraph/Private/AnimGraphNode_StateResult.cpp b/Engine/Source/Editor/AnimGraph/Private/AnimGraphNode_StateResult.cpp
--- a/Engine/Source/Editor/AnimGraph/Private/AnimGraphNode_StateResult.cpp
+++ b/Engine/Source/Editor/AnimGraph/Private/AnimGraphNode_StateResult.cpp
@@ -2,8 +2,6 @@
 
 #include "AnimGraphNode_StateResult.h"
 #include "GraphEditorSettings.h"
-#include "AnimBlueprintCompiler.h"
-#include "AnimBlueprintCompilerHandler_Base.h"
 #include "IAnimBlueprintCompilationContext.h"
 
 #define LOCTEXT_NAMESPACE "A3Nodes"
@@ -50,7 +48,7 @@ void UAnimGraphNode_StateResult::OnProcessDuringCompilation(IAnimBlueprintCompil
 {
 	UAnimGraphNode_StateResult* TrueNode = InCompilationContext.GetMessageLog().FindSourceObjectTypeChecked<UAnimGraphNode_StateResult>(this);
 
-	Node.Name = TrueNode->GetGraph()->GetFName();
+	Node.SetName(TrueNode->GetGraph()->GetFName());
 }
 
 #undef LOCTEXT_NAMESPACE
diff --git a/Engine/Source/Editor/AnimGraph/Private/AnimGraphNode_UseCachedPose.cpp b/Engine/Source/Editor/AnimGraph/Private/AnimGraphNode_UseCachedPose.cpp
--- a/Engine/Source/Editor/AnimGraph/Private/AnimGraphNode_UseCachedPose.cpp
+++ b/Engine/Source/Editor/AnimGraph/Private/AnimGraphNode_UseCachedPose.cpp
@@ -8,7 +8,7 @@
 #include "BlueprintActionFilter.h"
 #include "BlueprintNodeSpawner.h"
 #include "BlueprintActionDatabaseRegistrar.h"
-#include "AnimBlueprintCompilerHandler_CachedPose.h"
+#include "AnimBlueprintExtension_CachedPose.h"
 #include "IAnimBlueprintCompilationContext.h"
 #include "Kismet2/KismetEditorUtilities.h"
 #include "BlueprintEditorModule.h"
@@ -200,15 +200,15 @@ bool UAnimGraphNode_UseCachedPose::IsActionFilteredOut(class FBlueprintActionFil
 
 void UAnimGraphNode_UseCachedPose::OnProcessDuringCompilation(IAnimBlueprintCompilationContext& InCompilationContext, IAnimBlueprintGeneratedClassCompiledData& OutCompiledData)
 {	
-	FAnimBlueprintCompilerHandler_CachedPose* CompilerHandler = InCompilationContext.GetHandler<FAnimBlueprintCompilerHandler_CachedPose>("AnimBlueprintCompilerHandler_CachedPose");
-	check(CompilerHandler);
+	UAnimBlueprintExtension_CachedPose* CompilerExtension = UAnimBlueprintExtension::GetExtension<UAnimBlueprintExtension_CachedPose>(GetAnimBlueprint());
+	check(CompilerExtension);
 
 	bool bSuccessful = false;
 
 	// Link to the saved cached pose
 	if(SaveCachedPoseNode.IsValid())
 	{
-		if (UAnimGraphNode_SaveCachedPose* AssociatedSaveNode = CompilerHandler->GetSaveCachedPoseNodes().FindRef(SaveCachedPoseNode->CacheName))
+		if (UAnimGraphNode_SaveCachedPose* AssociatedSaveNode = CompilerExtension->GetSaveCachedPoseNodes().FindRef(SaveCachedPoseNode->CacheName))
 		{
 			FStructProperty* LinkProperty = FindFProperty<FStructProperty>(FAnimNode_UseCachedPose::StaticStruct(), TEXT("LinkToCachingNode"));
 			check(LinkProperty);
@@ -233,4 +233,9 @@ void UAnimGraphNode_UseCachedPose::OnProcessDuringCompilation(IAnimBlueprintComp
 	}
 }
 
+void UAnimGraphNode_UseCachedPose::GetRequiredExtensions(TArray<TSubclassOf<UAnimBlueprintExtension>>& OutExtensions) const
+{
+	OutExtensions.Add(UAnimBlueprintExtension_CachedPose::StaticClass());
+}
+
 #undef LOCTEXT_NAMESPACE
diff --git a/Engine/Source/Editor/AnimGraph/Private/AnimationGraphSchema.cpp b/Engine/Source/Editor/AnimGraph/Private/AnimationGraphSchema.cpp
--- a/Engine/Source/Editor/AnimGraph/Private/AnimationGraphSchema.cpp
+++ b/Engine/Source/Editor/AnimGraph/Private/AnimationGraphSchema.cpp
@@ -293,7 +293,7 @@ void UAnimationGraphSchema::CreateFunctionGraphTerminators(UEdGraph& Graph, UCla
 			Graph.GetNodesOfClass<UAnimGraphNode_Root>(RootNodes);
 
 			check(RootNodes.Num() == 1);
-			RootNodes[0]->Node.Group = *FObjectEditorUtils::GetCategoryText(InterfaceToImplement).ToString();
+			RootNodes[0]->Node.SetGroup(*FObjectEditorUtils::GetCategoryText(InterfaceToImplement).ToString());
 
 			int32 CurrentPoseIndex = 0;
 			for (TFieldIterator<FProperty> PropIt(InterfaceToImplement); PropIt && (PropIt->PropertyFlags & CPF_Parm); ++PropIt)
@@ -730,7 +730,7 @@ void UAnimationGraphSchema::ConformAnimGraphToInterface(UBlueprint* InBlueprint,
 		InGraph.GetNodesOfClass<UAnimGraphNode_Root>(RootNodes);
 
 		check(RootNodes.Num() == 1);
-		RootNodes[0]->Node.Group = *FObjectEditorUtils::GetCategoryText(InFunction).ToString();
+		RootNodes[0]->Node.SetGroup(*FObjectEditorUtils::GetCategoryText(InFunction).ToString());
 
 		TArray<UAnimGraphNode_LinkedInputPose*> LinkedInputPoseNodes;
 		InGraph.GetNodesOfClass<UAnimGraphNode_LinkedInputPose>(LinkedInputPoseNodes);
diff --git a/Engine/Source/Editor/AnimGraph/Public/AnimBlueprintExtension_PropertyAccess.h b/Engine/Source/Editor/AnimGraph/Public/AnimBlueprintExtension_PropertyAccess.h
new file mode 100644
--- /dev/null
+++ b/Engine/Source/Editor/AnimGraph/Public/AnimBlueprintExtension_PropertyAccess.h
@@ -0,0 +1,58 @@
+// Copyright Epic Games, Inc. All Rights Reserved.
+
+#pragma once
+
+#include "CoreMinimal.h"
+#include "AnimBlueprintExtension.h"
+#include "Animation/AnimSubsystem_PropertyAccess.h"
+#include "AnimBlueprintExtension_PropertyAccess.generated.h"
+
+class IAnimBlueprintCompilerCreationContext;
+class IAnimBlueprintCompilationBracketContext;
+class IAnimBlueprintGeneratedClassCompiledData;
+class UAnimBlueprintExtension_PropertyAccess;
+class IPropertyAccessLibraryCompiler;
+enum class EPropertyAccessBatchType : uint8;
+
+// Delegate called when the library is compiled (whether successfully or not)
+DECLARE_MULTICAST_DELEGATE_TwoParams(FOnPostLibraryCompiled, IAnimBlueprintCompilationBracketContext& /*InCompilationContext*/, IAnimBlueprintGeneratedClassCompiledData& /*OutCompiledData*/)
+
+UCLASS()
+class ANIMGRAPH_API UAnimBlueprintExtension_PropertyAccess : public UAnimBlueprintExtension
+{
+	GENERATED_BODY()
+
+	friend class UAnimBlueprintExtension_Base;
+
+public:
+	// Add a copy to the property access library we are compiling
+	// @return an integer handle to the pending copy. This can be resolved to a true copy index by calling MapCopyIndex
+	int32 AddCopy(TArrayView<FString> InSourcePath, TArrayView<FString> InDestPath, EPropertyAccessBatchType InBatchType, UObject* InObject = nullptr);
+
+	// Delegate called when the library is compiled (whether successfully or not)
+	FSimpleMulticastDelegate& OnPreLibraryCompiled() { return OnPreLibraryCompiledDelegate; }
+
+	// Delegate called when the library is compiled (whether successfully or not)
+	FOnPostLibraryCompiled& OnPostLibraryCompiled() { return OnPostLibraryCompiledDelegate; }
+
+	// Maps the initial integer copy handle to a true handle, post compilation
+	int32 MapCopyIndex(int32 InIndex) const;
+
+private:
+	// UAnimBlueprintExtension interface
+	virtual void HandleStartCompilingClass(const UClass* InClass, IAnimBlueprintCompilationBracketContext& InCompilationContext, IAnimBlueprintGeneratedClassCompiledData& OutCompiledData) override;
+	virtual void HandleFinishCompilingClass(const UClass* InClass, IAnimBlueprintCompilationBracketContext& InCompilationContext, IAnimBlueprintGeneratedClassCompiledData& OutCompiledData) override;
+
+private:	
+	// Property access library compiler
+	TUniquePtr<IPropertyAccessLibraryCompiler> PropertyAccessLibraryCompiler;
+
+	// Delegate called before the library is compiled
+	FSimpleMulticastDelegate OnPreLibraryCompiledDelegate;
+
+	// Delegate called when the library is compiled (whether successfully or not)
+	FOnPostLibraryCompiled OnPostLibraryCompiledDelegate;
+
+	UPROPERTY()
+	FAnimSubsystem_PropertyAccess Subsystem;
+};
\ No newline at end of file
diff --git a/Engine/Source/Editor/AnimGraph/Public/AnimGraphNode_AimOffsetLookAt.h b/Engine/Source/Editor/AnimGraph/Public/AnimGraphNode_AimOffsetLookAt.h
--- a/Engine/Source/Editor/AnimGraph/Public/AnimGraphNode_AimOffsetLookAt.h
+++ b/Engine/Source/Editor/AnimGraph/Public/AnimGraphNode_AimOffsetLookAt.h
@@ -29,7 +29,7 @@ class UAnimGraphNode_AimOffsetLookAt : public UAnimGraphNode_BlendSpaceBase
 
 	// UAnimGraphNode_Base interface
 	virtual void ValidateAnimNodeDuringCompilation(class USkeleton* ForSkeleton, class FCompilerResultsLog& MessageLog) override;
-	virtual UAnimationAsset* GetAnimationAsset() const override { return Node.BlendSpace; }
+	virtual UAnimationAsset* GetAnimationAsset() const override { return Node.GetBlendSpace(); }
 	virtual void GetAllAnimationSequencesReferred(TArray<UAnimationAsset*>& AnimationAssets) const override;
 	virtual void ReplaceReferredAnimations(const TMap<UAnimationAsset*, UAnimationAsset*>& AnimAssetReplacementMap) override;
 	ANIMGRAPH_API virtual void CustomizePinData(UEdGraphPin* Pin, FName SourcePropertyName, int32 ArrayIndex) const override;
diff --git a/Engine/Source/Editor/AnimGraph/Public/AnimGraphNode_AssetPlayerBase.h b/Engine/Source/Editor/AnimGraph/Public/AnimGraphNode_AssetPlayerBase.h
--- a/Engine/Source/Editor/AnimGraph/Public/AnimGraphNode_AssetPlayerBase.h
+++ b/Engine/Source/Editor/AnimGraph/Public/AnimGraphNode_AssetPlayerBase.h
@@ -20,9 +20,9 @@ class ANIMGRAPH_API UAnimGraphNode_AssetPlayerBase : public UAnimGraphNode_Base
 {
 	GENERATED_BODY()
 public:
-	// Sync group settings for this player.  Sync groups keep related animations with different lengths synchronized.
-	UPROPERTY(EditAnywhere, Category = Settings)
-	FAnimationGroupReference SyncGroup;
+	// Deprecated - sync group data is held on the contained FAnimNode_Base
+	UPROPERTY()
+	FAnimationGroupReference SyncGroup_DEPRECATED;
 
 	/** UObject interface */
 	void Serialize(FArchive& Ar) override;
diff --git a/Engine/Source/Editor/AnimGraph/Public/AnimGraphNode_Base.h b/Engine/Source/Editor/AnimGraph/Public/AnimGraphNode_Base.h
--- a/Engine/Source/Editor/AnimGraph/Public/AnimGraphNode_Base.h
+++ b/Engine/Source/Editor/AnimGraph/Public/AnimGraphNode_Base.h
@@ -27,6 +27,8 @@ class IAnimBlueprintGeneratedClassCompiledData;
 class IAnimBlueprintCompilationContext;
 class IAnimBlueprintCopyTermDefaultsContext;
 class IAnimBlueprintNodeCopyTermDefaultsContext;
+class IAnimBlueprintNodeOverrideAssetsContext;
+class UAnimBlueprintExtension;
 
 struct FPoseLinkMappingRecord
 {
@@ -205,6 +207,8 @@ class ANIMGRAPH_API UAnimGraphNode_Base : public UK2Node
 	virtual void AddSearchMetaDataInfo(TArray<struct FSearchTagDataPair>& OutTaggedMetaData) const override;
 	virtual void PinConnectionListChanged(UEdGraphPin* Pin) override;
 	virtual void AutowireNewNode(UEdGraphPin* FromPin) override;
+	virtual void PostPlacedNewNode() override;
+	virtual void DestroyNode() override;
 	// End of UEdGraphNode interface
 
 	// UK2Node interface
@@ -352,6 +356,14 @@ class ANIMGRAPH_API UAnimGraphNode_Base : public UK2Node
 	 */
 	bool IsPinExposedAndLinked(const FString& InPinName, const EEdGraphPinDirection Direction = EGPD_MAX) const;
 
+	/**
+	 * Helper function to check whether a pin is valid and bound via property access
+	 * @param	InPinName		The name of the pin @see UEdGraphNode::FindPin
+	 * @param	InPinDirection	The direction of the pin we are looking for. If this is EGPD_MAX, all directions are considered
+	 * @return true if the pin is present and bound
+	 */
+	bool IsPinExposedAndBound(const FString& InPinName, const EEdGraphPinDirection InDirection = EGPD_MAX) const;
+
 	// Event that is broadcast to inform observers that the node title has changed
 	// The default SAnimationGraphNode uses this to invalidate cached node title text
 	DECLARE_EVENT(UAnimGraphNode_Base, FOnNodeTitleChangedEvent);
@@ -376,7 +388,8 @@ class ANIMGRAPH_API UAnimGraphNode_Base : public UK2Node
 protected:
 	friend class FAnimBlueprintCompilerContext;
 	friend class FAnimGraphNodeDetails;
-	friend class FAnimBlueprintCompilerHandler_Base;
+	friend class UAnimBlueprintExtension;
+	friend class UAnimBlueprintExtension_Base;
 
 	// Gets the animation FNode type represented by this ed graph node
 	UScriptStruct* GetFNodeType() const;
@@ -384,6 +397,9 @@ protected:
 	// Gets the animation FNode property represented by this ed graph node
 	FStructProperty* GetFNodeProperty() const;
 
+	// Get the extension types that this node type holds on the anim blueprint. Some extension types are always requested by the system
+	virtual void GetRequiredExtensions(TArray<TSubclassOf<UAnimBlueprintExtension>>& OutExtensions) const {}
+	
 	// This will be called when a pose link is found, and can be called with PoseProperty being either of:
 	//  - an array property (ArrayIndex >= 0)
 	//  - a single pose property (ArrayIndex == INDEX_NONE)
@@ -408,6 +424,12 @@ protected:
 	// Copy this node's data during the last phase of compilation where term defaults are copied to the new CDO
 	virtual void OnCopyTermDefaultsToDefaultObject(IAnimBlueprintCopyTermDefaultsContext& InCompilationContext, IAnimBlueprintNodeCopyTermDefaultsContext& InPerNodeContext, IAnimBlueprintGeneratedClassCompiledData& OutCompiledData) {}
 
+	// Called to override the assets held on a runtime anim node. Implements per-node logic for child anim blueprints.
+	void OverrideAssets(IAnimBlueprintNodeOverrideAssetsContext& InContext) const;
+
+	// Override point for OverrideAssets
+	virtual void OnOverrideAssets(IAnimBlueprintNodeOverrideAssetsContext& InContext) const {}
+	
 	// Allocates or reallocates pins
 	void InternalPinCreation(TArray<UEdGraphPin*>* OldPins);
 
diff --git a/Engine/Source/Editor/AnimGraph/Public/AnimGraphNode_BlendSpaceEvaluator.h b/Engine/Source/Editor/AnimGraph/Public/AnimGraphNode_BlendSpaceEvaluator.h
--- a/Engine/Source/Editor/AnimGraph/Public/AnimGraphNode_BlendSpaceEvaluator.h
+++ b/Engine/Source/Editor/AnimGraph/Public/AnimGraphNode_BlendSpaceEvaluator.h
@@ -34,6 +34,7 @@ class UAnimGraphNode_BlendSpaceEvaluator : public UAnimGraphNode_BlendSpaceBase
 	virtual const TCHAR* GetTimePropertyName() const override;
 	virtual UScriptStruct* GetTimePropertyStruct() const override;
 	virtual EAnimAssetHandlerType SupportsAssetClass(const UClass* AssetClass) const override;
+	virtual void OnOverrideAssets(IAnimBlueprintNodeOverrideAssetsContext& InContext) const override;
 	// End of UAnimGraphNode_Base interface
 
 	// UK2Node interface
diff --git a/Engine/Source/Editor/AnimGraph/Public/AnimGraphNode_BlendSpaceGraphBase.h b/Engine/Source/Editor/AnimGraph/Public/AnimGraphNode_BlendSpaceGraphBase.h
--- a/Engine/Source/Editor/AnimGraph/Public/AnimGraphNode_BlendSpaceGraphBase.h
+++ b/Engine/Source/Editor/AnimGraph/Public/AnimGraphNode_BlendSpaceGraphBase.h
@@ -96,6 +96,7 @@ protected:
 	virtual void OnCopyTermDefaultsToDefaultObject(IAnimBlueprintCopyTermDefaultsContext& InCompilationContext, IAnimBlueprintNodeCopyTermDefaultsContext& InPerNodeContext, IAnimBlueprintGeneratedClassCompiledData& OutCompiledData) override;
 	virtual void CustomizeDetails(IDetailLayoutBuilder& InDetailBuilder) override;
 	virtual void GetInputLinkAttributes(FNodeAttributeArray& OutAttributes) const override;
+	virtual void GetRequiredExtensions(TArray<TSubclassOf<UAnimBlueprintExtension>>& OutExtensions) const override;
 
 	// UK2Node interface
 	virtual void PreloadRequiredAssets() override;
diff --git a/Engine/Source/Editor/AnimGraph/Public/AnimGraphNode_BlendSpacePlayer.h b/Engine/Source/Editor/AnimGraph/Public/AnimGraphNode_BlendSpacePlayer.h
--- a/Engine/Source/Editor/AnimGraph/Public/AnimGraphNode_BlendSpacePlayer.h
+++ b/Engine/Source/Editor/AnimGraph/Public/AnimGraphNode_BlendSpacePlayer.h
@@ -35,6 +35,7 @@ class UAnimGraphNode_BlendSpacePlayer : public UAnimGraphNode_BlendSpaceBase
 	virtual void GetAllAnimationSequencesReferred(TArray<UAnimationAsset*>& AnimationAssets) const override;
 	virtual void ReplaceReferredAnimations(const TMap<UAnimationAsset*, UAnimationAsset*>& AnimAssetReplacementMap) override;
 	virtual EAnimAssetHandlerType SupportsAssetClass(const UClass* AssetClass) const override;
+	virtual void OnOverrideAssets(IAnimBlueprintNodeOverrideAssetsContext& InContext) const override;
 	// End of UAnimGraphNode_Base interface
 
 	// UK2Node interface
diff --git a/Engine/Source/Editor/AnimGraph/Public/AnimGraphNode_CustomProperty.h b/Engine/Source/Editor/AnimGraph/Public/AnimGraphNode_CustomProperty.h
--- a/Engine/Source/Editor/AnimGraph/Public/AnimGraphNode_CustomProperty.h
+++ b/Engine/Source/Editor/AnimGraph/Public/AnimGraphNode_CustomProperty.h
@@ -14,7 +14,6 @@
 class FCompilerResultsLog;
 class IDetailLayoutBuilder;
 class IPropertyHandle;
-class FAnimBlueprintCompilerHandler;
 
 UCLASS(Abstract)
 class ANIMGRAPH_API UAnimGraphNode_CustomProperty : public UAnimGraphNode_Base, public IClassVariableCreator
diff --git a/Engine/Source/Editor/AnimGraph/Public/AnimGraphNode_LinkedAnimGraphBase.h b/Engine/Source/Editor/AnimGraph/Public/AnimGraphNode_LinkedAnimGraphBase.h
--- a/Engine/Source/Editor/AnimGraph/Public/AnimGraphNode_LinkedAnimGraphBase.h
+++ b/Engine/Source/Editor/AnimGraph/Public/AnimGraphNode_LinkedAnimGraphBase.h
@@ -37,6 +37,7 @@ public:
 	virtual void GetOutputLinkAttributes(FNodeAttributeArray& OutAttributes) const override;
 	virtual bool ShouldShowAttributesOnPins() const override { return false; }
 	virtual void OnCopyTermDefaultsToDefaultObject(IAnimBlueprintCopyTermDefaultsContext& InCompilationContext, IAnimBlueprintNodeCopyTermDefaultsContext& InPerNodeContext, IAnimBlueprintGeneratedClassCompiledData& OutCompiledData) override;
+	virtual void GetRequiredExtensions(TArray<TSubclassOf<UAnimBlueprintExtension>>& OutExtensions) const override;
 
 	// UAnimGraphNode_CustomProperty interface
 	virtual bool IsStructuralProperty(FProperty* InProperty) const override;
@@ -46,7 +47,7 @@ public:
 	virtual const FAnimNode_LinkedAnimGraph* GetLinkedAnimGraphNode() const PURE_VIRTUAL(UAnimGraphNode_LinkedAnimGraphBase::GetLinkedAnimGraphNode, return nullptr;);
 
 protected:
-	friend class FAnimBlueprintCompilerHandler_LinkedAnimGraph;
+	friend class UAnimBlueprintExtension_LinkedAnimGraph;
 
 	// Called pre-compilation to allocate pose links
 	void AllocatePoseLinks();
diff --git a/Engine/Source/Editor/AnimGraph/Public/AnimGraphNode_RotationOffsetBlendSpace.h b/Engine/Source/Editor/AnimGraph/Public/AnimGraphNode_RotationOffsetBlendSpace.h
--- a/Engine/Source/Editor/AnimGraph/Public/AnimGraphNode_RotationOffsetBlendSpace.h
+++ b/Engine/Source/Editor/AnimGraph/Public/AnimGraphNode_RotationOffsetBlendSpace.h
@@ -28,7 +28,7 @@ class UAnimGraphNode_RotationOffsetBlendSpace: public UAnimGraphNode_BlendSpaceB
 
 	// UAnimGraphNode_Base interface
 	virtual void ValidateAnimNodeDuringCompilation(class USkeleton* ForSkeleton, class FCompilerResultsLog& MessageLog) override;
-	virtual UAnimationAsset* GetAnimationAsset() const override { return Node.BlendSpace; }
+	virtual UAnimationAsset* GetAnimationAsset() const override { return Node.GetBlendSpace(); }
 	virtual void GetAllAnimationSequencesReferred(TArray<UAnimationAsset*>& AnimationAssets) const override;
 	virtual void ReplaceReferredAnimations(const TMap<UAnimationAsset*, UAnimationAsset*>& AnimAssetReplacementMap) override;
 	virtual EAnimAssetHandlerType SupportsAssetClass(const UClass* AssetClass) const override;
diff --git a/Engine/Source/Editor/AnimGraph/Public/AnimGraphNode_SaveCachedPose.h b/Engine/Source/Editor/AnimGraph/Public/AnimGraphNode_SaveCachedPose.h
--- a/Engine/Source/Editor/AnimGraph/Public/AnimGraphNode_SaveCachedPose.h
+++ b/Engine/Source/Editor/AnimGraph/Public/AnimGraphNode_SaveCachedPose.h
@@ -40,6 +40,7 @@ class UAnimGraphNode_SaveCachedPose : public UAnimGraphNode_Base
 	// UAnimGraphNode_Base interface
 	virtual FString GetNodeCategory() const override;
 	virtual bool IsSinkNode() const override { return true; }
+	virtual void GetRequiredExtensions(TArray<TSubclassOf<UAnimBlueprintExtension>>& OutExtensions) const override;
 	// End of UAnimGraphNode_Base interface
 
 private:
diff --git a/Engine/Source/Editor/AnimGraph/Public/AnimGraphNode_SequencePlayer.h b/Engine/Source/Editor/AnimGraph/Public/AnimGraphNode_SequencePlayer.h
--- a/Engine/Source/Editor/AnimGraph/Public/AnimGraphNode_SequencePlayer.h
+++ b/Engine/Source/Editor/AnimGraph/Public/AnimGraphNode_SequencePlayer.h
@@ -19,6 +19,9 @@ class UAnimGraphNode_SequencePlayer : public UAnimGraphNode_AssetPlayerBase
 	UPROPERTY(EditAnywhere, Category=Settings)
 	FAnimNode_SequencePlayer Node;
 
+	// UObject interface
+	void Serialize(FArchive& Ar) override;
+
 	// UEdGraphNode interface
 	virtual FLinearColor GetNodeTitleColor() const override;
 	virtual FText GetTooltipText() const override;
@@ -42,6 +45,7 @@ class UAnimGraphNode_SequencePlayer : public UAnimGraphNode_AssetPlayerBase
 	virtual void GetMenuActions(FBlueprintActionDatabaseRegistrar& ActionRegistrar) const override;
 	virtual bool IsActionFilteredOut(class FBlueprintActionFilter const& Filter) override;
 	virtual EAnimAssetHandlerType SupportsAssetClass(const UClass* AssetClass) const override;
+	virtual void OnOverrideAssets(IAnimBlueprintNodeOverrideAssetsContext& InContext) const override;
 	// End of UAnimGraphNode_Base interface
 
 	// UK2Node interface
diff --git a/Engine/Source/Editor/AnimGraph/Public/IAnimBlueprintCompilationBracketContext.h b/Engine/Source/Editor/AnimGraph/Public/IAnimBlueprintCompilationBracketContext.h
--- a/Engine/Source/Editor/AnimGraph/Public/IAnimBlueprintCompilationBracketContext.h
+++ b/Engine/Source/Editor/AnimGraph/Public/IAnimBlueprintCompilationBracketContext.h
@@ -5,7 +5,6 @@
 #include "CoreMinimal.h"
 
 class FCompilerResultsLog;
-class IAnimBlueprintCompilerHandler;
 
 // Interface passed to start/end compilation delegates
 class ANIMGRAPH_API IAnimBlueprintCompilationBracketContext
@@ -13,21 +12,16 @@ class ANIMGRAPH_API IAnimBlueprintCompilationBracketContext
 public:	
 	virtual ~IAnimBlueprintCompilationBracketContext() {}
 
-	// Get a handler of the specified type and name (i.e. via simple name-based RTTI)
-	// Handlers are registered via IAnimBlueprintCompilerHandlerCollection::RegisterHandler
-	template <typename THandlerClass>
-	THandlerClass* GetHandler(FName InName) const
-	{
-		return static_cast<THandlerClass*>(GetHandlerInternal(InName));
-	}
-
 	// Get the message log for the current compilation
 	FCompilerResultsLog& GetMessageLog() const { return GetMessageLogImpl(); }
 
-protected:	
-	// GetHandler helper function
-	virtual IAnimBlueprintCompilerHandler* GetHandlerInternal(FName InName) const = 0;
+	// Index of the nodes (must match up with the runtime discovery process of nodes, which runs thru the property chain)
+	const TMap<UAnimGraphNode_Base*, int32>& GetAllocatedAnimNodeIndices() const { return GetAllocatedAnimNodeIndicesImpl(); }
 
+protected:
 	// Get the message log for the current compilation
 	virtual FCompilerResultsLog& GetMessageLogImpl() const = 0;
+
+	// Map of anim node properties to original anim graph node
+	virtual const TMap<UAnimGraphNode_Base*, int32>& GetAllocatedAnimNodeIndicesImpl() const = 0;
 };
diff --git a/Engine/Source/Editor/AnimGraph/Public/IAnimBlueprintCompilationContext.h b/Engine/Source/Editor/AnimGraph/Public/IAnimBlueprintCompilationContext.h
--- a/Engine/Source/Editor/AnimGraph/Public/IAnimBlueprintCompilationContext.h
+++ b/Engine/Source/Editor/AnimGraph/Public/IAnimBlueprintCompilationContext.h
@@ -16,7 +16,6 @@ class UBlueprint;
 class UAnimBlueprint;
 class FProperty;
 struct FKismetCompilerOptions;
-class IAnimBlueprintCompilerHandler;
 
 #define ANIM_FUNC_DECORATOR	TEXT("__AnimFunc")
 
@@ -24,19 +23,48 @@ class IAnimBlueprintCompilerHandler;
 class ANIMGRAPH_API IAnimBlueprintCompilationContext
 {
 public:
+	// Record of a property that can be folded into the class members/constant blocks
+	struct FFoldedPropertyRecord
+	{
+		FFoldedPropertyRecord(UAnimGraphNode_Base* InAnimGraphNode, FStructProperty* InAnimNodeProperty, FProperty* InProperty, bool bInIsOnClass)
+			: AnimGraphNode(InAnimGraphNode)
+			, AnimNodeProperty(InAnimNodeProperty)
+			, Property(InProperty)
+			, GeneratedProperty(nullptr)
+			, FoldIndex(INDEX_NONE)
+			, PropertyIndex(INDEX_NONE)
+			, bIsOnClass(bInIsOnClass)
+		{
+		}
+
+		// The anim graph node that this property record came from
+		UAnimGraphNode_Base* AnimGraphNode = nullptr;
+
+		// The property of the FAnimNode_Base-derived structure within the anim graph node
+		FStructProperty* AnimNodeProperty = nullptr;
+
+		// The original property within the FAnimNode_Base
+		FProperty* Property = nullptr;
+
+		// The generated property within the respective data area (either constants or mutables)
+		FProperty* GeneratedProperty = nullptr;
+
+		// The index that this property was folded to. INDEX_NONE if it was not folded.
+		int32 FoldIndex = INDEX_NONE;
+
+		// The index of the property in its respective data area (either constant or mutable)
+		// This will be INDEX_NONE if the property was folded.
+		int32 PropertyIndex = INDEX_NONE;
+
+		// Whether this property will be held on the class (constants on sparse class data), or on the instance (mutables struct)
+		bool bIsOnClass = false;
+	};
+
 	virtual ~IAnimBlueprintCompilationContext() {}
 
 	// Get a compilation context from a kismet compiler context assuming that it is an FAnimBlueprintCompilerContext
 	static TUniquePtr<IAnimBlueprintCompilationContext> Get(FKismetCompilerContext& InKismetCompiler);
 
-	// Get a handler of the specified type and name (i.e. via simple name-based RTTI)
-	// Handlers are registered via IAnimBlueprintCompilerHandlerCollection::RegisterHandler
-	template <typename THandlerClass>
-	THandlerClass* GetHandler(FName InName) const
-	{
-		return static_cast<THandlerClass*>(GetHandlerInternal(InName));
-	}
-
 	// Spawns an intermediate node associated with the source node (for error purposes)
 	template <typename NodeType>
 	NodeType* SpawnIntermediateNode(UEdGraphNode* SourceNode, UEdGraph* ParentGraph = nullptr)
@@ -109,6 +137,15 @@ public:
 	// Get the current compiled-in attributes uniquely assigned to the specified node
 	TArrayView<const FName> GetAttributesFromNode(UAnimGraphNode_Base* InNode) const { return GetAttributesFromNodeImpl(InNode); }
 
+	// Check whether an anim node participates in constant folding
+	bool IsAnimGraphNodeFolded(UAnimGraphNode_Base* InNode) const { return IsAnimGraphNodeFoldedImpl(InNode); }
+
+	// Get the folded property record, if any, for the supplied node & named property
+	const FFoldedPropertyRecord* GetFoldedPropertyRecord(UAnimGraphNode_Base* InNode, FName InPropertyName) const { return GetFoldedPropertyRecordImpl(InNode, InPropertyName); }
+
+	// Get the generated property of the class that mutable data is added to
+	virtual const FStructProperty* GetMutableDataProperty() const { return GetMutableDataPropertyImpl(); }
+	
 protected:
 	// Adds a pose link mapping record
 	virtual void AddPoseLinkMappingRecordImpl(const FPoseLinkMappingRecord& InRecord) = 0;
@@ -161,8 +198,14 @@ protected:
 	// Get the current compiled-in attributes uniquely assigned to the specified node
 	virtual TArrayView<const FName> GetAttributesFromNodeImpl(UAnimGraphNode_Base* InNode) const = 0;
 
-	// GetHandler helper function
-	virtual IAnimBlueprintCompilerHandler* GetHandlerInternal(FName InName) const = 0;
+	// Check whether an anim node participates in constant folding
+	virtual bool IsAnimGraphNodeFoldedImpl(UAnimGraphNode_Base* InNode) const = 0;
+
+	// Get the folded property record, if any, for the supplied node & named property
+	virtual const FFoldedPropertyRecord* GetFoldedPropertyRecordImpl(UAnimGraphNode_Base* InNode, FName InPropertyName) const = 0;
+
+	// Get the generated property of the class that mutable data is added to
+	virtual const FStructProperty* GetMutableDataPropertyImpl() const = 0;
 
 	// Get the compiler as a base class to avoid circular include issues with templated functions/classes
 	virtual FKismetCompilerContext* GetKismetCompiler() const = 0;
diff --git a/Engine/Source/Editor/AnimGraph/Public/IAnimBlueprintCompilerCreationContext.h b/Engine/Source/Editor/AnimGraph/Public/IAnimBlueprintCompilerCreationContext.h
--- a/Engine/Source/Editor/AnimGraph/Public/IAnimBlueprintCompilerCreationContext.h
+++ b/Engine/Source/Editor/AnimGraph/Public/IAnimBlueprintCompilerCreationContext.h
@@ -17,83 +17,51 @@ class UAnimGraphNode_Base;
 class UEdGraph;
 class UEdGraphSchema;
 
-/** Begin ordered delegate calls - these functions are called in the order presented here */
-
-/** 
- * Delegate fired when the class starts compiling. The class may be new or recycled.
- * @param	InClass			The class that is being compiled. This could be a newly created class or a recycled class
- * @param	InCompilerContext	The compiler context for the current compilation
- * @param	OutCompiledData	The compiled data that this handler can write to
- */
+UE_DEPRECATED(5.0, "Anim BP compiler delegate system is deprecated. Use UAnimBlueprintExtension instead")
 DECLARE_MULTICAST_DELEGATE_ThreeParams(FOnStartCompilingClass, const UClass* /*InClass*/, IAnimBlueprintCompilationBracketContext& /*InCompilerContext*/, IAnimBlueprintGeneratedClassCompiledData& /*OutCompiledData*/);
 
-/** 
- * Delegate fired before all animation nodes are processed
- * @param	InAnimNodes		The set of anim nodes that should be processed. Note that these nodes have not yet been pruned for connectivity
- * @param	InCompilerContext	The compiler context for the current compilation
- * @param	OutCompiledData	The compiled data that this handler can write to
- */
+UE_DEPRECATED(5.0, "Anim BP compiler delegate system is deprecated. Use UAnimBlueprintExtension instead")
 DECLARE_MULTICAST_DELEGATE_ThreeParams(FOnPreProcessAnimationNodes, TArrayView<UAnimGraphNode_Base*> /*InAnimNodes*/, IAnimBlueprintCompilationContext& /*InCompilerContext*/, IAnimBlueprintGeneratedClassCompiledData& /*OutCompiledData*/);
 
-/** 
- * Delegate fired after all animation nodes are processed
- * @param	InAnimNodes		The set of anim nodes that were processed. Note that these nodes were not pruned for connectivity (they will be the same set passed to PreProcessAnimationNodes)
- * @param	InCompilerContext	The compiler context for the current compilation
- * @param	OutCompiledData	The compiled data that this handler can write to
- */
+UE_DEPRECATED(5.0, "Anim BP compiler delegate system is deprecated. Use UAnimBlueprintExtension instead")
 DECLARE_MULTICAST_DELEGATE_ThreeParams(FOnPostProcessAnimationNodes, TArrayView<UAnimGraphNode_Base*> /*InAnimNodes*/, IAnimBlueprintCompilationContext& /*InCompilerContext*/, IAnimBlueprintGeneratedClassCompiledData& /*OutCompiledData*/);
 
-/** 
- * Delegate fired post- graph expansion
- * @param	InGraph		The graph that was just expanded
- * @param	InCompilerContext	The compiler context for the current compilation
- * @param	OutCompiledData	The compiled data that this handler can write to
- */
+UE_DEPRECATED(5.0, "Anim BP compiler delegate system is deprecated. Use UAnimBlueprintExtension instead")
 DECLARE_MULTICAST_DELEGATE_ThreeParams(FOnPostExpansionStep, const UEdGraph* /*InGraph*/, IAnimBlueprintPostExpansionStepContext& /*InCompilerContext*/, IAnimBlueprintGeneratedClassCompiledData& /*OutCompiledData*/);
 
-/** 
- * Delegate fired when the class has finished compiling
- * @param	InClass			The class that was compiled
- * @param	InCompilerContext	The compiler context for the current compilation
- * @param	OutCompiledData	The compiled data that this handler can write to
- */
+UE_DEPRECATED(5.0, "Anim BP compiler delegate system is deprecated. Use UAnimBlueprintExtension instead")
 DECLARE_MULTICAST_DELEGATE_ThreeParams(FOnFinishCompilingClass, const UClass* /*InClass*/, IAnimBlueprintCompilationBracketContext& /*InCompilerContext*/, IAnimBlueprintGeneratedClassCompiledData& /*OutCompiledData*/);
 
-/** 
- * Delegate fired when data is being copied to the CDO
- * @param	InDefaultObject		The CDO for the just-compiled class
- * @param	InCompilerContext	The compiler context for the current compilation
- * @param	OutCompiledData	The compiled data that this handler can write to
- */
+UE_DEPRECATED(5.0, "Anim BP compiler delegate system is deprecated. Use UAnimBlueprintExtension instead")
 DECLARE_MULTICAST_DELEGATE_ThreeParams(FOnCopyTermDefaultsToDefaultObject, UObject* /*InDefaultObject*/, IAnimBlueprintCopyTermDefaultsContext& /*InCompilerContext*/, IAnimBlueprintGeneratedClassCompiledData& /*OutCompiledData*/);
 
-/** End ordered calls */
-
 /** Interface to the creation of an anim BP compiler context */
 class ANIMGRAPH_API IAnimBlueprintCompilerCreationContext
 {
 public:
 	virtual ~IAnimBlueprintCompilerCreationContext() {}
 
-	/** Delegate fired when the class starts compiling. The class may be new or recycled. */
-	virtual FOnStartCompilingClass& OnStartCompilingClass() = 0;
-
-	/** Delegate fired before all animation nodes are processed */
-	virtual FOnPreProcessAnimationNodes& OnPreProcessAnimationNodes() = 0;
+	/** Registers a graphs schema class to the anim BP compiler so that default function processing is not performed on it */
+	virtual void RegisterKnownGraphSchema(TSubclassOf<UEdGraphSchema> InGraphSchemaClass) = 0;
+	
+	PRAGMA_DISABLE_DEPRECATION_WARNINGS
+	UE_DEPRECATED(5.0, "Anim BP compiler delegate system is deprecated. Use UAnimBlueprintExtension instead")
+	virtual FOnStartCompilingClass& OnStartCompilingClass() { static FOnStartCompilingClass Dummy; return Dummy; }
 
-	/** Delegate fired after all animation nodes are processed */
-	virtual FOnPostProcessAnimationNodes& OnPostProcessAnimationNodes() = 0;
+	UE_DEPRECATED(5.0, "Anim BP compiler delegate system is deprecated. Use UAnimBlueprintExtension instead")
+	virtual FOnPreProcessAnimationNodes& OnPreProcessAnimationNodes() { static FOnPreProcessAnimationNodes Dummy; return Dummy; }
 
-	/** Delegate fired post- graph expansion */
-	virtual FOnPostExpansionStep& OnPostExpansionStep() = 0;
+	UE_DEPRECATED(5.0, "Anim BP compiler delegate system is deprecated. Use UAnimBlueprintExtension instead")
+	virtual FOnPostProcessAnimationNodes& OnPostProcessAnimationNodes() { static FOnPostProcessAnimationNodes Dummy; return Dummy; }
 
-	/** Delegate fired when the class has finished compiling */
-	virtual FOnFinishCompilingClass& OnFinishCompilingClass() = 0;
+	UE_DEPRECATED(5.0, "Anim BP compiler delegate system is deprecated. Use UAnimBlueprintExtension instead")
+	virtual FOnPostExpansionStep& OnPostExpansionStep() { static FOnPostExpansionStep Dummy; return Dummy; }
 
-	/** Delegate fired when data is being copied to the CDO */
-	virtual FOnCopyTermDefaultsToDefaultObject& OnCopyTermDefaultsToDefaultObject() = 0;
+	UE_DEPRECATED(5.0, "Anim BP compiler delegate system is deprecated. Use UAnimBlueprintExtension instead")
+	virtual FOnFinishCompilingClass& OnFinishCompilingClass() { static FOnFinishCompilingClass Dummy; return Dummy; }
 
-	/** Registers a graphs schema class to the anim BP compiler so that default function processing is not performed on it */
-	virtual void RegisterKnownGraphSchema(TSubclassOf<UEdGraphSchema> InGraphSchemaClass) = 0;
+	UE_DEPRECATED(5.0, "Anim BP compiler delegate system is deprecated. Use UAnimBlueprintExtension instead")
+	virtual FOnCopyTermDefaultsToDefaultObject& OnCopyTermDefaultsToDefaultObject() { static FOnCopyTermDefaultsToDefaultObject Dummy; return Dummy; }
+	PRAGMA_ENABLE_DEPRECATION_WARNINGS
 };
 
diff --git a/Engine/Source/Editor/AnimGraph/Public/IAnimBlueprintCompilerHandler.h b/Engine/Source/Editor/AnimGraph/Public/IAnimBlueprintCompilerHandler.h
--- a/Engine/Source/Editor/AnimGraph/Public/IAnimBlueprintCompilerHandler.h
+++ b/Engine/Source/Editor/AnimGraph/Public/IAnimBlueprintCompilerHandler.h
@@ -4,6 +4,8 @@
 
 #include "CoreMinimal.h"
 
+class UE_DEPRECATED(5.0, "IAnimBlueprintCompilerHandler is deprecated, use UAnimBlueprintExtension_PropertyAccess instead") IAnimBlueprintCompilerHandler;
+
 /** Empty base class for handlers that are maintained for the lifetime of the compiler context */
 class ANIMGRAPH_API IAnimBlueprintCompilerHandler
 {
diff --git a/Engine/Source/Editor/AnimGraph/Public/IAnimBlueprintCompilerHandlerCollection.h b/Engine/Source/Editor/AnimGraph/Public/IAnimBlueprintCompilerHandlerCollection.h
--- a/Engine/Source/Editor/AnimGraph/Public/IAnimBlueprintCompilerHandlerCollection.h
+++ b/Engine/Source/Editor/AnimGraph/Public/IAnimBlueprintCompilerHandlerCollection.h
@@ -8,16 +8,19 @@
 class FAnimBlueprintCompilerContext;
 class IAnimBlueprintCompilerCreationContext;
 
+class UE_DEPRECATED(5.0, "IAnimBlueprintCompilerHandlerCollection is no longer used") IAnimBlueprintCompilerHandlerCollection;
+
 class ANIMGRAPH_API IAnimBlueprintCompilerHandlerCollection
 {
 public:
+	PRAGMA_DISABLE_DEPRECATION_WARNINGS
 	virtual ~IAnimBlueprintCompilerHandlerCollection() {}
 
 	/** Register a named handler */
-	static void RegisterHandler(FName InName, TFunction<TUniquePtr<IAnimBlueprintCompilerHandler>(IAnimBlueprintCompilerCreationContext&)> InFunction);
+	static void RegisterHandler(FName InName, TFunction<TUniquePtr<IAnimBlueprintCompilerHandler>(IAnimBlueprintCompilerCreationContext&)> InFunction) {}
 
 	/** Register a named handler */
-	static void UnregisterHandler(FName InName);
+	static void UnregisterHandler(FName InName) {}
 
 	// Get a GetHandler of specified type with the specified name
 	template<typename THandlerClass>
@@ -29,4 +32,5 @@ public:
 protected:
 	// Get a GetHandler with the specified name
 	virtual IAnimBlueprintCompilerHandler* GetHandlerByName(FName InName) const = 0;
+	PRAGMA_ENABLE_DEPRECATION_WARNINGS
 };
\ No newline at end of file
diff --git a/Engine/Source/Editor/AnimGraph/Public/IAnimBlueprintCopyTermDefaultsContext.h b/Engine/Source/Editor/AnimGraph/Public/IAnimBlueprintCopyTermDefaultsContext.h
--- a/Engine/Source/Editor/AnimGraph/Public/IAnimBlueprintCopyTermDefaultsContext.h
+++ b/Engine/Source/Editor/AnimGraph/Public/IAnimBlueprintCopyTermDefaultsContext.h
@@ -11,7 +11,7 @@ class UAnimBlueprint;
 class IAnimBlueprintCopyTermDefaultsContext
 {
 public:	
-	virtual ~IAnimBlueprintCopyTermDefaultsContext() {}
+	virtual ~IAnimBlueprintCopyTermDefaultsContext() = default;
 
 	// Get the message log for the current compilation
 	FCompilerResultsLog& GetMessageLog() const { return GetMessageLogImpl(); }
@@ -32,7 +32,7 @@ protected:
 class IAnimBlueprintNodeCopyTermDefaultsContext
 {
 public:	
-	virtual ~IAnimBlueprintNodeCopyTermDefaultsContext() {}
+	virtual ~IAnimBlueprintNodeCopyTermDefaultsContext() = default;
 
 	// Get the property that we are writing to
 	const FProperty* GetTargetProperty() const { return GetTargetPropertyImpl(); }
@@ -58,4 +58,36 @@ protected:
 
 	// Get the property index for this node
 	virtual int32 GetNodePropertyIndexImpl() const  = 0;
-};
\ No newline at end of file
+};
+
+// Interface passed to per-extension CopyTermDefaults override point
+class IAnimBlueprintExtensionCopyTermDefaultsContext
+{
+public:	
+	virtual ~IAnimBlueprintExtensionCopyTermDefaultsContext() = default;
+
+	// Get the property that we are writing to
+	const FProperty* GetTargetProperty() const { return GetTargetPropertyImpl(); }
+
+	// Get the destination ptr (the node) that we are writing to
+	uint8* GetDestinationPtr() const { return GetDestinationPtrImpl(); }
+
+	// Get the source ptr (the node in the anim graph node) that we are reading from
+	const uint8* GetSourcePtr() const { return GetSourcePtrImpl(); }
+
+	// Get the property index for this node
+	int32 GetNodePropertyIndex() const { return GetNodePropertyIndexImpl(); }
+
+protected:
+	// Get the property that we are writing to
+	virtual const FProperty* GetTargetPropertyImpl() const = 0;
+
+	// Get the destination ptr (the node) that we are writing to
+	virtual uint8* GetDestinationPtrImpl() const = 0;
+
+	// Get the source ptr (the node in the anim graph node) that we are reading from
+	virtual const uint8* GetSourcePtrImpl() const = 0;
+
+	// Get the property index for this node
+	virtual int32 GetNodePropertyIndexImpl() const = 0;
+};
diff --git a/Engine/Source/Editor/AnimGraph/Public/IAnimBlueprintGeneratedClassCompiledData.h b/Engine/Source/Editor/AnimGraph/Public/IAnimBlueprintGeneratedClassCompiledData.h
--- a/Engine/Source/Editor/AnimGraph/Public/IAnimBlueprintGeneratedClassCompiledData.h
+++ b/Engine/Source/Editor/AnimGraph/Public/IAnimBlueprintGeneratedClassCompiledData.h
@@ -7,9 +7,9 @@
 #include "Animation/AnimClassInterface.h"
 #include "Animation/AnimTypes.h"
 #include "Animation/AnimNodeBase.h"
+#include "PropertyAccess.h"
 
 struct FBlueprintDebugData;
-struct FPropertyAccessLibrary;
 struct FAnimBlueprintDebugData;
 
 /** Interface to the writable parts of the generated class that handlers can operate on */
@@ -31,11 +31,11 @@ public:
 	// Finds a notify event or adds if it doesn't already exist
 	virtual int32 FindOrAddNotify(FAnimNotifyEvent& Notify) const = 0;
 
-	// Get the currently-compiled classes exposed value handlers
-	virtual TArray<FExposedValueHandler>& GetExposedValueHandlers() const = 0;
+	UE_DEPRECATED(5.0, "Exposed value handlers are no longer accessible via this method. Access them via UAnimBlueprintExtension_Base.")
+	virtual TArray<FExposedValueHandler>& GetExposedValueHandlers() const {	static TArray<FExposedValueHandler> Dummy; return Dummy; }
 
-	// Get the currently-compiled classes property access library
-	virtual FPropertyAccessLibrary& GetPropertyAccessLibrary() const = 0;
+	UE_DEPRECATED(5.0, "Property access library is no longer accessible via this method. Access it via UAnimBlueprintExtension_PropertyAccess.")
+	virtual FPropertyAccessLibrary& GetPropertyAccessLibrary() const { static FPropertyAccessLibrary Dummy; return Dummy; }
 
 	// Get the anim debug data for the currently-compiled class
 	virtual FAnimBlueprintDebugData& GetAnimBlueprintDebugData() const = 0;
diff --git a/Engine/Source/Editor/AnimGraph/Public/IAnimBlueprintPostExpansionStepContext.h b/Engine/Source/Editor/AnimGraph/Public/IAnimBlueprintPostExpansionStepContext.h
--- a/Engine/Source/Editor/AnimGraph/Public/IAnimBlueprintPostExpansionStepContext.h
+++ b/Engine/Source/Editor/AnimGraph/Public/IAnimBlueprintPostExpansionStepContext.h
@@ -3,7 +3,6 @@
 #pragma once
 
 #include "CoreMinimal.h"
-#include "IAnimBlueprintCompilerHandler.h"
 
 class FCompilerResultsLog;
 class UEdGraph;
@@ -15,14 +14,6 @@ class IAnimBlueprintPostExpansionStepContext
 public:
 	virtual ~IAnimBlueprintPostExpansionStepContext() {}
 
-	// Get a handler of the specified type and name (i.e. via simple name-based RTTI)
-	// Handlers are registered via IAnimBlueprintCompilerHandlerCollection::RegisterHandler
-	template <typename THandlerClass>
-	THandlerClass* GetHandler(FName InName) const
-	{
-		return static_cast<THandlerClass*>(GetHandlerInternal(InName));
-	}
-
 	// Get the message log for the current compilation
 	FCompilerResultsLog& GetMessageLog() const { return GetMessageLogImpl(); }
 
@@ -41,7 +32,4 @@ protected:
 
 	// Get the compiler options we are currently using
 	virtual const FKismetCompilerOptions& GetCompileOptionsImpl() const = 0;
-
-	// GetHandler helper function
-	virtual IAnimBlueprintCompilerHandler* GetHandlerInternal(FName InName) const = 0;
 };
