Backing out CL 16071104

diff --git a/Engine/Config/BaseEngine.ini b/Engine/Config/BaseEngine.ini
--- a/Engine/Config/BaseEngine.ini
+++ b/Engine/Config/BaseEngine.ini
@@ -923,9 +923,6 @@ NearClipPlane=10.0
 +PropertyRedirects=(OldName="AnimationBlueprintLibrary.FindBonePathToRoot.AnimationSequence",NewName="AnimationSequenceBase")
 +ClassRedirects=(OldName="/Script/Engine.BlendSpaceBase",NewName="/Script/Engine.BlendSpace")
 +FunctionRedirects=(OldName="Image.SetBrushSize",NewName="SetDesiredSizeOverride")
-+ClassRedirects=(OldName="/Script/PropertyAccess.AnimBlueprintClassSubsystem_PropertyAccess",NewName="/Script/Engine.AnimBlueprintClassSubsystem_PropertyAccess")
-+PackageRedirects=(OldName="/Script/PropertyAccess", Removed=true)
-+StructRedirects=(OldName="AnimInstanceSubsystemData",NewName="/Script/Engine.AnimSubsystemInstance")
 
 [CoreUObject.Metadata]
 ; Note: UnrealHeaderTool should be rerun after making changes to MetadataRedirects to catch any keys specified in class headers
@@ -2394,7 +2391,6 @@ DefaultASTCQualityBySize=3
 +VersionedIntRValues=r.Shaders.Optimize
 +VersionedIntRValues=r.CompileShadersForDevelopment
 +VersionedIntRValues=r.MobileHDR
-+VersionedIntRValues=r.UsePreExposure
 
 [/Script/UnrealEd.ChunkDependencyInfo]
 ; Entries can be added to this section of DefaultEngine.ini to set dependencies for package chunks, used by cooking and packaging
diff --git a/Engine/Plugins/Developer/PropertyAccessNode/Source/PropertyAccessNode/Private/K2Node_PropertyAccess.cpp b/Engine/Plugins/Developer/PropertyAccessNode/Source/PropertyAccessNode/Private/K2Node_PropertyAccess.cpp
--- a/Engine/Plugins/Developer/PropertyAccessNode/Source/PropertyAccessNode/Private/K2Node_PropertyAccess.cpp
+++ b/Engine/Plugins/Developer/PropertyAccessNode/Source/PropertyAccessNode/Private/K2Node_PropertyAccess.cpp
@@ -15,13 +15,12 @@
 #include "K2Node_VariableGet.h"
 #include "AnimationGraph.h"
 #include "Kismet2/BlueprintEditorUtils.h"
-#include "AnimBlueprintExtension_PropertyAccess.h"
+#include "PropertyAccessCompilerHandler.h"
 #include "Modules/ModuleManager.h"
 #include "IPropertyAccessEditor.h"
 #include "IPropertyAccessCompiler.h"
 #include "Features/IModularFeatures.h"
 #include "IAnimBlueprintCompilationContext.h"
-#include "Animation/AnimBlueprint.h"
 
 #define LOCTEXT_NAMESPACE "K2Node_PropertyAccess"
 
@@ -50,9 +49,9 @@ void UK2Node_PropertyAccess::ExpandNode(FKismetCompilerContext& InCompilerContex
 
 		// Create a copy event in the complied generated class
 		TUniquePtr<IAnimBlueprintCompilationContext> CompilationContext = IAnimBlueprintCompilationContext::Get(InCompilerContext);
-		UAnimBlueprintExtension_PropertyAccess* PropertyAccessExtension = UAnimBlueprintExtension::GetExtension<UAnimBlueprintExtension_PropertyAccess>(CastChecked<UAnimBlueprint>(GetBlueprint()));
-		check(PropertyAccessExtension);
-		PropertyAccessExtension->AddCopy(Path, DestPropertyPath, EPropertyAccessBatchType::Batched, this);
+		FPropertyAccessCompilerHandler* PropertyAccessHandler = CompilationContext->GetHandler<FPropertyAccessCompilerHandler>("PropertyAccessCompilerHandler");
+		check(PropertyAccessHandler);
+		PropertyAccessHandler->AddCopy(Path, DestPropertyPath, EPropertyAccessBatchType::Batched, this);
 
 		// Replace us with a get node
 		UK2Node_VariableGet* VariableGetNode = InCompilerContext.SpawnIntermediateNode<UK2Node_VariableGet>(this, InSourceGraph);
diff --git a/Engine/Plugins/Developer/PropertyAccessNode/Source/PropertyAccessNode/PropertyAccessNode.Build.cs b/Engine/Plugins/Developer/PropertyAccessNode/Source/PropertyAccessNode/PropertyAccessNode.Build.cs
--- a/Engine/Plugins/Developer/PropertyAccessNode/Source/PropertyAccessNode/PropertyAccessNode.Build.cs
+++ b/Engine/Plugins/Developer/PropertyAccessNode/Source/PropertyAccessNode/PropertyAccessNode.Build.cs
@@ -12,6 +12,7 @@ public class PropertyAccessNode : ModuleRules
 				"Core",
 				"CoreUObject",
 				"Slate",	
+				"PropertyAccess",
 				"GraphEditor",
 				"BlueprintGraph",
 				"EditorStyle",
diff --git a/Engine/Plugins/Runtime/PropertyAccess/Source/PropertyAccessEditor/Private/AnimBlueprintCompilerHandler_PropertyAccess.cpp b/Engine/Plugins/Runtime/PropertyAccess/Source/PropertyAccessEditor/Private/AnimBlueprintCompilerHandler_PropertyAccess.cpp
new file mode 100644
--- /dev/null
+++ b/Engine/Plugins/Runtime/PropertyAccess/Source/PropertyAccessEditor/Private/AnimBlueprintCompilerHandler_PropertyAccess.cpp
@@ -0,0 +1,53 @@
+// Copyright Epic Games, Inc. All Rights Reserved.
+
+#include "AnimBlueprintCompilerHandler_PropertyAccess.h"
+#include "Animation/AnimBlueprintGeneratedClass.h"
+#include "IAnimBlueprintGeneratedClassCompiledData.h"
+#include "IAnimBlueprintCompilerCreationContext.h"
+#include "IAnimBlueprintCompilationContext.h"
+#include "IAnimBlueprintCompilationBracketContext.h"
+
+FAnimBlueprintCompilerHandler_PropertyAccess::FAnimBlueprintCompilerHandler_PropertyAccess(IAnimBlueprintCompilerCreationContext& InCreationContext)
+{
+	InCreationContext.OnStartCompilingClass().AddRaw(this, &FAnimBlueprintCompilerHandler_PropertyAccess::StartCompilingClass);
+	InCreationContext.OnFinishCompilingClass().AddRaw(this, &FAnimBlueprintCompilerHandler_PropertyAccess::FinishCompilingClass);
+}
+
+int32 FAnimBlueprintCompilerHandler_PropertyAccess::AddCopy(TArrayView<FString> InSourcePath, TArrayView<FString> InDestPath, EPropertyAccessBatchType InBatchType, UObject* InObject)
+{
+	return PropertyAccessLibraryCompiler.AddCopy(InSourcePath, InDestPath, InBatchType, InObject);
+}
+
+void FAnimBlueprintCompilerHandler_PropertyAccess::StartCompilingClass(const UClass* InClass, IAnimBlueprintCompilationBracketContext& InCompilationContext, IAnimBlueprintGeneratedClassCompiledData& OutCompiledData)
+{
+	PropertyAccessLibraryCompiler.Setup(InClass, &OutCompiledData.GetPropertyAccessLibrary());
+	PropertyAccessLibraryCompiler.BeginCompilation(InClass);
+}
+
+void FAnimBlueprintCompilerHandler_PropertyAccess::FinishCompilingClass(const UClass* InClass, IAnimBlueprintCompilationBracketContext& InCompilationContext, IAnimBlueprintGeneratedClassCompiledData& OutCompiledData)
+{
+	OnPreLibraryCompiledDelegate.Broadcast();
+
+	if(!PropertyAccessLibraryCompiler.FinishCompilation())
+	{
+		PropertyAccessLibraryCompiler.IterateErrors([&InCompilationContext](const FText& InErrorText, UObject* InObject)
+		{
+			// Output any property access errors as warnings
+			if(InObject)
+			{
+				InCompilationContext.GetMessageLog().Warning(*InErrorText.ToString(), InObject);
+			}
+			else
+			{
+				InCompilationContext.GetMessageLog().Warning(*InErrorText.ToString());
+			}
+		});
+	}
+
+	OnPostLibraryCompiledDelegate.Broadcast(OutCompiledData);
+}
+
+int32 FAnimBlueprintCompilerHandler_PropertyAccess::MapCopyIndex(int32 InIndex) const
+{
+	return PropertyAccessLibraryCompiler.MapCopyIndex(InIndex);
+}
\ No newline at end of file
diff --git a/Engine/Plugins/Runtime/PropertyAccess/Source/PropertyAccessEditor/Private/AnimBlueprintCompilerHandler_PropertyAccess.h b/Engine/Plugins/Runtime/PropertyAccess/Source/PropertyAccessEditor/Private/AnimBlueprintCompilerHandler_PropertyAccess.h
new file mode 100644
--- /dev/null
+++ b/Engine/Plugins/Runtime/PropertyAccess/Source/PropertyAccessEditor/Private/AnimBlueprintCompilerHandler_PropertyAccess.h
@@ -0,0 +1,38 @@
+// Copyright Epic Games, Inc. All Rights Reserved.
+
+#pragma once
+
+#include "CoreMinimal.h"
+#include "IAnimBlueprintCompilerHandler.h"
+#include "PropertyAccessCompilerHandler.h"
+#include "PropertyAccessEditor.h"
+
+class IAnimBlueprintCompilerCreationContext;
+class IAnimBlueprintCompilationBracketContext;
+class IAnimBlueprintGeneratedClassCompiledData;
+
+class FAnimBlueprintCompilerHandler_PropertyAccess : public FPropertyAccessCompilerHandler
+{
+public:
+	FAnimBlueprintCompilerHandler_PropertyAccess(IAnimBlueprintCompilerCreationContext& InCreationContext);
+
+	// FPropertyAccessCompilerHandler interface
+	virtual int32 AddCopy(TArrayView<FString> InSourcePath, TArrayView<FString> InDestPath, EPropertyAccessBatchType InBatchType, UObject* InObject = nullptr) override;
+	virtual FSimpleMulticastDelegate& OnPreLibraryCompiled() override { return OnPreLibraryCompiledDelegate; }
+	virtual FOnPostLibraryCompiled& OnPostLibraryCompiled() override { return OnPostLibraryCompiledDelegate; }
+	virtual int32 MapCopyIndex(int32 InIndex) const override;
+
+private:
+	void StartCompilingClass(const UClass* InClass, IAnimBlueprintCompilationBracketContext& InCompilationContext, IAnimBlueprintGeneratedClassCompiledData& OutCompiledData);
+	void FinishCompilingClass(const UClass* InClass, IAnimBlueprintCompilationBracketContext& InCompilationContext, IAnimBlueprintGeneratedClassCompiledData& OutCompiledData);
+
+private:
+	// Property access library compiler
+	FPropertyAccessLibraryCompiler PropertyAccessLibraryCompiler;
+
+	// Delegate called before the library is compiled
+	FSimpleMulticastDelegate OnPreLibraryCompiledDelegate;
+
+	// Delegate called when the library is compiled (whether successfully or not)
+	FOnPostLibraryCompiled OnPostLibraryCompiledDelegate;
+};
\ No newline at end of file
diff --git a/Engine/Plugins/Runtime/PropertyAccess/Source/PropertyAccessEditor/Private/PropertyAccessEditor.cpp b/Engine/Plugins/Runtime/PropertyAccess/Source/PropertyAccessEditor/Private/PropertyAccessEditor.cpp
--- a/Engine/Plugins/Runtime/PropertyAccess/Source/PropertyAccessEditor/Private/PropertyAccessEditor.cpp
+++ b/Engine/Plugins/Runtime/PropertyAccess/Source/PropertyAccessEditor/Private/PropertyAccessEditor.cpp
@@ -131,10 +131,15 @@ struct FPropertyAccessEditorSystem
 		}
 		else
 		{
-			InContext.ErrorMessage = FText::Format(LOCTEXT("UnrecognisedProperty", "Property '{0}' is unrecognised in property path for @@"), InProperty ? FText::FromName(InProperty->GetFName()) : LOCTEXT("Null", "Null"));
 			return ESegmentResolveResult::Failed;
 		}
 
+		static const FName PropertyEventMetadata("PropertyEvent");
+		if(InProperty->HasMetaData(PropertyEventMetadata))
+		{
+			InSegment.Flags |= (uint16)EPropertyAccessSegmentFlags::Event;
+		}
+
 		return Result;
 	}
 
@@ -565,13 +570,13 @@ namespace PropertyAccess
 	}
 }
 
-FPropertyAccessLibraryCompiler::FPropertyAccessLibraryCompiler(FPropertyAccessLibrary* InLibrary, const UClass* InClass)
-	: Library(InLibrary)
-	, Class(InClass)
+FPropertyAccessLibraryCompiler::FPropertyAccessLibraryCompiler()
+	: Library(nullptr)
+	, Class(nullptr)
 {
 }
 
-void FPropertyAccessLibraryCompiler::BeginCompilation()
+void FPropertyAccessLibraryCompiler::BeginCompilation(const UClass* InClass)
 {
 	if(Class && Library)
 	{
diff --git a/Engine/Plugins/Runtime/PropertyAccess/Source/PropertyAccessEditor/Private/PropertyAccessEditor.h b/Engine/Plugins/Runtime/PropertyAccess/Source/PropertyAccessEditor/Private/PropertyAccessEditor.h
--- a/Engine/Plugins/Runtime/PropertyAccess/Source/PropertyAccessEditor/Private/PropertyAccessEditor.h
+++ b/Engine/Plugins/Runtime/PropertyAccess/Source/PropertyAccessEditor/Private/PropertyAccessEditor.h
@@ -29,10 +29,10 @@ namespace PropertyAccess
 class FPropertyAccessLibraryCompiler : public IPropertyAccessLibraryCompiler
 {
 public:
-	FPropertyAccessLibraryCompiler(FPropertyAccessLibrary* InLibrary, const UClass* InClass);
+	FPropertyAccessLibraryCompiler();
 
 	// IPropertyAccessLibraryCompiler interface
-	virtual void BeginCompilation() override;
+	virtual void BeginCompilation(const UClass* InClass) override;
 	virtual int32 AddCopy(TArrayView<FString> InSourcePath, TArrayView<FString> InDestPath, EPropertyAccessBatchType InBatchType, UObject* InAssociatedObject = nullptr) override;
 	virtual bool FinishCompilation() override;
 	virtual void IterateErrors(TFunctionRef<void(const FText&, UObject*)> InFunction) const override;
diff --git a/Engine/Plugins/Runtime/PropertyAccess/Source/PropertyAccessEditor/Private/PropertyAccessEditorModule.cpp b/Engine/Plugins/Runtime/PropertyAccess/Source/PropertyAccessEditor/Private/PropertyAccessEditorModule.cpp
--- a/Engine/Plugins/Runtime/PropertyAccess/Source/PropertyAccessEditor/Private/PropertyAccessEditorModule.cpp
+++ b/Engine/Plugins/Runtime/PropertyAccess/Source/PropertyAccessEditor/Private/PropertyAccessEditorModule.cpp
@@ -8,6 +8,8 @@
 #include "PropertyAccessEditor.h"
 #include "Modules/ModuleInterface.h"
 #include "Features/IModularFeatures.h"
+#include "IAnimBlueprintCompilerHandlerCollection.h"
+#include "AnimBlueprintCompilerHandler_PropertyAccess.h"
 
 class FPropertyAccessEditorModule : public IPropertyAccessEditor, public IModuleInterface
 {
@@ -16,10 +18,19 @@ public:
 	virtual void StartupModule() override
 	{
 		IModularFeatures::Get().RegisterModularFeature("PropertyAccessEditor", this);
+
+		// Register node compilation handlers
+		IAnimBlueprintCompilerHandlerCollection::RegisterHandler("PropertyAccessCompilerHandler", [](IAnimBlueprintCompilerCreationContext& InCreationContext)
+		{
+			return MakeUnique<FAnimBlueprintCompilerHandler_PropertyAccess>(InCreationContext);
+		});
 	}
 
 	virtual void ShutdownModule() override
 	{
+		// Register node compilation handlers
+		IAnimBlueprintCompilerHandlerCollection::UnregisterHandler("PropertyAccessCompilerHandler");
+
 		IModularFeatures::Get().UnregisterModularFeature("PropertyAccessEditor", this);
 	}
 
@@ -56,11 +67,6 @@ public:
 	{
 		PropertyAccess::MakeStringPath(InBindingChain, OutStringPath);
 	}
-
-	virtual TUniquePtr<IPropertyAccessLibraryCompiler> MakePropertyAccessCompiler(const FPropertyAccessLibraryCompilerArgs& InArgs) const override
-	{
-		return MakeUnique<FPropertyAccessLibraryCompiler>(&InArgs.Library, InArgs.ClassContext);
-	}
 };
 
 IMPLEMENT_MODULE(FPropertyAccessEditorModule, PropertyAccessEditor)
\ No newline at end of file
diff --git a/Engine/Source/Developer/BlueprintCompilerCppBackend/Private/BlueprintCompilerCppBackendAnim.cpp b/Engine/Source/Developer/BlueprintCompilerCppBackend/Private/BlueprintCompilerCppBackendAnim.cpp
--- a/Engine/Source/Developer/BlueprintCompilerCppBackend/Private/BlueprintCompilerCppBackendAnim.cpp
+++ b/Engine/Source/Developer/BlueprintCompilerCppBackend/Private/BlueprintCompilerCppBackendAnim.cpp
@@ -81,21 +81,20 @@ void FBackendHelperAnim::AddAnimNodeInitializationFunction(FEmitterLocalContext&
 			FEmitDefaultValueHelper::OuterGenerate(Context, InProperty, TEXT(""), reinterpret_cast<const uint8*>(InCDO), bInNewProperty ? nullptr : reinterpret_cast<const uint8*>(InParentCDO), FEmitDefaultValueHelper::EPropertyAccessOperator::None, FEmitDefaultValueHelper::EPropertyGenerationControlFlags::AllowProtected);
 			
 			// anim nodes constructed, finish anim node initialization:
-			// @TODO: if we are to support nativizaton of anim BPs this will need to be reworked
-// 			if (UAnimBlueprintGeneratedClass* AnimClass = Cast<UAnimBlueprintGeneratedClass>(Context.GetCurrentlyGeneratedClass()))
-// 			{
-// 				for (int32 i = 0; i < AnimClass->GetExposedValueHandlers().Num(); ++i)
-// 				{
-// 					if (AnimClass->GetExposedValueHandlers()[i].ValueHandlerNodeProperty == InProperty)
-// 					{
-// 						const FString ClassName = FEmitHelper::GetCppName(Context.GetCurrentlyGeneratedClass());
-// 						const FString MemberName = FEmitHelper::GetCppName(const_cast<FProperty*>(InProperty));
-// 						Context.Body.AddLine(FString::Printf(TEXT("%s.SetExposedValueHandler(&CastChecked<UAnimClassData>(CastChecked<UDynamicClass>(%s::StaticClass())->%s)->GetExposedValueHandlers()[%d]);"), *MemberName, *ClassName, GET_MEMBER_NAME_STRING_CHECKED(UDynamicClass, AnimClassImplementation), i));
-// 
-// 						break;
-// 					}
-// 				}
-// 			}
+			if (UAnimBlueprintGeneratedClass* AnimClass = Cast<UAnimBlueprintGeneratedClass>(Context.GetCurrentlyGeneratedClass()))
+			{
+				for (int32 i = 0; i < AnimClass->GetExposedValueHandlers().Num(); ++i)
+				{
+					if (AnimClass->GetExposedValueHandlers()[i].ValueHandlerNodeProperty == InProperty)
+					{
+						const FString ClassName = FEmitHelper::GetCppName(Context.GetCurrentlyGeneratedClass());
+						const FString MemberName = FEmitHelper::GetCppName(const_cast<FProperty*>(InProperty));
+						Context.Body.AddLine(FString::Printf(TEXT("%s.SetExposedValueHandler(&CastChecked<UAnimClassData>(CastChecked<UDynamicClass>(%s::StaticClass())->%s)->GetExposedValueHandlers()[%d]);"), *MemberName, *ClassName, GET_MEMBER_NAME_STRING_CHECKED(UDynamicClass, AnimClassImplementation), i));
+
+						break;
+					}
+				}
+			}
 
 			Context.Body.DecreaseIndent();
 			Context.Body.AddLine(TEXT("}"));
diff --git a/Engine/Source/Editor/AnimGraph/Private/AnimBlueprintCompilationContext.cpp b/Engine/Source/Editor/AnimGraph/Private/AnimBlueprintCompilationContext.cpp
--- a/Engine/Source/Editor/AnimGraph/Private/AnimBlueprintCompilationContext.cpp
+++ b/Engine/Source/Editor/AnimGraph/Private/AnimBlueprintCompilationContext.cpp
@@ -2,7 +2,6 @@
 
 #include "AnimBlueprintCompilationContext.h"
 #include "AnimBlueprintCompiler.h"
-#include "AnimBlueprintExtension.h"
 
 TUniquePtr<IAnimBlueprintCompilationContext> IAnimBlueprintCompilationContext::Get(FKismetCompilerContext& InKismetCompiler)
 {
@@ -94,19 +93,9 @@ TArrayView<const FName> FAnimBlueprintCompilationContext::GetAttributesFromNodeI
 	return CompilerContext->GetAttributesFromNode(InNode);
 }
 
-bool FAnimBlueprintCompilationContext::IsAnimGraphNodeFoldedImpl(UAnimGraphNode_Base* InNode) const
+IAnimBlueprintCompilerHandler* FAnimBlueprintCompilationContext::GetHandlerInternal(FName InName) const
 {
-	return CompilerContext->IsAnimGraphNodeFolded(InNode);
-}
-
-const IAnimBlueprintCompilationContext::FFoldedPropertyRecord* FAnimBlueprintCompilationContext::GetFoldedPropertyRecordImpl(UAnimGraphNode_Base* InNode, FName InPropertyName) const
-{
-	return CompilerContext->GetFoldedPropertyRecord(InNode, InPropertyName);
-}
-
-const FStructProperty* FAnimBlueprintCompilationContext::GetMutableDataPropertyImpl() const
-{
-	return CompilerContext->NewMutablesProperty;
+	return CompilerContext->AnimBlueprintCompilerHandlerCollection.GetHandler<IAnimBlueprintCompilerHandler>(InName);
 }
 
 FKismetCompilerContext* FAnimBlueprintCompilationContext::GetKismetCompiler() const
@@ -119,10 +108,10 @@ FCompilerResultsLog& FAnimBlueprintCompilationBracketContext::GetMessageLogImpl(
 	return CompilerContext->MessageLog;
 }
 
-const TMap<UAnimGraphNode_Base*, int32>& FAnimBlueprintCompilationBracketContext::GetAllocatedAnimNodeIndicesImpl() const
+IAnimBlueprintCompilerHandler* FAnimBlueprintCompilationBracketContext::GetHandlerInternal(FName InName) const
 {
-	return CompilerContext->AllocatedAnimNodeIndices;
-}
+	return CompilerContext->AnimBlueprintCompilerHandlerCollection.GetHandler<IAnimBlueprintCompilerHandler>(InName);
+}	
 
 FCompilerResultsLog& FAnimBlueprintCopyTermDefaultsContext::GetMessageLogImpl() const
 {
@@ -148,3 +137,8 @@ const FKismetCompilerOptions& FAnimBlueprintPostExpansionStepContext::GetCompile
 {
 	return CompilerContext->CompileOptions;
 }
+
+IAnimBlueprintCompilerHandler* FAnimBlueprintPostExpansionStepContext::GetHandlerInternal(FName InName) const
+{
+	return CompilerContext->AnimBlueprintCompilerHandlerCollection.GetHandler<IAnimBlueprintCompilerHandler>(InName);
+}
diff --git a/Engine/Source/Editor/AnimGraph/Private/AnimBlueprintCompilationContext.h b/Engine/Source/Editor/AnimGraph/Private/AnimBlueprintCompilationContext.h
--- a/Engine/Source/Editor/AnimGraph/Private/AnimBlueprintCompilationContext.h
+++ b/Engine/Source/Editor/AnimGraph/Private/AnimBlueprintCompilationContext.h
@@ -7,12 +7,8 @@
 #include "IAnimBlueprintCompilationBracketContext.h"
 #include "IAnimBlueprintCopyTermDefaultsContext.h"
 #include "IAnimBlueprintPostExpansionStepContext.h"
-#include "IAnimBlueprintNodeOverrideAssetsContext.h"
-#include "Animation/AnimationAsset.h"
 
 class FAnimBlueprintCompilerContext;
-class UAnimBlueprintExtension;
-struct FAnimNode_Base;
 
 class FAnimBlueprintCompilationContext : public IAnimBlueprintCompilationContext
 {
@@ -41,9 +37,7 @@ private:
 	virtual const TMap<UAnimGraphNode_Base*, FProperty*>& GetAllocatedPropertiesByNodeImpl() const override;
 	virtual void AddAttributesToNodeImpl(UAnimGraphNode_Base* InNode, TArrayView<const FName> InAttributes) const override;
 	virtual TArrayView<const FName> GetAttributesFromNodeImpl(UAnimGraphNode_Base* InNode) const override;
-	virtual bool IsAnimGraphNodeFoldedImpl(UAnimGraphNode_Base* InNode) const override;
-	virtual const IAnimBlueprintCompilationContext::FFoldedPropertyRecord* GetFoldedPropertyRecordImpl(UAnimGraphNode_Base* InNode, FName InPropertyName) const override;
-	virtual const FStructProperty* GetMutableDataPropertyImpl() const override;
+	virtual IAnimBlueprintCompilerHandler* GetHandlerInternal(FName InName) const override;
 	virtual FKismetCompilerContext* GetKismetCompiler() const override;
 
 	FAnimBlueprintCompilerContext* CompilerContext;
@@ -59,8 +53,8 @@ private:
 	{}
 
 	virtual FCompilerResultsLog& GetMessageLogImpl() const override;
-	virtual const TMap<UAnimGraphNode_Base*, int32>& GetAllocatedAnimNodeIndicesImpl() const override;
-	
+	virtual IAnimBlueprintCompilerHandler* GetHandlerInternal(FName InName) const override;
+
 	FAnimBlueprintCompilerContext* CompilerContext;
 };
 
@@ -102,29 +96,6 @@ private:
 	int32 NodePropertyIndex;
 };
 
-class FAnimBlueprintExtensionCopyTermDefaultsContext : public IAnimBlueprintExtensionCopyTermDefaultsContext
-{
-private:
-	friend class FAnimBlueprintCompilerContext;
-
-	FAnimBlueprintExtensionCopyTermDefaultsContext(const FProperty* InTargetProperty, uint8* InDestinationPtr, const uint8* InSourcePtr, int32 InNodePropertyIndex)
-        : TargetProperty(InTargetProperty)
-        , DestinationPtr(InDestinationPtr)
-        , SourcePtr(InSourcePtr)
-        , NodePropertyIndex(InNodePropertyIndex)
-	{}
-
-	virtual const FProperty* GetTargetPropertyImpl() const override { return TargetProperty; }
-	virtual uint8* GetDestinationPtrImpl() const override { return DestinationPtr; }
-	virtual const uint8* GetSourcePtrImpl() const override { return SourcePtr; }
-	virtual int32 GetNodePropertyIndexImpl() const override { return NodePropertyIndex; }
-
-	const FProperty* TargetProperty;
-	uint8* DestinationPtr;
-	const uint8* SourcePtr;
-	int32 NodePropertyIndex;
-};
-
 class FAnimBlueprintPostExpansionStepContext : public IAnimBlueprintPostExpansionStepContext
 {
 private:
@@ -137,33 +108,7 @@ private:
 	virtual FCompilerResultsLog& GetMessageLogImpl() const override;
 	virtual UEdGraph* GetConsolidatedEventGraphImpl() const override;
 	virtual const FKismetCompilerOptions& GetCompileOptionsImpl() const override;
+	virtual IAnimBlueprintCompilerHandler* GetHandlerInternal(FName InName) const override;
 
 	FAnimBlueprintCompilerContext* CompilerContext;
-};
-
-class FAnimBlueprintNodeOverrideAssetsContext : public IAnimBlueprintNodeOverrideAssetsContext
-{
-private:
-	friend class FAnimBlueprintCompilerContext;
-	
-	FAnimBlueprintNodeOverrideAssetsContext(FAnimNode_Base* InAnimNode, UScriptStruct* InAnimNodeType)
-		: AnimNode(InAnimNode)
-		, AnimNodeType(InAnimNodeType)
-	{}
-
-	// Add an asset to this context
-	void AddAsset(UObject* InAsset) { Assets.Add(InAsset); }
-
-	virtual FAnimNode_Base* GetAnimNodeImpl() const override { return AnimNode; }
-	virtual UScriptStruct* GetAnimNodeTypeImpl() const override { return AnimNodeType; }
-	virtual TArrayView<UObject* const> GetAssetsImpl() const override { return MakeArrayView(Assets); }
-
-	// The anim node we are targeting
-	FAnimNode_Base* AnimNode;
-
-	// The type of the anim node we are targeting
-	UScriptStruct* AnimNodeType;
-	
-	// All the assets to override
-	TArray<UObject*> Assets;
 };
\ No newline at end of file
diff --git a/Engine/Source/Editor/AnimGraph/Private/AnimBlueprintCompiler.cpp b/Engine/Source/Editor/AnimGraph/Private/AnimBlueprintCompiler.cpp
--- a/Engine/Source/Editor/AnimGraph/Private/AnimBlueprintCompiler.cpp
+++ b/Engine/Source/Editor/AnimGraph/Private/AnimBlueprintCompiler.cpp
@@ -1,12 +1,20 @@
 // Copyright Epic Games, Inc. All Rights Reserved.
 
 #include "AnimBlueprintCompiler.h"
+#include "UObject/UObjectHash.h"
 #include "Animation/AnimInstance.h"
 #include "EdGraphUtilities.h"
 #include "K2Node_CallFunction.h"
+#include "K2Node_StructMemberGet.h"
+#include "K2Node_BreakStruct.h"
+#include "K2Node_MakeStruct.h"
+#include "K2Node_CallArrayFunction.h"
+#include "K2Node_CustomEvent.h"
 #include "K2Node_Knot.h"
 #include "K2Node_StructMemberSet.h"
 #include "K2Node_VariableGet.h"
+#include "K2Node_VariableSet.h"
+#include "K2Node_GetArrayItem.h"
 
 #include "AnimationGraphSchema.h"
 #include "Kismet/BlueprintFunctionLibrary.h"
@@ -16,6 +24,8 @@
 #include "Kismet2/KismetReinstanceUtilities.h"
 #include "AnimGraphNode_Root.h"
 #include "Animation/AnimNode_CustomProperty.h"
+#include "Animation/AnimNode_LinkedAnimGraph.h"
+#include "AnimGraphNode_LinkedAnimGraph.h"
 #include "AnimationEditorUtils.h"
 #include "AnimationGraph.h"
 #include "AnimBlueprintPostCompileValidation.h" 
@@ -24,38 +34,25 @@
 #include "K2Node_FunctionResult.h"
 #include "AnimGraphNode_LinkedAnimLayer.h"
 #include "String/ParseTokens.h"
-#include "Algo/Sort.h"
+#include "Algo/Transform.h"
+#include "Algo/Accumulate.h"
 #include "IClassVariableCreator.h"
 #include "AnimBlueprintGeneratedClassCompiledData.h"
 #include "AnimBlueprintCompilationContext.h"
-#include "AnimBlueprintCompilerCreationContext.h"
 #include "AnimBlueprintVariableCreationContext.h"
-#include "Animation/AnimSubsystemInstance.h"
-#include "AnimBlueprintExtension.h"
-#include "UObject/UE5MainStreamObjectVersion.h"
 
 #define LOCTEXT_NAMESPACE "AnimBlueprintCompiler"
 
-DECLARE_CYCLE_STAT(TEXT("Merge Ubergraph Pages In"), EAnimBlueprintCompilerStats_MergeUbergraphPagesIn, STATGROUP_KismetCompiler);
-DECLARE_CYCLE_STAT(TEXT("Process All Animation Nodes"), EAnimBlueprintCompilerStats_ProcessAllAnimationNodes, STATGROUP_KismetCompiler);
-DECLARE_CYCLE_STAT(TEXT("Process Animation Nodes"), EAnimBlueprintCompilerStats_ProcessAnimationNodes, STATGROUP_KismetCompiler);
-DECLARE_CYCLE_STAT(TEXT("Expand Split Pins"), EAnimBlueprintCompilerStats_ExpandSplitPins, STATGROUP_KismetCompiler);
-DECLARE_CYCLE_STAT(TEXT("Move Graphs"), EAnimBlueprintCompilerStats_MoveGraphs, STATGROUP_KismetCompiler);
-DECLARE_CYCLE_STAT(TEXT("Clone Graph"), EAnimBlueprintCompilerStats_CloneGraph, STATGROUP_KismetCompiler);
-DECLARE_CYCLE_STAT(TEXT("Process Animation Node"), EAnimBlueprintCompilerStats_ProcessAnimationNode, STATGROUP_KismetCompiler);
-DECLARE_CYCLE_STAT(TEXT("Gather Fold Records"), EAnimBlueprintCompilerStats_GatherFoldRecords, STATGROUP_KismetCompiler);
-
 //////////////////////////////////////////////////////////////////////////
 // FAnimBlueprintCompiler
 
 FAnimBlueprintCompilerContext::FAnimBlueprintCompilerContext(UAnimBlueprint* SourceSketch, FCompilerResultsLog& InMessageLog, const FKismetCompilerOptions& InCompileOptions)
 	: FKismetCompilerContext(SourceSketch, InMessageLog, InCompileOptions)
-	, NewAnimBlueprintConstants(nullptr)
-	, NewAnimBlueprintMutables(nullptr)
-	, NewMutablesProperty(nullptr)
 	, AnimBlueprint(SourceSketch)
 	, bIsDerivedAnimBlueprint(false)
 {
+	AnimBlueprintCompilerHandlerCollection.Initialize(this);
+
 	// Add the animation graph schema to skip default function processing on them
 	KnownGraphSchemas.AddUnique(UAnimationGraphSchema::StaticClass());
 
@@ -68,13 +65,6 @@ FAnimBlueprintCompilerContext::FAnimBlueprintCompilerContext(UAnimBlueprint* Sou
 		}
 	}
 
-	// If we need to, refresh all extensions here
-	if(AnimBlueprint->bRefreshExtensions)
-	{
-		UAnimBlueprintExtension::RefreshExtensions(AnimBlueprint);
-		AnimBlueprint->bRefreshExtensions = false;
-	}
-	
 	if (AnimBlueprint->HasAnyFlags(RF_NeedPostLoad))
 	{
 		//Compilation during loading .. need to verify node guids as some anim blueprints have duplicated guids
@@ -143,12 +133,6 @@ FAnimBlueprintCompilerContext::FAnimBlueprintCompilerContext(UAnimBlueprint* Sou
 		}
 	}
 
-	FAnimBlueprintCompilerCreationContext CreationContext(this);
-	UAnimBlueprintExtension::ForEachExtension(AnimBlueprint, [&CreationContext](UAnimBlueprintExtension* InExtension)
-	{
-		InExtension->BeginCompilation(CreationContext);
-	});
-
 	// Determine if there is an anim blueprint in the ancestry of this class
 	bIsDerivedAnimBlueprint = UAnimBlueprint::FindRootAnimBlueprint(AnimBlueprint) != NULL;
 
@@ -160,11 +144,6 @@ FAnimBlueprintCompilerContext::FAnimBlueprintCompilerContext(UAnimBlueprint* Sou
 FAnimBlueprintCompilerContext::~FAnimBlueprintCompilerContext()
 {
 	DestroyAnimGraphStubFunctions();
-
-	UAnimBlueprintExtension::ForEachExtension(AnimBlueprint, [](UAnimBlueprintExtension* InExtension)
-	{
-	    InExtension->EndCompilation();
-	});
 }
 
 void FAnimBlueprintCompilerContext::ForAllSubGraphs(UEdGraph* InGraph, TFunctionRef<void(UEdGraph*)> InPerGraphFunction)
@@ -220,8 +199,6 @@ UEdGraphSchema_K2* FAnimBlueprintCompilerContext::CreateSchema()
 
 void FAnimBlueprintCompilerContext::ProcessAnimationNode(UAnimGraphNode_Base* VisualAnimNode)
 {
-	BP_SCOPED_COMPILER_EVENT_STAT(EAnimBlueprintCompilerStats_ProcessAnimationNode);
-	
 	// Early out if this node has already been processed
 	if (AllocatedAnimNodes.Contains(VisualAnimNode))
 	{
@@ -230,7 +207,7 @@ void FAnimBlueprintCompilerContext::ProcessAnimationNode(UAnimGraphNode_Base* Vi
 
 	// Make sure the visual node has a runtime node template
 	const UScriptStruct* NodeType = VisualAnimNode->GetFNodeType();
-	if (NodeType == nullptr)
+	if (NodeType == NULL)
 	{
 		MessageLog.Error(TEXT("@@ has no animation node member"), VisualAnimNode);
 		return;
@@ -258,13 +235,11 @@ void FAnimBlueprintCompilerContext::ProcessAnimationNode(UAnimGraphNode_Base* Vi
 
 	FStructProperty* NewProperty = CastField<FStructProperty>(CreateVariable(FName(*NodeVariableName), NodeVariableType));
 
-	if (NewProperty == nullptr)
+	if (NewProperty == NULL)
 	{
 		MessageLog.Error(TEXT("Failed to create node property for @@"), VisualAnimNode);
 	}
 
-	GatherFoldRecordsForAnimationNode(NodeType, NewProperty, VisualAnimNode);
-
 	// Register this node with the compile-time data structures
 	const int32 AllocatedIndex = AllocateNodeIndexCounter++;
 	AllocatedAnimNodes.Add(VisualAnimNode, NewProperty);
@@ -276,7 +251,6 @@ void FAnimBlueprintCompilerContext::ProcessAnimationNode(UAnimGraphNode_Base* Vi
 	SourceNodeToProcessedNodeMap.Add(TrueSourceObject, VisualAnimNode);
 
 	// Register the slightly more permanent debug information
-	UAnimBlueprintGeneratedClass* NewAnimBlueprintClass = GetNewAnimBlueprintClass();
 	FAnimBlueprintDebugData& NewAnimBlueprintDebugData = NewAnimBlueprintClass->GetAnimBlueprintDebugData();
 	NewAnimBlueprintDebugData.NodePropertyToIndexMap.Add(TrueSourceObject, AllocatedIndex);
 	NewAnimBlueprintDebugData.NodeGuidToIndexMap.Add(TrueSourceObject->NodeGuid, AllocatedIndex);
@@ -288,178 +262,6 @@ void FAnimBlueprintCompilerContext::ProcessAnimationNode(UAnimGraphNode_Base* Vi
 	VisualAnimNode->ProcessDuringCompilation(CompilerContext, CompiledData);
 }
 
-void FAnimBlueprintCompilerContext::ProcessExtensions()
-{
-	TArray<UAnimBlueprintExtension*> Extensions = UAnimBlueprintExtension::GetExtensions(AnimBlueprint);
-	
-	// Sort extensions by class name (for determinism)
-	Algo::Sort(Extensions, [](UAnimBlueprintExtension* InExtensionA, UAnimBlueprintExtension* InExtensionB)
-	{
-		return InExtensionA->GetClass()->GetName() < InExtensionB->GetClass()->GetName();
-	});
-
-	// Process all gathered class extensions
-	for(UAnimBlueprintExtension* Extension : Extensions)
-	{
-		const FString ExtensionVariableName = ClassScopeNetNameMap.MakeValidName(Extension);
-
-		const UScriptStruct* InstanceDataType = Extension->GetInstanceDataType();
-		check(InstanceDataType->IsChildOf(FAnimSubsystemInstance::StaticStruct()));
-
-		const UScriptStruct* ClassDataType = Extension->GetClassDataType();
-		check(ClassDataType->IsChildOf(FAnimSubsystem::StaticStruct()));
-		
-		// Skip creating any properties if both are the default
-		if(InstanceDataType != FAnimSubsystemInstance::StaticStruct() || ClassDataType != FAnimSubsystem::StaticStruct())
-		{
-			// Process instance data (mutable)
-			{
-				FEdGraphPinType ExtensionVariableType;
-				ExtensionVariableType.PinCategory = UAnimationGraphSchema::PC_Struct;
-				ExtensionVariableType.PinSubCategoryObject = MakeWeakObjectPtr(const_cast<UScriptStruct*>(InstanceDataType));
-
-				FStructProperty* NewProperty = CastField<FStructProperty>(CreateVariable(FName(*ExtensionVariableName), ExtensionVariableType));
-				if (NewProperty == nullptr)
-				{
-					MessageLog.Error(*FText::Format(LOCTEXT("ExtensionPropertyCreationFailed", "Failed to create extension property for '{0}'"), FText::FromString(Extension->GetName())).ToString());
-				}
-				else
-				{
-					ExtensionToInstancePropertyMap.Add(Extension, NewProperty);
-					InstancePropertyToExtensionMap.Add(NewProperty, Extension);
-				}
-			}
-
-			// Process class data (constants)
-			{
-				FEdGraphPinType ExtensionVariableType;
-				ExtensionVariableType.PinCategory = UAnimationGraphSchema::PC_Struct;
-				ExtensionVariableType.PinSubCategoryObject = MakeWeakObjectPtr(const_cast<UScriptStruct*>(ClassDataType));
-
-				FStructProperty* NewProperty = CastField<FStructProperty>(CreateStructVariable(NewAnimBlueprintConstants, FName(*ExtensionVariableName), ExtensionVariableType));
-				if (NewProperty == nullptr)
-				{
-					MessageLog.Error(*FText::Format(LOCTEXT("ExtensionPropertyCreationFailed", "Failed to create extension property for '{0}'"), FText::FromString(Extension->GetName())).ToString());
-				}
-				else
-				{
-					ExtensionToClassPropertyMap.Add(Extension, NewProperty);
-					ClassPropertyToExtensionMap.Add(NewProperty, Extension);
-				}
-			}
-		}
-	}
-}
-
-void FAnimBlueprintCompilerContext::GatherFoldRecordsForAnimationNode(const UScriptStruct* InNodeType, FStructProperty* InNodeProperty, UAnimGraphNode_Base* InVisualAnimNode)
-{
-	BP_SCOPED_COMPILER_EVENT_STAT(EAnimBlueprintCompilerStats_GatherFoldRecords);
-
-	static const FName NAME_FoldProperty("FoldProperty");
-
-	// Run through node properties to see if any are eligible for folding
-	for(TFieldIterator<FProperty> It(InNodeType); It; ++It)
-	{
-		FProperty* SubProperty = *It;
-
-		if(FArrayProperty* ArrayProperty = CastField<FArrayProperty>(SubProperty))
-		{
-			bool bAllPinsExposed = true;
-			bool bAllPinsDisconnected = true;
-
-			// If a value is exposed on a pin but disconnected, push the value to the (intermediate) node here to simplify later logic
-			if(SubProperty->HasAnyPropertyFlags(CPF_Edit|CPF_BlueprintVisible))
-			{
-				FStructProperty* AnimGraphNodeProperty = InVisualAnimNode->GetFNodeProperty();
-
-				// Check the anim node property is contained in the anim graph node
-				check(AnimGraphNodeProperty->GetOwner<UClass>() && InVisualAnimNode->GetClass()->IsChildOf(AnimGraphNodeProperty->GetOwner<UClass>()));
-				const void* Node = AnimGraphNodeProperty->ContainerPtrToValuePtr<void>(InVisualAnimNode);
-
-				// Check the anim node's property is contained in the anim node
-				check(SubProperty->GetOwner<UStruct>() && AnimGraphNodeProperty->Struct->IsChildOf(SubProperty->GetOwner<UStruct>()));
-				const void* TargetPtr = SubProperty->ContainerPtrToValuePtr<void>(Node);
-
-				// Run through each array element - we can only fold array properties if all values are constants
-				FScriptArrayHelper ArrayHelper(ArrayProperty, TargetPtr);
-
-				for(int32 ArrayIndex = 0; ArrayIndex < ArrayHelper.Num(); ++ArrayIndex)
-				{
-					const FString ArrayElementPinName = SubProperty->GetName() + FString::Printf(TEXT("_%d"), ArrayIndex);
-					UEdGraphPin* Pin = InVisualAnimNode->FindPin(ArrayElementPinName);
-					const bool bExposedOnPin = Pin != nullptr;
-					const bool bPinConnected = InVisualAnimNode->IsPinExposedAndLinked(ArrayElementPinName, EGPD_Input);
-					const bool bPinBound = InVisualAnimNode->IsPinExposedAndBound(SubProperty->GetName(), EGPD_Input);
-					
-					if(bExposedOnPin)
-					{
-						if(!(bPinConnected || bPinBound))
-						{
-							check(Pin);
-
-							if(!FBlueprintEditorUtils::PropertyValueFromString_Direct(ArrayProperty->Inner, Pin->GetDefaultAsString(), ArrayHelper.GetRawPtr(ArrayIndex)))
-							{
-								MessageLog.Warning(TEXT("Unable to push default value for array pin @@ on @@"), Pin, InVisualAnimNode);
-							}
-						}
-						else
-						{
-							bAllPinsDisconnected = false;
-						}
-					}
-					else
-					{
-						bAllPinsExposed = false;
-					}
-				}
-			}
-
-			if(SubProperty->HasMetaData(NAME_FoldProperty))
-			{
-				// Add folding candidate
-				AddFoldedPropertyRecord(InVisualAnimNode, InNodeProperty, SubProperty, bAllPinsExposed, !bAllPinsDisconnected);
-			}
-		}
-		else
-		{
-			UEdGraphPin* Pin = InVisualAnimNode->FindPin(SubProperty->GetName());
-			const bool bExposedOnPin = Pin != nullptr;
-			const bool bPinConnected = InVisualAnimNode->IsPinExposedAndLinked(SubProperty->GetName(), EGPD_Input);
-			const bool bPinBound = InVisualAnimNode->IsPinExposedAndBound(SubProperty->GetName(), EGPD_Input);
-
-			// If a value is exposed on a pin but disconnected, push the value to the (intermediate) node here to simplify later logic
-			if(SubProperty->HasAnyPropertyFlags(CPF_Edit|CPF_BlueprintVisible))
-			{
-				if(bExposedOnPin && !(bPinConnected || bPinBound))
-				{
-					check(Pin);
-
-					FStructProperty* AnimGraphNodeProperty = InVisualAnimNode->GetFNodeProperty();
-
-					// Check the anim node property is contained in the anim graph node
-					check(AnimGraphNodeProperty->GetOwner<UClass>() && InVisualAnimNode->GetClass()->IsChildOf(AnimGraphNodeProperty->GetOwner<UClass>()));
-					const void* Node = AnimGraphNodeProperty->ContainerPtrToValuePtr<void>(InVisualAnimNode);
-
-					// Check the anim node's property is contained in the anim node
-					check(SubProperty->GetOwner<UStruct>() && AnimGraphNodeProperty->Struct->IsChildOf(SubProperty->GetOwner<UStruct>()));
-					const void* TargetPtr = SubProperty->ContainerPtrToValuePtr<void>(Node);
-
-					if(!FBlueprintEditorUtils::PropertyValueFromString_Direct(SubProperty, Pin->GetDefaultAsString(), (uint8*)TargetPtr))
-					{
-						MessageLog.Warning(TEXT("Unable to push default value for pin @@ on @@"), Pin, InVisualAnimNode);
-					}
-				}
-			}
-
-			if(SubProperty->HasMetaData(NAME_FoldProperty))
-			{
-				// Add folding candidate
-				AddFoldedPropertyRecord(InVisualAnimNode, InNodeProperty, SubProperty, bExposedOnPin, bPinConnected || bPinBound);
-			}
-		}
-	}
-}
-
 int32 FAnimBlueprintCompilerContext::GetAllocationIndexOfNode(UAnimGraphNode_Base* VisualAnimNode)
 {
 	ProcessAnimationNode(VisualAnimNode);
@@ -476,12 +278,9 @@ bool FAnimBlueprintCompilerContext::ShouldForceKeepNode(const UEdGraphNode* Node
 
 void FAnimBlueprintCompilerContext::PostExpansionStep(const UEdGraph* Graph)
 {
-	FAnimBlueprintGeneratedClassCompiledData CompiledData(GetNewAnimBlueprintClass());
+	FAnimBlueprintGeneratedClassCompiledData CompiledData(NewAnimBlueprintClass);
 	FAnimBlueprintPostExpansionStepContext CompilerContext(this);
-	UAnimBlueprintExtension::ForEachExtension(AnimBlueprint, [Graph, &CompiledData, &CompilerContext](UAnimBlueprintExtension* InExtension)
-	{
-		InExtension->PostExpansionStep(Graph, CompilerContext, CompiledData);
-	});
+	OnPostExpansionStepDelegate.Broadcast(Graph, CompilerContext, CompiledData);
 }
 
 void FAnimBlueprintCompilerContext::PruneIsolatedAnimationNodes(const TArray<UAnimGraphNode_Base*>& RootSet, TArray<UAnimGraphNode_Base*>& GraphNodes)
@@ -546,8 +345,6 @@ void FAnimBlueprintCompilerContext::PruneIsolatedAnimationNodes(const TArray<UAn
 
 void FAnimBlueprintCompilerContext::ProcessAnimationNodes(TArray<UAnimGraphNode_Base*>& AnimNodeList)
 {
-	BP_SCOPED_COMPILER_EVENT_STAT(EAnimBlueprintCompilerStats_ProcessAnimationNodes);
-	
 	// Process the remaining nodes
 	for (UAnimGraphNode_Base* AnimNode : AnimNodeList)
 	{
@@ -624,8 +421,6 @@ void FAnimBlueprintCompilerContext::GetLinkedAnimNodes_ProcessAnimNode(UAnimGrap
 
 void FAnimBlueprintCompilerContext::ProcessAllAnimationNodes()
 {
-	BP_SCOPED_COMPILER_EVENT_STAT(EAnimBlueprintCompilerStats_ProcessAllAnimationNodes);
-	
 	// Validate that we have a skeleton
 	if ((AnimBlueprint->TargetSkeleton == nullptr) && !AnimBlueprint->bIsNewlyCreated)
 	{
@@ -669,26 +464,14 @@ void FAnimBlueprintCompilerContext::ProcessAllAnimationNodes()
 		// Validate the graph
 		ValidateGraphIsWellFormed(ConsolidatedEventGraph);
 
-		FAnimBlueprintGeneratedClassCompiledData CompiledData(GetNewAnimBlueprintClass());
+		FAnimBlueprintGeneratedClassCompiledData CompiledData(NewAnimBlueprintClass);
 		FAnimBlueprintCompilationContext CompilerContext(this);
-		UAnimBlueprintExtension::ForEachExtension(AnimBlueprint, [&AllSubGraphsAnimNodeList, &CompiledData, &CompilerContext](UAnimBlueprintExtension* InExtension)
-		{
-			InExtension->PreProcessAnimationNodes(AllSubGraphsAnimNodeList, CompilerContext, CompiledData);
-		});
+		OnPreProcessAnimationNodesDelegate.Broadcast(AllSubGraphsAnimNodeList, CompilerContext, CompiledData);
 
 		// Process the animation nodes
 		ProcessAnimationNodes(RootAnimNodeList);
 
-		// Process any extensions
-		ProcessExtensions();
-		
-		// Fold any constants
-		ProcessFoldedPropertyRecords();
-
-		UAnimBlueprintExtension::ForEachExtension(AnimBlueprint, [&AllSubGraphsAnimNodeList, &CompiledData, &CompilerContext](UAnimBlueprintExtension* InExtension)
-		{
-			InExtension->PostProcessAnimationNodes(AllSubGraphsAnimNodeList, CompilerContext, CompiledData);
-		});	
+		OnPostProcessAnimationNodesDelegate.Broadcast(AllSubGraphsAnimNodeList, CompilerContext, CompiledData);
 	}
 	else
 	{
@@ -701,13 +484,9 @@ void FAnimBlueprintCompilerContext::CopyTermDefaultsToDefaultObject(UObject* Def
 	Super::CopyTermDefaultsToDefaultObject(DefaultObject);
 
 	UAnimInstance* DefaultAnimInstance = Cast<UAnimInstance>(DefaultObject);
-	UAnimBlueprintGeneratedClass* NewAnimBlueprintClass = GetNewAnimBlueprintClass();
 
 	if (bIsDerivedAnimBlueprint && DefaultAnimInstance)
 	{
-		// Make sure our sparse class data is re-created
-		RecreateSparseClassData();
-		
 		// If we are a derived animation graph; apply any stored overrides.
 		// Restore values from the root class to catch values where the override has been removed.
 		UAnimBlueprintGeneratedClass* RootAnimClass = NewAnimBlueprintClass;
@@ -715,15 +494,6 @@ void FAnimBlueprintCompilerContext::CopyTermDefaultsToDefaultObject(UObject* Def
 		{
 			RootAnimClass = NextClass;
 		}
-
-		// Copy constant/folded data from root class, remapping the class
-		NewAnimBlueprintClass->AnimNodeData = RootAnimClass->AnimNodeData;
-		NewAnimBlueprintClass->NodeTypeMap = RootAnimClass->NodeTypeMap;
-		for(FAnimNodeData& NodeData : NewAnimBlueprintClass->AnimNodeData)
-		{
-			NodeData.AnimClassInterface = NewAnimBlueprintClass;
-		}
-		
 		UObject* RootDefaultObject = RootAnimClass->GetDefaultObject();
 
 		for (TFieldIterator<FProperty> It(RootAnimClass); It; ++It)
@@ -743,9 +513,6 @@ void FAnimBlueprintCompilerContext::CopyTermDefaultsToDefaultObject(UObject* Def
 				}
 			}
 		}
-
-		// Re-initialize node data tables (they would be overwritten in the loop above)
-		NewAnimBlueprintClass->InitializeAnimNodeData(DefaultObject);
 	}
 
 	// Give game-specific logic a chance to replace animations
@@ -763,17 +530,10 @@ void FAnimBlueprintCompilerContext::CopyTermDefaultsToDefaultObject(UObject* Def
 		{
 			if (Override->NewAsset)
 			{
-				int32 NodeIndex = NewAnimBlueprintClass->GetNodeIndexFromGuid(Override->ParentNodeGuid, EPropertySearchMode::Hierarchy);
-				if(NodeIndex != INDEX_NONE)
+				FAnimNode_Base* BaseNode = NewAnimBlueprintClass->GetPropertyInstance<FAnimNode_Base>(DefaultAnimInstance, Override->ParentNodeGuid, EPropertySearchMode::Hierarchy);
+				if (BaseNode)
 				{
-					const UAnimGraphNode_Base* GraphNode = Cast<UAnimGraphNode_Base>(NewAnimBlueprintClass->GetVisualNodeFromNodePropertyIndex(NodeIndex, EPropertySearchMode::Hierarchy));
-					FAnimNode_Base* BaseNode = NewAnimBlueprintClass->GetPropertyInstance<FAnimNode_Base>(DefaultAnimInstance, Override->ParentNodeGuid, EPropertySearchMode::Hierarchy);
-					if (GraphNode && BaseNode)
-					{
-						FAnimBlueprintNodeOverrideAssetsContext Context(BaseNode, GraphNode->GetFNodeType());
-						Context.AddAsset(Override->NewAsset);
-						GraphNode->OverrideAssets(Context);
-					}
+					BaseNode->OverrideAsset(Override->NewAsset);
 				}
 			}
 		}
@@ -790,28 +550,6 @@ void FAnimBlueprintCompilerContext::CopyTermDefaultsToDefaultObject(UObject* Def
 		FAnimBlueprintGeneratedClassCompiledData CompiledData(NewAnimBlueprintClass);
 		FAnimBlueprintCopyTermDefaultsContext CompilerContext(this);
 
-		// Initialize extensions from their templates
-		for (TFieldIterator<FStructProperty> It(DefaultAnimInstance->GetClass(), EFieldIteratorFlags::ExcludeSuper); It; ++It)
-		{
-			FStructProperty* TargetProperty = *It;
-
-			if (UAnimBlueprintExtension* Extension = InstancePropertyToExtensionMap.FindRef(TargetProperty))
-			{
-				uint8* DestinationPtr = TargetProperty->ContainerPtrToValuePtr<uint8>(DefaultAnimInstance);
-				const uint8* SourcePtr = nullptr;
-				const FStructProperty* SourceExtensionProperty = Extension->GetInstanceDataProperty();
-				
-				if(SourceExtensionProperty)
-				{
-					check(TargetProperty->Struct == SourceExtensionProperty->Struct);				
-					SourcePtr = SourceExtensionProperty->ContainerPtrToValuePtr<uint8>(Extension);
-				}
-				
-				FAnimBlueprintExtensionCopyTermDefaultsContext ExtensionContext(TargetProperty, DestinationPtr, SourcePtr, LinkIndexCount);
-				Extension->CopyTermDefaultsToDefaultObject(DefaultAnimInstance, CompilerContext, ExtensionContext);
-			}
-		}
-		
 		// Initialize animation nodes from their templates
 		for (TFieldIterator<FProperty> It(DefaultAnimInstance->GetClass(), EFieldIteratorFlags::ExcludeSuper); It; ++It)
 		{
@@ -835,80 +573,6 @@ void FAnimBlueprintCompilerContext::CopyTermDefaultsToDefaultObject(UObject* Def
 			}
 		}
 
-		// Applies a set of folded property records to a data area (i.e. the constant or mutable structs)
-		auto PatchDataArea = [](void* InData, UScriptStruct* InStruct, const TArray<TSharedRef<IAnimBlueprintCompilationContext::FFoldedPropertyRecord>>& InRecords)
-		{
-			for(const TSharedRef<IAnimBlueprintCompilationContext::FFoldedPropertyRecord>& Record : InRecords)
-			{
-				if(Record->GeneratedProperty != nullptr)
-				{
-					FStructProperty* AnimGraphNodeProperty = Record->AnimGraphNode->GetFNodeProperty();
-
-					// Check the anim node property is contained in the anim graph node
-					check(AnimGraphNodeProperty->GetOwner<UClass>() && Record->AnimGraphNode->GetClass()->IsChildOf(AnimGraphNodeProperty->GetOwner<UClass>()));
-					const void* Node = AnimGraphNodeProperty->ContainerPtrToValuePtr<void>(Record->AnimGraphNode);
-
-					// Check the anim node's property is contained in the anim node
-					check(Record->Property->GetOwner<UStruct>() && AnimGraphNodeProperty->Struct->IsChildOf(Record->Property->GetOwner<UStruct>()));
-					const void* SourcePtr = Record->Property->ContainerPtrToValuePtr<void>(Node);
-
-					// Check the generated property is a member of the constants struct
-					check(Record->GeneratedProperty->GetOwner<UStruct>() && InStruct->IsChildOf(Record->GeneratedProperty->GetOwner<UStruct>()));
-					void* TargetPtr = Record->GeneratedProperty->ContainerPtrToValuePtr<void>(InData);
-
-					// Extract underlying property for enums
-					FProperty* PropertyToCopyWith = Record->GeneratedProperty;
-					if(const FEnumProperty* EnumProperty = CastField<const FEnumProperty>(Record->GeneratedProperty))
-					{
-						PropertyToCopyWith = EnumProperty->GetUnderlyingProperty();
-					}
-
-					PropertyToCopyWith->CopyCompleteValue(TargetPtr, SourcePtr);
-				}
-			}
-		};
-
-		if(void* Constants = NewAnimBlueprintClass->GetOrCreateSparseClassData())
-		{
- 			UScriptStruct* ConstantsStruct = NewAnimBlueprintClass->GetSparseClassDataStruct();
- 			check(ConstantsStruct == NewAnimBlueprintConstants);
-
-			// Initialize extensions from their templates
-			for (TFieldIterator<FStructProperty> It(ConstantsStruct, EFieldIteratorFlags::ExcludeSuper); It; ++It)
-			{
-				FStructProperty* TargetProperty = *It;
-
-				if (UAnimBlueprintExtension* Extension = ClassPropertyToExtensionMap.FindRef(TargetProperty))
-				{
-					if(const FStructProperty* SourceExtensionProperty = Extension->GetClassDataProperty())
-					{
-						check(SourceExtensionProperty != NULL);
-						check(TargetProperty->Struct == SourceExtensionProperty->Struct);
-
-						uint8* DestinationPtr = TargetProperty->ContainerPtrToValuePtr<uint8>(Constants);
-						const uint8* SourcePtr = SourceExtensionProperty->ContainerPtrToValuePtr<uint8>(Extension);
-					
-						FAnimBlueprintExtensionCopyTermDefaultsContext NodeContext(TargetProperty, DestinationPtr, SourcePtr, LinkIndexCount);
-						Extension->CopyTermDefaultsToSparseClassData(CompilerContext, NodeContext);
-					}
-				}
-			}
-
-			// Patch constants
-			PatchDataArea(Constants, ConstantsStruct, ConstantPropertyRecords);
-		}
-
-		if(NewMutablesProperty)
-		{
-			check(NewMutablesProperty->GetOwner<UClass>() && DefaultObject->GetClass()->IsChildOf(NewMutablesProperty->GetOwner<UClass>()));
-			void* Mutables = NewMutablesProperty->ContainerPtrToValuePtr<void>(DefaultObject);
-			UScriptStruct* MutablesStruct = NewMutablesProperty->Struct;
-			check(MutablesStruct == NewAnimBlueprintMutables);
-
-			// Patch mutables
-			PatchDataArea(Mutables, MutablesStruct, MutablePropertyRecords);
-		}
-		
 		// And wire up node links
 		for (auto PoseLinkIt = ValidPoseLinkList.CreateIterator(); PoseLinkIt; ++PoseLinkIt)
 		{
@@ -930,8 +594,10 @@ void FAnimBlueprintCompilerContext::CopyTermDefaultsToDefaultObject(UObject* Def
 			}
 		}   
 
+		OnCopyTermDefaultsToDefaultObjectDelegate.Broadcast(DefaultObject, CompilerContext, CompiledData);
+
 		UAnimBlueprintGeneratedClass* AnimBlueprintGeneratedClass = CastChecked<UAnimBlueprintGeneratedClass>(NewClass);
-		
+
 		// copy threaded update flag to CDO
 		DefaultAnimInstance->bUseMultiThreadedAnimationUpdate = AnimBlueprint->bUseMultiThreadedAnimationUpdate;
 
@@ -1061,8 +727,6 @@ void FAnimBlueprintCompilerContext::CopyTermDefaultsToDefaultObject(UObject* Def
 
 void FAnimBlueprintCompilerContext::ExpandSplitPins(UEdGraph* InGraph)
 {
-	BP_SCOPED_COMPILER_EVENT_STAT(EAnimBlueprintCompilerStats_ExpandSplitPins);
-
 	for (TArray<UEdGraphNode*>::TIterator NodeIt(InGraph->Nodes); NodeIt; ++NodeIt)
 	{
 		UK2Node* K2Node = Cast<UK2Node>(*NodeIt);
@@ -1076,8 +740,6 @@ void FAnimBlueprintCompilerContext::ExpandSplitPins(UEdGraph* InGraph)
 // Merges in any all ubergraph pages into the gathering ubergraph
 void FAnimBlueprintCompilerContext::MergeUbergraphPagesIn(UEdGraph* Ubergraph)
 {
-	BP_SCOPED_COMPILER_EVENT_STAT(EAnimBlueprintCompilerStats_MergeUbergraphPagesIn);
-
 	Super::MergeUbergraphPagesIn(Ubergraph);
 
 	if (bIsDerivedAnimBlueprint)
@@ -1086,49 +748,37 @@ void FAnimBlueprintCompilerContext::MergeUbergraphPagesIn(UEdGraph* Ubergraph)
 	}
 	else
 	{
-		RecreateSparseClassData();
-		RecreateMutables();
-	
+		// Move all animation graph nodes and associated pure logic chains into the consolidated event graph
+		auto MoveGraph = [this](UEdGraph* InGraph)
 		{
-			BP_SCOPED_COMPILER_EVENT_STAT(EAnimBlueprintCompilerStats_MoveGraphs);
-			
-			// Move all animation graph nodes and associated pure logic chains into the consolidated event graph
-			auto MoveGraph = [this](UEdGraph* InGraph)
+			if (InGraph->Schema->IsChildOf(UAnimationGraphSchema::StaticClass()))
 			{
-				if (InGraph->Schema->IsChildOf(UAnimationGraphSchema::StaticClass()))
-				{
-					UEdGraph* ClonedGraph;
-					{
-						BP_SCOPED_COMPILER_EVENT_STAT(EAnimBlueprintCompilerStats_CloneGraph);
-					
-						// Merge all the animation nodes, contents, etc... into the ubergraph
-						ClonedGraph = FEdGraphUtilities::CloneGraph(InGraph, NULL, &MessageLog, true);
-					}
+				// Merge all the animation nodes, contents, etc... into the ubergraph
+				UEdGraph* ClonedGraph = FEdGraphUtilities::CloneGraph(InGraph, NULL, &MessageLog, true);
 
-					// Prune the graph up-front
-					const bool bIncludePotentialRootNodes = false;
-					PruneIsolatedNodes(ClonedGraph, bIncludePotentialRootNodes);
+				// Prune the graph up-front
+				const bool bIncludePotentialRootNodes = false;
+				PruneIsolatedNodes(ClonedGraph, bIncludePotentialRootNodes);
 
-					const bool bIsLoading = Blueprint->bIsRegeneratingOnLoad || IsAsyncLoading();
-					const bool bIsCompiling = Blueprint->bBeingCompiled;
-					ClonedGraph->MoveNodesToAnotherGraph(ConsolidatedEventGraph, bIsLoading, bIsCompiling);
+				const bool bIsLoading = Blueprint->bIsRegeneratingOnLoad || IsAsyncLoading();
+				const bool bIsCompiling = Blueprint->bBeingCompiled;
+				ClonedGraph->MoveNodesToAnotherGraph(ConsolidatedEventGraph, bIsLoading, bIsCompiling);
 
-					// Move subgraphs too
-					ConsolidatedEventGraph->SubGraphs.Append(ClonedGraph->SubGraphs);
-				}
-			};
-
-			for (UEdGraph* Graph : Blueprint->FunctionGraphs)
-			{
-				MoveGraph(Graph);
+				// Move subgraphs too
+				ConsolidatedEventGraph->SubGraphs.Append(ClonedGraph->SubGraphs);
 			}
+		};
 
-			for(FBPInterfaceDescription& InterfaceDesc : Blueprint->ImplementedInterfaces)
+		for (UEdGraph* Graph : Blueprint->FunctionGraphs)
+		{
+			MoveGraph(Graph);
+		}
+
+		for(FBPInterfaceDescription& InterfaceDesc : Blueprint->ImplementedInterfaces)
+		{
+			for(UEdGraph* Graph : InterfaceDesc.Graphs)
 			{
-				for(UEdGraph* Graph : InterfaceDesc.Graphs)
-				{
-					MoveGraph(Graph);
-				}
+				MoveGraph(Graph);
 			}
 		}
 
@@ -1164,144 +814,76 @@ void FAnimBlueprintCompilerContext::EnsureProperGeneratedClass(UClass*& TargetUC
 	}
 }
 
-void FAnimBlueprintCompilerContext::RecreateSparseClassData()
-{
-	UAnimBlueprintGeneratedClass* NewAnimBlueprintClass = GetNewAnimBlueprintClass();
-	
-	// Set up our sparse class data struct
-	const FString ConstantsStructTypeName = NewAnimBlueprintClass->GetName() + TEXT("_ConstantData");
-
-	// Rename old struct out of the way
-	if(UScriptStruct* OldAnimBlueprintConstants = FindObject<UScriptStruct>(NewAnimBlueprintClass->GetOuter(), *ConstantsStructTypeName))
-	{
-		OldAnimBlueprintConstants->Rename(nullptr, GetTransientPackage(), REN_DontCreateRedirectors | REN_NonTransactional | REN_ForceNoResetLoaders);
-	}
-
-	if (bIsDerivedAnimBlueprint)
-	{
-		// Get parent class
-		UAnimBlueprint* ParentAnimBP = UAnimBlueprint::GetParentAnimBlueprint(AnimBlueprint);
-		UAnimBlueprintGeneratedClass* ParentAnimClass = Cast<UAnimBlueprintGeneratedClass>(ParentAnimBP->GeneratedClass);
-		
-		check(ParentAnimClass);
-		check(ParentAnimClass->GetSparseClassDataStruct());
-		
-		// Derive sparse class data from parent class
-		NewAnimBlueprintConstants = NewObject<UScriptStruct>(NewAnimBlueprintClass->GetOuter(), *ConstantsStructTypeName, RF_Public);
-		NewAnimBlueprintConstants->SetSuperStruct(ParentAnimClass->GetSparseClassDataStruct());
-
-		// Just link & assign sparse class data struct here, no additional members are added
-		NewAnimBlueprintConstants->StaticLink(true);
-		NewAnimBlueprintClass->SetSparseClassDataStruct(NewAnimBlueprintConstants);
-		NewAnimBlueprintClass->GetOrCreateSparseClassData();
-		NewAnimBlueprintClass->BuildConstantProperties();
-	}
-	else
-	{
-		UClass* ParentClass = NewAnimBlueprintClass->GetSuperClass();
-		check(ParentClass);
-		
-	    // Create new sparse class data struct
-		NewAnimBlueprintConstants = NewObject<UScriptStruct>(NewAnimBlueprintClass->GetOuter(), *ConstantsStructTypeName, RF_Public);
-
-		// Inherit from archetype struct if there is any 
-		UScriptStruct* ArchetypeStruct = ParentClass->GetSparseClassDataArchetypeStruct();
-		UScriptStruct* SuperStruct = ArchetypeStruct ? ArchetypeStruct : FAnimBlueprintConstantData::StaticStruct();
-		NewAnimBlueprintConstants->SetSuperStruct(SuperStruct);
-	}
-}
-
-void FAnimBlueprintCompilerContext::RecreateMutables()
-{
-	UAnimBlueprintGeneratedClass* NewAnimBlueprintClass = GetNewAnimBlueprintClass();
-
-	// Set up our mutables struct
-	const FString MutablesStructTypeName = NewAnimBlueprintClass->GetName() + TEXT("_MutableData");
-
-	// Rename old struct out of the way
-	if(UScriptStruct* OldAnimBlueprintMutables = FindObject<UScriptStruct>(NewAnimBlueprintClass->GetOuter(), *MutablesStructTypeName))
-	{
-		OldAnimBlueprintMutables->Rename(nullptr, GetTransientPackage(), REN_DontCreateRedirectors | REN_NonTransactional | REN_ForceNoResetLoaders);
-	}
-
-	NewAnimBlueprintMutables = NewObject<UScriptStruct>(NewAnimBlueprintClass->GetOuter(), *MutablesStructTypeName);
-	NewAnimBlueprintMutables->SetSuperStruct(FAnimBlueprintMutableData::StaticStruct());
-}
-
 void FAnimBlueprintCompilerContext::SpawnNewClass(const FString& NewClassName)
 {
-	NewClass = FindObject<UAnimBlueprintGeneratedClass>(Blueprint->GetOutermost(), *NewClassName);
+	NewAnimBlueprintClass = FindObject<UAnimBlueprintGeneratedClass>(Blueprint->GetOutermost(), *NewClassName);
 
-	if (NewClass == NULL)
+	if (NewAnimBlueprintClass == NULL)
 	{
-		NewClass = NewObject<UAnimBlueprintGeneratedClass>(Blueprint->GetOutermost(), *NewClassName, RF_Public | RF_Transactional);
+		NewAnimBlueprintClass = NewObject<UAnimBlueprintGeneratedClass>(Blueprint->GetOutermost(), FName(*NewClassName), RF_Public | RF_Transactional);
 	}
 	else
 	{
 		// Already existed, but wasn't linked in the Blueprint yet due to load ordering issues
-		FBlueprintCompileReinstancer::Create(NewClass);
+		FBlueprintCompileReinstancer::Create(NewAnimBlueprintClass);
 	}
-	
-	FAnimBlueprintGeneratedClassCompiledData CompiledData(GetNewAnimBlueprintClass());
+	NewClass = NewAnimBlueprintClass;
+
+	FAnimBlueprintGeneratedClassCompiledData CompiledData(NewAnimBlueprintClass);
 	FAnimBlueprintCompilationBracketContext CompilerContext(this);
-	UAnimBlueprintExtension::ForEachExtension(AnimBlueprint, [this, &CompiledData, &CompilerContext](UAnimBlueprintExtension* InExtension)
-    {
-		InExtension->StartCompilingClass(GetNewAnimBlueprintClass(), CompilerContext, CompiledData);
-	});
+	OnStartCompilingClassDelegate.Broadcast(NewAnimBlueprintClass, CompilerContext, CompiledData);
 }
 
 void FAnimBlueprintCompilerContext::OnPostCDOCompiled()
 {
-	UAnimBlueprintGeneratedClass* NewAnimBlueprintClass = GetNewAnimBlueprintClass();
 	for (UAnimBlueprintGeneratedClass* ClassWithInputHandlers = NewAnimBlueprintClass; ClassWithInputHandlers != nullptr; ClassWithInputHandlers = Cast<UAnimBlueprintGeneratedClass>(ClassWithInputHandlers->GetSuperClass()))
 	{
-		ClassWithInputHandlers->OnPostLoadDefaults(NewAnimBlueprintClass->ClassDefaultObject);
+		FExposedValueHandler::ClassInitialization(ClassWithInputHandlers->EvaluateGraphExposedInputs, NewAnimBlueprintClass->ClassDefaultObject);
+
+		ClassWithInputHandlers->LinkFunctionsToDefaultObjectNodes(NewAnimBlueprintClass->ClassDefaultObject);
 	}
 }
 
 void FAnimBlueprintCompilerContext::OnNewClassSet(UBlueprintGeneratedClass* ClassToUse)
 {
-	NewClass = CastChecked<UAnimBlueprintGeneratedClass>(ClassToUse);
+	NewAnimBlueprintClass = CastChecked<UAnimBlueprintGeneratedClass>(ClassToUse);
 }
 
 void FAnimBlueprintCompilerContext::CleanAndSanitizeClass(UBlueprintGeneratedClass* ClassToClean, UObject*& InOldCDO)
 {
 	Super::CleanAndSanitizeClass(ClassToClean, InOldCDO);
 
-	UAnimBlueprintGeneratedClass* AnimBlueprintClassToClean= CastChecked<UAnimBlueprintGeneratedClass>(ClassToClean);
-	
-	AnimBlueprintClassToClean->AnimBlueprintDebugData = FAnimBlueprintDebugData();
+	// Make sure our typed pointer is set
+	check(ClassToClean == NewClass && NewAnimBlueprintClass == NewClass);
+
+	NewAnimBlueprintClass->AnimBlueprintDebugData = FAnimBlueprintDebugData();
 
 	// Reset the baked data
 	//@TODO: Move this into PurgeClass
-	AnimBlueprintClassToClean->BakedStateMachines.Empty();
-	AnimBlueprintClassToClean->AnimNotifies.Empty();
-	AnimBlueprintClassToClean->AnimBlueprintFunctions.Empty();
-	AnimBlueprintClassToClean->OrderedSavedPoseIndicesMap.Empty();
-	AnimBlueprintClassToClean->AnimNodeProperties.Empty();
-	AnimBlueprintClassToClean->LinkedAnimGraphNodeProperties.Empty();
-	AnimBlueprintClassToClean->LinkedAnimLayerNodeProperties.Empty();
-	AnimBlueprintClassToClean->PreUpdateNodeProperties.Empty();
-	AnimBlueprintClassToClean->DynamicResetNodeProperties.Empty();
-	AnimBlueprintClassToClean->StateMachineNodeProperties.Empty();
-	AnimBlueprintClassToClean->InitializationNodeProperties.Empty();
-	AnimBlueprintClassToClean->GraphAssetPlayerInformation.Empty();
-	AnimBlueprintClassToClean->GraphBlendOptions.Empty();
-	AnimBlueprintClassToClean->AnimNodeData.Empty();
-	AnimBlueprintClassToClean->NodeTypeMap.Empty();
+	NewAnimBlueprintClass->BakedStateMachines.Empty();
+	NewAnimBlueprintClass->AnimNotifies.Empty();
+	NewAnimBlueprintClass->AnimBlueprintFunctions.Empty();
+	NewAnimBlueprintClass->OrderedSavedPoseIndicesMap.Empty();
+	NewAnimBlueprintClass->AnimNodeProperties.Empty();
+	NewAnimBlueprintClass->LinkedAnimGraphNodeProperties.Empty();
+	NewAnimBlueprintClass->LinkedAnimLayerNodeProperties.Empty();
+	NewAnimBlueprintClass->PreUpdateNodeProperties.Empty();
+	NewAnimBlueprintClass->DynamicResetNodeProperties.Empty();
+	NewAnimBlueprintClass->StateMachineNodeProperties.Empty();
+	NewAnimBlueprintClass->InitializationNodeProperties.Empty();
+	NewAnimBlueprintClass->EvaluateGraphExposedInputs.Empty();
+	NewAnimBlueprintClass->GraphAssetPlayerInformation.Empty();
+	NewAnimBlueprintClass->GraphBlendOptions.Empty();
 
 	// Copy over runtime data from the blueprint to the class
-	AnimBlueprintClassToClean->TargetSkeleton = AnimBlueprint->TargetSkeleton;
+	NewAnimBlueprintClass->TargetSkeleton = AnimBlueprint->TargetSkeleton;
 
 	UAnimBlueprint* RootAnimBP = UAnimBlueprint::FindRootAnimBlueprint(AnimBlueprint);
 	bIsDerivedAnimBlueprint = RootAnimBP != NULL;
 
-	FAnimBlueprintGeneratedClassCompiledData CompiledData(AnimBlueprintClassToClean);
+	FAnimBlueprintGeneratedClassCompiledData CompiledData(NewAnimBlueprintClass);
 	FAnimBlueprintCompilationBracketContext CompilerContext(this);
-	UAnimBlueprintExtension::ForEachExtension(AnimBlueprint, [this, &CompiledData, &CompilerContext, AnimBlueprintClassToClean](UAnimBlueprintExtension* InExtension)
-	{
-		InExtension->StartCompilingClass(AnimBlueprintClassToClean, CompilerContext, CompiledData);
-	});
+	OnStartCompilingClassDelegate.Broadcast(NewAnimBlueprintClass, CompilerContext, CompiledData);
 }
 
 void FAnimBlueprintCompilerContext::FinishCompilingClass(UClass* Class)
@@ -1344,12 +926,9 @@ void FAnimBlueprintCompilerContext::FinishCompilingClass(UClass* Class)
 		}
 	}
 
-	FAnimBlueprintGeneratedClassCompiledData CompiledData(GetNewAnimBlueprintClass());
+	FAnimBlueprintGeneratedClassCompiledData CompiledData(NewAnimBlueprintClass);
 	FAnimBlueprintCompilationBracketContext CompilerContext(this);
-	UAnimBlueprintExtension::ForEachExtension(AnimBlueprint, [this, &CompiledData, &CompilerContext](UAnimBlueprintExtension* InExtension)
-	{
-		InExtension->FinishCompilingClass(GetNewAnimBlueprintClass(), CompilerContext, CompiledData);
-	});
+	OnFinishCompilingClassDelegate.Broadcast(AnimBlueprintGeneratedClass, CompilerContext, CompiledData);
 
 	Super::FinishCompilingClass(Class);
 }
@@ -1383,8 +962,6 @@ void FAnimBlueprintCompilerContext::PostCompileDiagnostics()
 {
 	FKismetCompilerContext::PostCompileDiagnostics();
 
-	UAnimBlueprintGeneratedClass* NewAnimBlueprintClass = GetNewAnimBlueprintClass();
-	
 #if WITH_EDITORONLY_DATA // ANIMINST_PostCompileValidation
 	// See if AnimInstance implements a PostCompileValidation Class. 
 	// If so, instantiate it, and let it perform Validation of our newly compiled AnimBlueprint.
@@ -1511,7 +1088,7 @@ void FAnimBlueprintCompilerContext::CreateAnimGraphStubFunctions()
 				UK2Node_FunctionEntry* EntryNode = SpawnIntermediateNode<UK2Node_FunctionEntry>(RootNode, StubGraph);
 				EntryNode->NodePosX = -200;
 				EntryNode->CustomGeneratedFunctionName = GraphToUseforSignature->GetFName();	// Note that the function generated from this temporary graph is undecorated
-				EntryNode->MetaData.Category = RootNode->Node.GetGroup() == NAME_None ? FText::GetEmpty() : FText::FromName(RootNode->Node.GetGroup());
+				EntryNode->MetaData.Category = RootNode->Node.Group == NAME_None ? FText::GetEmpty() : FText::FromName(RootNode->Node.Group);
 
 				// Add linked input poses as parameters
 				for(UAnimGraphNode_LinkedInputPose* LinkedInputPoseNode : LinkedInputPoseNodes)
@@ -1654,7 +1231,7 @@ void FAnimBlueprintCompilerContext::AddAttributesToNode(UAnimGraphNode_Base* InN
 {
 	if(UAnimGraphNode_Base* OriginalNode = CastChecked<UAnimGraphNode_Base>(MessageLog.FindSourceObject(InNode)))
 	{
-		TArray<FName>& AttributeSet = GetNewAnimBlueprintClass()->GetAnimBlueprintDebugData().NodeAttributes.FindOrAdd(OriginalNode);
+		TArray<FName>& AttributeSet = NewAnimBlueprintClass->GetAnimBlueprintDebugData().NodeAttributes.FindOrAdd(OriginalNode);
 		AttributeSet.Reserve(AttributeSet.Num() + InAttributes.Num());
 
 		for(const FName& Attribute : InAttributes)
@@ -1668,7 +1245,7 @@ TArrayView<const FName> FAnimBlueprintCompilerContext::GetAttributesFromNode(UAn
 {
 	if(UAnimGraphNode_Base* OriginalNode = CastChecked<UAnimGraphNode_Base>(MessageLog.FindSourceObject(InNode)))
 	{
-		if(const TArray<FName>* AttributeSetPtr = GetNewAnimBlueprintClass()->GetAnimBlueprintDebugData().NodeAttributes.Find(OriginalNode))
+		if(const TArray<FName>* AttributeSetPtr = NewAnimBlueprintClass->GetAnimBlueprintDebugData().NodeAttributes.Find(OriginalNode))
 		{
 			return MakeArrayView(*AttributeSetPtr);
 		}
@@ -1683,263 +1260,6 @@ FProperty* FAnimBlueprintCompilerContext::CreateUniqueVariable(UObject* InForObj
 	return CreateVariable(*VariableName, Type);
 }
 
-FProperty* FAnimBlueprintCompilerContext::CreateStructVariable(UScriptStruct* InStruct, const FName VarName, const FEdGraphPinType& VarType)
-{
-	FProperty* NewProperty = FKismetCompilerUtilities::CreatePropertyOnScope(InStruct, VarName, VarType, nullptr, CPF_None, Schema, MessageLog);
-	if (NewProperty != nullptr)
-	{
-		// This fixes a rare bug involving asynchronous loading of BPs in editor builds. The pattern was established
-		// in FKismetCompilerContext::CompileFunctions where we do this for the uber graph function. By setting
-		// the RF_LoadCompleted we prevent the linker from overwriting our regenerated property, although the
-		// circumstances under which this occurs are murky. More testing of BPs loading asynchronously in the editor
-		// needs to be added:
-		NewProperty->SetFlags(RF_LoadCompleted);
-		FKismetCompilerUtilities::LinkAddedProperty(InStruct, NewProperty);
-	}
-	else
-	{
-		MessageLog.Error(
-			*FText::Format(
-				LOCTEXT("VariableInvalidType_ErrorFmt", "The variable {0} declared in @@ has an invalid type {1}"),
-				FText::FromName(VarName),
-				UEdGraphSchema_K2::TypeToText(VarType)
-			).ToString(),
-			Blueprint
-		);
-	}
-
-	return NewProperty;
-}
-
-void FAnimBlueprintCompilerContext::AddFoldedPropertyRecord(UAnimGraphNode_Base* InAnimGraphNode, FStructProperty* InAnimNodeProperty, FProperty* InProperty, bool bInExposedOnPin, bool bInPinConnected)
-{
-	const bool bConstant = !bInExposedOnPin || (bInExposedOnPin && !bInPinConnected);
-
-	if(!InProperty->HasAnyPropertyFlags(CPF_EditorOnly))
-	{
-		MessageLog.Warning(*FString::Printf(TEXT("Property %s on @@ is foldable, but not editor only"), *InProperty->GetName()), InAnimGraphNode);
-	}
-
-	// Create record and add it our lookup map
-	TSharedRef<IAnimBlueprintCompilationContext::FFoldedPropertyRecord> Record = MakeShared<IAnimBlueprintCompilationContext::FFoldedPropertyRecord>(InAnimGraphNode, InAnimNodeProperty, InProperty, bConstant);
-	TArray<TSharedRef<IAnimBlueprintCompilationContext::FFoldedPropertyRecord>>& Array = NodeToFoldedPropertyRecordMap.FindOrAdd(InAnimGraphNode);
-	Array.Add(Record);
-
-	// Record it in the appropriate data area
-	if(bConstant)
-	{
-		ConstantPropertyRecords.Add(Record);
-	}
-	else
-	{
-		MutablePropertyRecords.Add(Record);
-	}
-}
-
-void FAnimBlueprintCompilerContext::ProcessFoldedPropertyRecords()
-{
-	UAnimBlueprintGeneratedClass* NewAnimBlueprintClass = GetNewAnimBlueprintClass();
-	
-	if(ConstantPropertyRecords.Num() > 0)
-	{
-		// Set constants struct as sparse class data
-		check(NewAnimBlueprintConstants);
-
-		auto GetRecordValue = [](const TSharedRef<IAnimBlueprintCompilationContext::FFoldedPropertyRecord>& InRecord)
-		{
-			FStructProperty* AnimGraphNodeProperty = InRecord->AnimGraphNode->GetFNodeProperty();
-
-			check(AnimGraphNodeProperty->GetOwner<UClass>() && InRecord->AnimGraphNode->GetClass()->IsChildOf(AnimGraphNodeProperty->GetOwner<UClass>()));
-			const void* Node = AnimGraphNodeProperty->ContainerPtrToValuePtr<void>(InRecord->AnimGraphNode);
-			check(InRecord->Property->GetOwner<UStruct>() && AnimGraphNodeProperty->Struct->IsChildOf(InRecord->Property->GetOwner<UStruct>()));
-			return InRecord->Property->ContainerPtrToValuePtr<void>(Node);
-		};
-
-		// Reduce any constant properties before patching
-		for(int32 RecordIndex0 = 0; RecordIndex0 < ConstantPropertyRecords.Num(); ++RecordIndex0)
-		{
-			const TSharedRef<IAnimBlueprintCompilationContext::FFoldedPropertyRecord>& Record0 = ConstantPropertyRecords[RecordIndex0];
-			if(Record0->FoldIndex == INDEX_NONE)
-			{
-				const void* Value0 = GetRecordValue(Record0);
-
-				for(int32 RecordIndex1 = RecordIndex0 + 1; RecordIndex1 < ConstantPropertyRecords.Num(); ++RecordIndex1)
-				{
-					const TSharedRef<IAnimBlueprintCompilationContext::FFoldedPropertyRecord>& Record1 = ConstantPropertyRecords[RecordIndex1];
-					if(Record1->FoldIndex == INDEX_NONE)
-					{
-						if(Record1->Property->SameType(Record0->Property))
-						{
-							const void* Value1 = GetRecordValue(Record1);
-
-							// same type, now test for equality
-							if(Record0->Property->Identical(Value0, Value1))
-							{
-								// Values are the same - fold
-								Record1->FoldIndex = RecordIndex0;
-							}
-						}
-					}
-				}
-			}
-		}
-	}
-
-	// Builds a 'data area', returns the total number of properties that were inserted into that area's struct
-	auto BuildDataArea = [this](TArray<TSharedRef<IAnimBlueprintCompilationContext::FFoldedPropertyRecord>>& InRecords, UScriptStruct* InStruct)
-	{
-		int32 PropertyIndex = 0;
-
-		if(InStruct)
-		{
-			const UAnimationGraphSchema* AnimationGraphSchema = GetDefault<UAnimationGraphSchema>();
-
-			for(TSharedRef<IAnimBlueprintCompilationContext::FFoldedPropertyRecord>& Record : InRecords)
-			{
-				// Skip folded records
-				if(Record->FoldIndex == INDEX_NONE)
-				{
-					FEdGraphPinType VariableType;
-					if(AnimationGraphSchema->ConvertPropertyToPinType(Record->Property, VariableType))
-					{
-						// Patch into sparse class data
-						FName PropertyName = FName(Record->Property->GetClass()->GetFName(), InRecords.Num() - 1 - PropertyIndex);
-						Record->GeneratedProperty = CreateStructVariable(InStruct, PropertyName, VariableType);
-						if (Record->GeneratedProperty == nullptr)
-						{
-							MessageLog.Error(*FString::Printf(TEXT("Property %s on node @@ could not be patched into data area."), *Record->Property->GetName()), Record->AnimGraphNode);
-						}
-						else
-						{
-							Record->PropertyIndex = PropertyIndex++;
-						}
-					}
-					else
-					{
-						MessageLog.Error(*FString::Printf(TEXT("Property %s on node @@ could not be patched into data area."), *Record->Property->GetName()), Record->AnimGraphNode);
-					}
-				}
-			}
-
-			if(InRecords.Num() > 0)
-			{
-				InStruct->StaticLink(true);
-			}
-		}
-
-		return PropertyIndex;
-	};
-
-	// Next, patch into the relevant data areas
-	const int32 NumConstantProperties = BuildDataArea(ConstantPropertyRecords, NewAnimBlueprintConstants);
-
-	// Set constants as sparse class data
-	if(ConstantPropertyRecords.Num() > 0)
-	{
-		NewAnimBlueprintClass->SetSparseClassDataStruct(NewAnimBlueprintConstants);
-	}
-
-	const int32 NumMutableProperties = BuildDataArea(MutablePropertyRecords, NewAnimBlueprintMutables);
-
-	// Create the property for our mutables, if we have any
-	if(MutablePropertyRecords.Num() > 0)
-	{
-		const FName MutablesStructName("__AnimBlueprintMutables");
-		FEdGraphPinType MutablesPinType;
-		MutablesPinType.PinCategory = UAnimationGraphSchema::PC_Struct;
-		MutablesPinType.PinSubCategoryObject = MakeWeakObjectPtr(NewAnimBlueprintMutables);
-		NewMutablesProperty = CastFieldChecked<FStructProperty>(CreateVariable(MutablesStructName, MutablesPinType));
-		NewMutablesProperty->SetMetaData(TEXT("BlueprintCompilerGeneratedDefaults"), TEXT("true"));
-	}
-
-	// Set up per-node mappings
-	NewAnimBlueprintClass->AnimNodeData.Empty();
-	NewAnimBlueprintClass->AnimNodeData.SetNum(AllocatedAnimNodeIndices.Num());
-	NewAnimBlueprintClass->NodeTypeMap.Empty();
-	
-	// First index & setup the node data
-	int32 NodeIndex = 0;
-	for(FAnimNodeData& NodeData : NewAnimBlueprintClass->AnimNodeData)
-	{
-		const int32 NodePropertyIndex = AllocatedAnimNodeIndices.Num() - 1 - NodeIndex;
-		FStructProperty* NodeStructProperty = CastFieldChecked<FStructProperty>(AllocatedPropertiesByIndex.FindChecked(NodePropertyIndex));
-		const UScriptStruct* Struct = NodeStructProperty->Struct;
-		const FAnimNodeStructData& AnimNodeStructData = NewAnimBlueprintClass->NodeTypeMap.Add(Struct, FAnimNodeStructData(Struct));
-
-		// Add any super-structs as values can be accessed via base classes
-		const UScriptStruct* SuperStruct = Cast<UScriptStruct>(Struct->GetSuperStruct());
-		while(SuperStruct)
-		{
-			NewAnimBlueprintClass->NodeTypeMap.Add(SuperStruct, FAnimNodeStructData(SuperStruct));
-			SuperStruct = Cast<UScriptStruct>(SuperStruct->GetSuperStruct());
-		}
-		
-		NodeData.AnimClassInterface = NewAnimBlueprintClass;
-		NodeData.NodeIndex = NodeIndex++;
-		const int32 NumProperties = AnimNodeStructData.GetNumProperties();
-		check(NumProperties >= 0);
-		NodeData.Entries.SetNum(NumProperties);
-		for(uint32& Entry : NodeData.Entries)
-		{
-			Entry = ANIM_NODE_DATA_INVALID_ENTRY;
-		}
-	}
-
-	auto BuildAnimNodeData = [this, &NewAnimBlueprintClass](const TArray<TSharedRef<IAnimBlueprintCompilationContext::FFoldedPropertyRecord>>& InRecords, int32 InTotalPropertyCount)
-	{
-		const int32 NumNodes = AllocatedAnimNodeIndices.Num();
-
-		for(const TSharedRef<IAnimBlueprintCompilationContext::FFoldedPropertyRecord>& Record : InRecords)
-		{
-			int32 NodeIndex = NumNodes - 1 - AllocatedAnimNodeIndices.FindChecked(Record->AnimGraphNode);
-			
-			FAnimNodeData& NodeData = NewAnimBlueprintClass->AnimNodeData[NodeIndex];
-			
-			int32 PropertyIndex = Record->PropertyIndex;
-			if(Record->FoldIndex != INDEX_NONE)
-			{
-				PropertyIndex = InRecords[Record->FoldIndex]->PropertyIndex;
-			}
-			check(PropertyIndex >= 0 && PropertyIndex < InTotalPropertyCount);
-			PropertyIndex = InTotalPropertyCount - 1 - PropertyIndex;
-
-			uint32 PropertyEntry = PropertyIndex;
-			if(!Record->bIsOnClass)
-			{
-				PropertyEntry |= ANIM_NODE_DATA_INSTANCE_DATA_FLAG;
-			}
-
-			const FAnimNodeStructData& AnimNodeStructData = NewAnimBlueprintClass->NodeTypeMap.FindChecked(Record->AnimNodeProperty->Struct);
-			const int32 EntryIndex = AnimNodeStructData.GetPropertyIndex(Record->Property->GetFName());
-			NodeData.Entries[EntryIndex] = PropertyEntry;
-		}
-	};
-
-	BuildAnimNodeData(ConstantPropertyRecords, NumConstantProperties);
-	BuildAnimNodeData(MutablePropertyRecords, NumMutableProperties);
-}
-
-bool FAnimBlueprintCompilerContext::IsAnimGraphNodeFolded(UAnimGraphNode_Base* InNode) const
-{
-	return NodeToFoldedPropertyRecordMap.Find(InNode) != nullptr;
-}
-
-const IAnimBlueprintCompilationContext::FFoldedPropertyRecord* FAnimBlueprintCompilerContext::GetFoldedPropertyRecord(UAnimGraphNode_Base* InNode, FName InPropertyName) const
-{
-	if(const TArray<TSharedRef<IAnimBlueprintCompilationContext::FFoldedPropertyRecord>>* FoundRecordsPtr = NodeToFoldedPropertyRecordMap.Find(InNode))
-	{
-		for(const TSharedRef<IAnimBlueprintCompilationContext::FFoldedPropertyRecord>& Record : *FoundRecordsPtr)
-		{
-			if(Record->Property->GetFName() == InPropertyName)
-			{
-				return &Record.Get();
-			}
-		}
-	}
-
-	return nullptr;
-}
-
 //////////////////////////////////////////////////////////////////////////
 
 #undef LOCTEXT_NAMESPACE
\ No newline at end of file
diff --git a/Engine/Source/Editor/AnimGraph/Private/AnimBlueprintCompiler.h b/Engine/Source/Editor/AnimGraph/Private/AnimBlueprintCompiler.h
--- a/Engine/Source/Editor/AnimGraph/Private/AnimBlueprintCompiler.h
+++ b/Engine/Source/Editor/AnimGraph/Private/AnimBlueprintCompiler.h
@@ -7,9 +7,10 @@
 #include "Animation/AnimNodeBase.h"
 #include "AnimGraphNode_Base.h"
 #include "KismetCompilerModule.h"
+#include "AnimBlueprintCompilerHandlerCollection.h"
+#include "IAnimBlueprintCompilerHandlerCollection.h"
 #include "IAnimBlueprintCompilerCreationContext.h"
 #include "Containers/ArrayView.h"
-#include "IAnimBlueprintCompilationContext.h"
 
 class UAnimationGraphSchema;
 class UAnimGraphNode_SaveCachedPose;
@@ -79,9 +80,7 @@ protected:
 	typedef TArray<UEdGraphPin*> UEdGraphPinArray;
 
 protected:
-	UScriptStruct* NewAnimBlueprintConstants;
-	UScriptStruct* NewAnimBlueprintMutables;
-	FStructProperty* NewMutablesProperty;
+	UAnimBlueprintGeneratedClass* NewAnimBlueprintClass;
 	UAnimBlueprint* AnimBlueprint;
 
 	UAnimationGraphSchema* AnimSchema;
@@ -89,7 +88,6 @@ protected:
 	// Map of allocated v3 nodes that are members of the class
 	TMap<class UAnimGraphNode_Base*, FProperty*> AllocatedAnimNodes;
 	TMap<FProperty*, class UAnimGraphNode_Base*> AllocatedNodePropertiesToNodes;
-	TMap<FProperty*, class UAnimGraphNode_Base*> AllocatedNodeConstantPropertiesToNodes;
 	TMap<int32, FProperty*> AllocatedPropertiesByIndex;
 
 	// Map of true source objects (user edited ones) to the cloned ones that are actually compiled
@@ -109,31 +107,34 @@ protected:
 	// True if any parent class is also generated from an animation blueprint
 	bool bIsDerivedAnimBlueprint;
 
+	// Handlers that this context is hosting
+	FAnimBlueprintCompilerHandlerCollection AnimBlueprintCompilerHandlerCollection;
+
 	// Graph schema classes that this compiler is aware of - they will skip default function processing
 	TArray<TSubclassOf<UEdGraphSchema>> KnownGraphSchemas;
 
-	// Expose compile options to handlers
-	using FKismetCompilerContext::CompileOptions;
+	/** Delegate fired when the class starts compiling. The class may be new or recycled. */
+	FOnStartCompilingClass OnStartCompilingClassDelegate;
+
+	/** Delegate fired before all animation nodes are processed */
+	FOnPreProcessAnimationNodes OnPreProcessAnimationNodesDelegate;
 
-	// Records of folded properties gleaned from nodes
-	TArray<TSharedRef<IAnimBlueprintCompilationContext::FFoldedPropertyRecord>> ConstantPropertyRecords;
-	TArray<TSharedRef<IAnimBlueprintCompilationContext::FFoldedPropertyRecord>> MutablePropertyRecords;
+	/** Delegate fired after all animation nodes are processed */
+	FOnPostProcessAnimationNodes OnPostProcessAnimationNodesDelegate;
 
-	// Allows lookups to see if a node participates in constant folding
-	TMap<UAnimGraphNode_Base*, TArray<TSharedRef<IAnimBlueprintCompilationContext::FFoldedPropertyRecord>>> NodeToFoldedPropertyRecordMap;
+	/** Delegate fired post- graph expansion */
+	FOnPostExpansionStep OnPostExpansionStepDelegate;
 
-	// Maps of extension <-> generated property on the instance
-	TMap<UAnimBlueprintExtension*, FStructProperty*> ExtensionToInstancePropertyMap;
-	TMap<FStructProperty*, UAnimBlueprintExtension*> InstancePropertyToExtensionMap;
+	/** Delegate fired when the class has finished compiling */
+	FOnFinishCompilingClass OnFinishCompilingClassDelegate;
 
-	// Maps of extension <-> generated property on the class
-	TMap<UAnimBlueprintExtension*, FStructProperty*> ExtensionToClassPropertyMap;
-	TMap<FStructProperty*, UAnimBlueprintExtension*> ClassPropertyToExtensionMap;
+	/** Delegate fired when data is being copied to the CDO */
+	FOnCopyTermDefaultsToDefaultObject OnCopyTermDefaultsToDefaultObjectDelegate;
+
+	// Expose compile options to handlers
+	using FKismetCompilerContext::CompileOptions;
 
 private:
-	// Get the generated class as an anim blueprint generated class
-	UAnimBlueprintGeneratedClass* GetNewAnimBlueprintClass() const { return CastChecked<UAnimBlueprintGeneratedClass>(NewClass); };
-	
 	// Run a function on the passed-in graph and each subgraph of it
 	void ForAllSubGraphs(UEdGraph* InGraph, TFunctionRef<void(UEdGraph*)> InPerGraphFunction);
 
@@ -143,8 +144,8 @@ private:
 	// Compiles one animation node
 	void ProcessAnimationNode(UAnimGraphNode_Base* VisualAnimNode);
 
-	// Called during ProcessAnimationNode - gather property folding records for the node
-	void GatherFoldRecordsForAnimationNode(const UScriptStruct* InNodeType, FStructProperty* InNodeProperty, UAnimGraphNode_Base* InVisualAnimNode);
+	// Compiles one root node
+	void ProcessRoot(UAnimGraphNode_Root* Root);
 
 	// Compiles an entire animation graph
 	void ProcessAllAnimationNodes();
@@ -152,9 +153,6 @@ private:
 	// Processes all the supplied anim nodes
 	void ProcessAnimationNodes(TArray<UAnimGraphNode_Base*>& AnimNodeList);
 
-	// Process all the requested extensions
-	void ProcessExtensions();
-	
 	// Gets all anim graph nodes that are piped into the provided node (traverses input pins)
 	void GetLinkedAnimNodes(UAnimGraphNode_Base* InGraphNode, TArray<UAnimGraphNode_Base*>& LinkedAnimNodes) const;
 	void GetLinkedAnimNodes_TraversePin(UEdGraphPin* InPin, TArray<UAnimGraphNode_Base*>& LinkedAnimNodes) const;
@@ -177,32 +175,8 @@ private:
 
 	// Get the current compiled-in attributes uniquely assigned to the specified node
 	TArrayView<const FName> GetAttributesFromNode(UAnimGraphNode_Base* InNode) const;
-
-	// Called at the start of compilation to (re-) create the mutable struct
-	void RecreateMutables();
-	
-	// (Re-)creates sparse class data structure. Called at the start of compilation to (re-) create the internal sparse 
-	// class data. For derived anim BPs this is called just-in-time before CDO copy. This is to ensure that the sparse 
-	// class data is always updated, as in some cases as the layout does not change a full compilation of a (data-only) 
-	// derived anim BP may be skipped.  
-	void RecreateSparseClassData();
 	
 	// Create a uniquely named variable corresponding to an object in the current class
 	FProperty* CreateUniqueVariable(UObject* InForObject, const FEdGraphPinType& Type);
-
-	/** Creates a variable on the specified struct */
-	FProperty* CreateStructVariable(UScriptStruct* InStruct, const FName VarName, const FEdGraphPinType& VarType);
-
-	/** Adds a record for a potentially-folded anim node property */
-	void AddFoldedPropertyRecord(UAnimGraphNode_Base* InAnimGraphNode, FStructProperty* InAnimNodeProperty, FProperty* InProperty, bool bInExposedOnPin, bool bInPinConnected);
-
-	/** Process any anim node properties that are 'foldable' */
-	void ProcessFoldedPropertyRecords();
-
-	// Check whether an anim node participates in constant folding
-	bool IsAnimGraphNodeFolded(UAnimGraphNode_Base* InNode) const;
-
-	// Get the folded property record, if any, for the supplied node & named property
-	const IAnimBlueprintCompilationContext::FFoldedPropertyRecord* GetFoldedPropertyRecord(UAnimGraphNode_Base* InNode, FName InPropertyName) const;
 };
 
diff --git a/Engine/Source/Editor/AnimGraph/Private/AnimBlueprintCompilerCreationContext.cpp b/Engine/Source/Editor/AnimGraph/Private/AnimBlueprintCompilerCreationContext.cpp
--- a/Engine/Source/Editor/AnimGraph/Private/AnimBlueprintCompilerCreationContext.cpp
+++ b/Engine/Source/Editor/AnimGraph/Private/AnimBlueprintCompilerCreationContext.cpp
@@ -4,6 +4,36 @@
 #include "AnimBlueprintCompiler.h"
 #include "EdGraph/EdGraphSchema.h"
 
+FOnStartCompilingClass& FAnimBlueprintCompilerCreationContext::OnStartCompilingClass()
+{
+	return CompilerContext->OnStartCompilingClassDelegate;
+}
+
+FOnPreProcessAnimationNodes& FAnimBlueprintCompilerCreationContext::OnPreProcessAnimationNodes()
+{
+	return CompilerContext->OnPreProcessAnimationNodesDelegate;
+}
+
+FOnPostProcessAnimationNodes& FAnimBlueprintCompilerCreationContext::OnPostProcessAnimationNodes()
+{
+	return CompilerContext->OnPostProcessAnimationNodesDelegate;
+}
+
+FOnPostExpansionStep& FAnimBlueprintCompilerCreationContext::OnPostExpansionStep()
+{
+	return CompilerContext->OnPostExpansionStepDelegate;
+}
+
+FOnFinishCompilingClass& FAnimBlueprintCompilerCreationContext::OnFinishCompilingClass()
+{
+	return CompilerContext->OnFinishCompilingClassDelegate;
+}
+
+FOnCopyTermDefaultsToDefaultObject& FAnimBlueprintCompilerCreationContext::OnCopyTermDefaultsToDefaultObject()
+{
+	return CompilerContext->OnCopyTermDefaultsToDefaultObjectDelegate;
+}
+
 void FAnimBlueprintCompilerCreationContext::RegisterKnownGraphSchema(TSubclassOf<UEdGraphSchema> InGraphSchemaClass)
 {
 	CompilerContext->KnownGraphSchemas.AddUnique(InGraphSchemaClass);
diff --git a/Engine/Source/Editor/AnimGraph/Private/AnimBlueprintCompilerCreationContext.h b/Engine/Source/Editor/AnimGraph/Private/AnimBlueprintCompilerCreationContext.h
--- a/Engine/Source/Editor/AnimGraph/Private/AnimBlueprintCompilerCreationContext.h
+++ b/Engine/Source/Editor/AnimGraph/Private/AnimBlueprintCompilerCreationContext.h
@@ -11,13 +11,19 @@ class FAnimBlueprintCompilerContext;
 class FAnimBlueprintCompilerCreationContext : public IAnimBlueprintCompilerCreationContext
 {
 private:
-	friend class FAnimBlueprintCompilerContext;
-	
+	friend class FAnimBlueprintCompilerHandlerCollection;
+
 	FAnimBlueprintCompilerCreationContext(FAnimBlueprintCompilerContext* InCompilerContext)
 		: CompilerContext(InCompilerContext)
 	{}
 
 	// IAnimBlueprintCompilerCreationContext interface
+	virtual FOnStartCompilingClass& OnStartCompilingClass() override;
+	virtual FOnPreProcessAnimationNodes& OnPreProcessAnimationNodes() override;
+	virtual FOnPostProcessAnimationNodes& OnPostProcessAnimationNodes() override;
+	virtual FOnPostExpansionStep& OnPostExpansionStep() override;
+	virtual FOnFinishCompilingClass& OnFinishCompilingClass() override;
+	virtual FOnCopyTermDefaultsToDefaultObject& OnCopyTermDefaultsToDefaultObject() override;
 	virtual void RegisterKnownGraphSchema(TSubclassOf<UEdGraphSchema> InGraphSchemaClass) override;
 
 private:
diff --git a/Engine/Source/Editor/AnimGraph/Private/AnimBlueprintCompilerHandlerCollection.cpp b/Engine/Source/Editor/AnimGraph/Private/AnimBlueprintCompilerHandlerCollection.cpp
new file mode 100644
--- /dev/null
+++ b/Engine/Source/Editor/AnimGraph/Private/AnimBlueprintCompilerHandlerCollection.cpp
@@ -0,0 +1,38 @@
+// Copyright Epic Games, Inc. All Rights Reserved.
+
+#include "AnimBlueprintCompilerHandlerCollection.h"
+#include "AnimBlueprintCompilerCreationContext.h"
+
+/** All of our registered handler factories */
+static TMap<FName, TFunction<TUniquePtr<IAnimBlueprintCompilerHandler>(IAnimBlueprintCompilerCreationContext&)>> HandlerFactories;
+
+void IAnimBlueprintCompilerHandlerCollection::RegisterHandler(FName InName, TFunction<TUniquePtr<IAnimBlueprintCompilerHandler>(IAnimBlueprintCompilerCreationContext&)> InFunction)
+{
+	HandlerFactories.Add(InName, InFunction);
+}
+
+void IAnimBlueprintCompilerHandlerCollection::UnregisterHandler(FName InName)
+{
+	HandlerFactories.Remove(InName);
+}
+
+void FAnimBlueprintCompilerHandlerCollection::Initialize(FAnimBlueprintCompilerContext* InCompilerContext)
+{
+	FAnimBlueprintCompilerCreationContext CreationContext(InCompilerContext);
+
+	// Create all of the registered handlers
+	for(const TPair<FName, TFunction<TUniquePtr<IAnimBlueprintCompilerHandler>(IAnimBlueprintCompilerCreationContext&)>>& HandlerFactoryPair : HandlerFactories)
+	{
+		Handlers.Add(HandlerFactoryPair.Key, HandlerFactoryPair.Value(CreationContext));
+	}
+}
+
+IAnimBlueprintCompilerHandler* FAnimBlueprintCompilerHandlerCollection::GetHandlerByName(FName InName) const
+{
+	if(const TUniquePtr<IAnimBlueprintCompilerHandler>* FoundHandler = Handlers.Find(InName))
+	{
+		return (*FoundHandler).Get();
+	}
+
+	return nullptr;
+}
\ No newline at end of file
diff --git a/Engine/Source/Editor/AnimGraph/Private/AnimBlueprintCompilerHandlerCollection.h b/Engine/Source/Editor/AnimGraph/Private/AnimBlueprintCompilerHandlerCollection.h
new file mode 100644
--- /dev/null
+++ b/Engine/Source/Editor/AnimGraph/Private/AnimBlueprintCompilerHandlerCollection.h
@@ -0,0 +1,24 @@
+// Copyright Epic Games, Inc. All Rights Reserved.
+
+#pragma once
+
+#include "CoreMinimal.h"
+#include "IAnimBlueprintCompilerHandler.h"
+#include "IAnimBlueprintCompilerHandlerCollection.h"
+
+class FAnimBlueprintCompilerContext;
+
+/** Handler collection for the anim blueprint compiler */
+class FAnimBlueprintCompilerHandlerCollection : public IAnimBlueprintCompilerHandlerCollection
+{
+private:
+	friend class FAnimBlueprintCompilerContext;
+
+	void Initialize(FAnimBlueprintCompilerContext* InCompilerContext);
+
+	/** Get a named handler */
+	virtual IAnimBlueprintCompilerHandler* GetHandlerByName(FName InName) const override;
+
+	/** All the currently constructed handlers */
+	TMap<FName, TUniquePtr<IAnimBlueprintCompilerHandler>> Handlers;
+};
\ No newline at end of file
