Backing out CL 16071104

diff --git a/Engine/Source/Editor/AnimGraph/Private/AnimBlueprintGeneratedClassCompiledData.cpp b/Engine/Source/Editor/AnimGraph/Private/AnimBlueprintGeneratedClassCompiledData.cpp
--- a/Engine/Source/Editor/AnimGraph/Private/AnimBlueprintGeneratedClassCompiledData.cpp
+++ b/Engine/Source/Editor/AnimGraph/Private/AnimBlueprintGeneratedClassCompiledData.cpp
@@ -51,6 +51,16 @@ TArray<FAnimNotifyEvent>& FAnimBlueprintGeneratedClassCompiledData::GetAnimNotif
 	return Class->AnimNotifies;
 }
 
+TArray<FExposedValueHandler>& FAnimBlueprintGeneratedClassCompiledData::GetExposedValueHandlers() const
+{
+	return Class->EvaluateGraphExposedInputs;
+}
+
+FPropertyAccessLibrary& FAnimBlueprintGeneratedClassCompiledData::GetPropertyAccessLibrary() const
+{
+	return Class->PropertyAccessLibrary;
+}
+
 FAnimBlueprintDebugData& FAnimBlueprintGeneratedClassCompiledData::GetAnimBlueprintDebugData() const
 {
 	return Class->AnimBlueprintDebugData;
@@ -63,5 +73,7 @@ TMap<FName, FGraphAssetPlayerInformation>& FAnimBlueprintGeneratedClassCompiledD
 
 UBlendSpace* FAnimBlueprintGeneratedClassCompiledData::AddBlendSpace(UBlendSpace* InSourceBlendSpace)
 {
-	return nullptr;
-}
+	UBlendSpace* CopiedBlendSpace = DuplicateObject(InSourceBlendSpace, Class);
+	Class->BlendSpaces.Add(CopiedBlendSpace);
+	return CopiedBlendSpace;
+}
\ No newline at end of file
diff --git a/Engine/Source/Editor/AnimGraph/Private/AnimBlueprintGeneratedClassCompiledData.h b/Engine/Source/Editor/AnimGraph/Private/AnimBlueprintGeneratedClassCompiledData.h
--- a/Engine/Source/Editor/AnimGraph/Private/AnimBlueprintGeneratedClassCompiledData.h
+++ b/Engine/Source/Editor/AnimGraph/Private/AnimBlueprintGeneratedClassCompiledData.h
@@ -20,6 +20,8 @@ private:
 	virtual FBlueprintDebugData& GetBlueprintDebugData() const override;
 	virtual TArray<FAnimNotifyEvent>& GetAnimNotifies() const override;
 	virtual int32 FindOrAddNotify(FAnimNotifyEvent& Notify) const override;
+	virtual TArray<FExposedValueHandler>& GetExposedValueHandlers() const override;
+	virtual FPropertyAccessLibrary& GetPropertyAccessLibrary() const override;
 	virtual FAnimBlueprintDebugData& GetAnimBlueprintDebugData() const override;
 	virtual TMap<FName, FGraphAssetPlayerInformation>& GetGraphAssetPlayerInformation() const override;
 	virtual UBlendSpace* AddBlendSpace(UBlendSpace* InSourceBlendSpace) override;
diff --git a/Engine/Source/Editor/AnimGraph/Private/AnimBlueprintPostCompileValidation.cpp b/Engine/Source/Editor/AnimGraph/Private/AnimBlueprintPostCompileValidation.cpp
--- a/Engine/Source/Editor/AnimGraph/Private/AnimBlueprintPostCompileValidation.cpp
+++ b/Engine/Source/Editor/AnimGraph/Private/AnimBlueprintPostCompileValidation.cpp
@@ -145,11 +145,11 @@ void UAnimBlueprintPostCompileValidation::PCV_GatherAnimSequencesFromGraph(TArra
 		{
 			if (const FAnimNode_BlendSpacePlayer* const BlendSpacePlayer = Property->ContainerPtrToValuePtr<FAnimNode_BlendSpacePlayer>(PCV_Params.DefaultAnimInstance))
 			{
-				const bool bPassSyncGroupFilter = !GatherParams.bFilterBySyncGroup || (BlendSpacePlayer->GetGroupName() == GatherParams.SyncGroupName);
-				const bool bPassLoopingFilter = !GatherParams.bFilterByLoopingCondition || (BlendSpacePlayer->GetLoop() == GatherParams.bLoopingCondition);
+				const bool bPassSyncGroupFilter = !GatherParams.bFilterBySyncGroup || (BlendSpacePlayer->GroupName == GatherParams.SyncGroupName);
+				const bool bPassLoopingFilter = !GatherParams.bFilterByLoopingCondition || (BlendSpacePlayer->bLoop == GatherParams.bLoopingCondition);
 				if (bPassSyncGroupFilter && bPassLoopingFilter)
 				{
-					PCV_GatherAnimSequences(OutAnimSequences, BlendSpacePlayer->GetBlendSpace());
+					PCV_GatherAnimSequences(OutAnimSequences, BlendSpacePlayer->BlendSpace);
 				}
 			}
 		}
@@ -157,11 +157,11 @@ void UAnimBlueprintPostCompileValidation::PCV_GatherAnimSequencesFromGraph(TArra
 		{
 			if (const FAnimNode_SequencePlayer* const SequencePlayer = Property->ContainerPtrToValuePtr<FAnimNode_SequencePlayer>(PCV_Params.DefaultAnimInstance))
 			{
-				const bool bPassSyncGroupFilter = !GatherParams.bFilterBySyncGroup || (SequencePlayer->GetGroupName() == GatherParams.SyncGroupName);
-				const bool bPassLoopingFilter = !GatherParams.bFilterByLoopingCondition || (SequencePlayer->GetLoopAnimation() == GatherParams.bLoopingCondition);
+				const bool bPassSyncGroupFilter = !GatherParams.bFilterBySyncGroup || (SequencePlayer->GroupName == GatherParams.SyncGroupName);
+				const bool bPassLoopingFilter = !GatherParams.bFilterByLoopingCondition || (SequencePlayer->bLoopAnimation == GatherParams.bLoopingCondition);
 				if (bPassSyncGroupFilter && bPassLoopingFilter)
 				{
-					PCV_GatherAnimSequences(OutAnimSequences, SequencePlayer->GetSequence());
+					PCV_GatherAnimSequences(OutAnimSequences, SequencePlayer->Sequence);
 				}
 			}
 		}
@@ -176,7 +176,7 @@ void UAnimBlueprintPostCompileValidation::PCV_GatherBlendSpacesFromGraph(TArray<
 		{
 			if (const FAnimNode_BlendSpacePlayer* const BlendSpacePlayer = Property->ContainerPtrToValuePtr<FAnimNode_BlendSpacePlayer>(PCV_Params.DefaultAnimInstance))
 			{
-				OutBlendSpaces.AddUnique(BlendSpacePlayer->GetBlendSpace());
+				OutBlendSpaces.AddUnique(BlendSpacePlayer->BlendSpace);
 			}
 		}
 	}
diff --git a/Engine/Source/Editor/AnimGraph/Private/AnimGraphDetails.cpp b/Engine/Source/Editor/AnimGraph/Private/AnimGraphDetails.cpp
--- a/Engine/Source/Editor/AnimGraph/Private/AnimGraphDetails.cpp
+++ b/Engine/Source/Editor/AnimGraph/Private/AnimGraphDetails.cpp
@@ -321,12 +321,12 @@ FReply FAnimGraphDetails::OnRemoveInputPoseClicked(UAnimGraphNode_LinkedInputPos
 FText FAnimGraphDetails::OnGetGroupText() const
 {
 	UAnimGraphNode_Root* Root = FBlueprintEditorUtils::GetAnimGraphRoot(Graph);
-	if(Root->Node.GetGroup() == NAME_None)
+	if(Root->Node.Group == NAME_None)
 	{
 		return LOCTEXT("DefaultGroup", "Default");
 	}
 		
-	return FText::FromName(Root->Node.GetGroup());
+	return FText::FromName(Root->Node.Group);
 }
 
 void FAnimGraphDetails::OnGroupTextCommitted(const FText& NewText, ETextCommit::Type InTextCommit)
@@ -392,7 +392,7 @@ void FAnimGraphDetails::RefreshGroupSource()
 	GroupSource.Empty();
 
 	UAnimGraphNode_Root* Root = FBlueprintEditorUtils::GetAnimGraphRoot(Graph);
-	if(Root->Node.GetGroup() != NAME_None)
+	if(Root->Node.Group != NAME_None)
 	{
 		GroupSource.Add(MakeShared<FText>(LOCTEXT("DefaultGroup", "Default")));
 	}
diff --git a/Engine/Source/Editor/AnimGraph/Private/AnimGraphModule.cpp b/Engine/Source/Editor/AnimGraph/Private/AnimGraphModule.cpp
--- a/Engine/Source/Editor/AnimGraph/Private/AnimGraphModule.cpp
+++ b/Engine/Source/Editor/AnimGraph/Private/AnimGraphModule.cpp
@@ -24,6 +24,11 @@
 #include "AnimGraphDetails.h"
 #include "AnimationGraphSchema.h"
 #include "AnimBlueprintCompiler.h"
+#include "AnimBlueprintCompilerHandler_Base.h"
+#include "AnimBlueprintCompilerHandler_CachedPose.h"
+#include "AnimBlueprintCompilerHandler_LinkedAnimGraph.h"
+#include "AnimBlueprintCompilerHandler_StateMachine.h"
+#include "AnimBlueprintCompilerHandler_Attributes.h"
 #include "Animation/AnimAttributes.h"
 #include "AnimGraphAttributes.h"
 #include "Animation/AnimSync.h"
@@ -42,6 +47,32 @@ void FAnimGraphModule::StartupModule()
 		return MakeShared<FAnimBlueprintCompilerContext>(CastChecked<UAnimBlueprint>(InBlueprint), InMessageLog, InCompileOptions);
 	});
 
+	// Register node compilation handlers
+	IAnimBlueprintCompilerHandlerCollection::RegisterHandler("AnimBlueprintCompilerHandler_Base", [](IAnimBlueprintCompilerCreationContext& InCreationContext)
+	{
+		return MakeUnique<FAnimBlueprintCompilerHandler_Base>(InCreationContext);
+	});
+
+	IAnimBlueprintCompilerHandlerCollection::RegisterHandler("AnimBlueprintCompilerHandler_CachedPose", [](IAnimBlueprintCompilerCreationContext& InCreationContext)
+	{
+		return MakeUnique<FAnimBlueprintCompilerHandler_CachedPose>(InCreationContext);
+	});
+
+	IAnimBlueprintCompilerHandlerCollection::RegisterHandler("AnimBlueprintCompilerHandler_LinkedAnimGraph", [](IAnimBlueprintCompilerCreationContext& InCreationContext)
+	{
+		return MakeUnique<FAnimBlueprintCompilerHandler_LinkedAnimGraph>(InCreationContext);
+	});
+
+	IAnimBlueprintCompilerHandlerCollection::RegisterHandler("AnimBlueprintCompilerHandler_StateMachine", [](IAnimBlueprintCompilerCreationContext& InCreationContext)
+	{
+		return MakeUnique<FAnimBlueprintCompilerHandler_StateMachine>(InCreationContext);
+	});
+
+	IAnimBlueprintCompilerHandlerCollection::RegisterHandler("AnimBlueprintCompilerHandler_Attributes", [](IAnimBlueprintCompilerCreationContext& InCreationContext)
+	{
+		return MakeUnique<FAnimBlueprintCompilerHandler_Attributes>(InCreationContext);
+	});
+
 	// Register the editor modes
 	FEditorModeRegistry::Get().RegisterMode<FAnimNodeEditMode>(AnimNodeEditModes::AnimNode, LOCTEXT("AnimNodeEditMode", "Anim Node"), FSlateIcon(), false);
 	FEditorModeRegistry::Get().RegisterMode<FTwoBoneIKEditMode>(AnimNodeEditModes::TwoBoneIK, LOCTEXT("TwoBoneIKEditMode", "2-Bone IK"), FSlateIcon(), false);
@@ -139,6 +170,12 @@ void FAnimGraphModule::StartupModule()
 
 void FAnimGraphModule::ShutdownModule()
 {
+	IAnimBlueprintCompilerHandlerCollection::UnregisterHandler("AnimBlueprintCompilerHandler_Base");
+	IAnimBlueprintCompilerHandlerCollection::UnregisterHandler("AnimBlueprintCompilerHandler_CachedPose");
+	IAnimBlueprintCompilerHandlerCollection::UnregisterHandler("AnimBlueprintCompilerHandler_LinkedAnimGraph");
+	IAnimBlueprintCompilerHandlerCollection::UnregisterHandler("AnimBlueprintCompilerHandler_StateMachine");
+	IAnimBlueprintCompilerHandlerCollection::UnregisterHandler("AnimBlueprintCompilerHandler_Attributes");
+
 	// Unregister the editor modes
 	FEditorModeRegistry::Get().UnregisterMode(AnimNodeEditModes::CCDIK);
 	FEditorModeRegistry::Get().UnregisterMode(AnimNodeEditModes::SplineIK);
diff --git a/Engine/Source/Editor/AnimGraph/Private/AnimGraphNode_AimOffsetLookAt.cpp b/Engine/Source/Editor/AnimGraph/Private/AnimGraphNode_AimOffsetLookAt.cpp
--- a/Engine/Source/Editor/AnimGraph/Private/AnimGraphNode_AimOffsetLookAt.cpp
+++ b/Engine/Source/Editor/AnimGraph/Private/AnimGraphNode_AimOffsetLookAt.cpp
@@ -31,7 +31,7 @@ FText UAnimGraphNode_AimOffsetLookAt::GetTooltipText() const
 
 FText UAnimGraphNode_AimOffsetLookAt::GetNodeTitle(ENodeTitleType::Type TitleType) const
 {
-	UBlendSpace* BlendSpaceToCheck = Node.GetBlendSpace();
+	UBlendSpace* BlendSpaceToCheck = Node.BlendSpace;
 	UEdGraphPin* BlendSpacePin = FindPin(GET_MEMBER_NAME_STRING_CHECKED(FAnimNode_AimOffsetLookAt, BlendSpace));
 	if (BlendSpacePin != nullptr && BlendSpaceToCheck == nullptr)
 	{
@@ -78,7 +78,7 @@ void UAnimGraphNode_AimOffsetLookAt::GetMenuActions(FBlueprintActionDatabaseRegi
 		static void SetNodeBlendSpace(UEdGraphNode* NewNode, bool /*bIsTemplateNode*/, TWeakObjectPtr<UBlendSpace> BlendSpace)
 		{
 			UAnimGraphNode_AimOffsetLookAt* BlendSpaceNode = CastChecked<UAnimGraphNode_AimOffsetLookAt>(NewNode);
-			BlendSpaceNode->Node.SetBlendSpace(BlendSpace.Get());
+			BlendSpaceNode->Node.BlendSpace = BlendSpace.Get();
 		}
 
 		static UBlueprintNodeSpawner* MakeBlendSpaceAction(TSubclassOf<UEdGraphNode> const NodeClass, const UBlendSpace* BlendSpace)
@@ -133,7 +133,7 @@ void UAnimGraphNode_AimOffsetLookAt::GetMenuActions(FBlueprintActionDatabaseRegi
 FBlueprintNodeSignature UAnimGraphNode_AimOffsetLookAt::GetSignature() const
 {
 	FBlueprintNodeSignature NodeSignature = Super::GetSignature();
-	NodeSignature.AddSubObject(Node.GetBlendSpace());
+	NodeSignature.AddSubObject(Node.BlendSpace);
 
 	return NodeSignature;
 }
@@ -142,7 +142,7 @@ void UAnimGraphNode_AimOffsetLookAt::SetAnimationAsset(UAnimationAsset* Asset)
 {
 	if (UBlendSpace* BlendSpace = Cast<UBlendSpace>(Asset))
 	{
-		Node.SetBlendSpace(BlendSpace);
+		Node.BlendSpace = BlendSpace;
 	}
 }
 
@@ -150,7 +150,7 @@ void UAnimGraphNode_AimOffsetLookAt::ValidateAnimNodeDuringCompilation(class USk
 {
 	Super::ValidateAnimNodeDuringCompilation(ForSkeleton, MessageLog);
 
-	UBlendSpace* BlendSpaceToCheck = Node.GetBlendSpace();
+	UBlendSpace* BlendSpaceToCheck = Node.BlendSpace;
 	UEdGraphPin* BlendSpacePin = FindPin(GET_MEMBER_NAME_STRING_CHECKED(FAnimNode_AimOffsetLookAt, BlendSpace));
 	if (BlendSpacePin != nullptr && BlendSpaceToCheck == nullptr)
 	{
@@ -230,7 +230,7 @@ void UAnimGraphNode_AimOffsetLookAt::GetNodeContextMenuActions(UToolMenu* Menu,
 
 void UAnimGraphNode_AimOffsetLookAt::GetAllAnimationSequencesReferred(TArray<UAnimationAsset*>& AnimationAssets) const
 {
-	if (Node.GetBlendSpace())
+	if (Node.BlendSpace)
 	{
 		HandleAnimReferenceCollection(Node.BlendSpace, AnimationAssets);
 	}
diff --git a/Engine/Source/Editor/AnimGraph/Private/AnimGraphNode_AssetPlayerBase.cpp b/Engine/Source/Editor/AnimGraph/Private/AnimGraphNode_AssetPlayerBase.cpp
--- a/Engine/Source/Editor/AnimGraph/Private/AnimGraphNode_AssetPlayerBase.cpp
+++ b/Engine/Source/Editor/AnimGraph/Private/AnimGraphNode_AssetPlayerBase.cpp
@@ -17,6 +17,8 @@
 #include "AnimGraphNode_PoseDriver.h"
 #include "UObject/UObjectIterator.h"
 #include "Animation/AnimLayerInterface.h"
+#include "IAnimBlueprintCompilerHandlerCollection.h"
+#include "AnimBlueprintCompilerHandler_Base.h"
 #include "IAnimBlueprintGeneratedClassCompiledData.h"
 #include "IAnimBlueprintCompilationContext.h"
 #include "Animation/AnimSync.h"
@@ -33,21 +35,9 @@ void UAnimGraphNode_AssetPlayerBase::Serialize(FArchive& Ar)
 
 	if(Ar.CustomVer(FUE5MainStreamObjectVersion::GUID) < FUE5MainStreamObjectVersion::AnimSyncGroupsExplicitSyncMethod)
 	{
-		if(SyncGroup_DEPRECATED.GroupName != NAME_None)
+		if(SyncGroup.GroupName != NAME_None)
 		{
-			SyncGroup_DEPRECATED.Method = EAnimSyncMethod::SyncGroup;
-		}
-	}
-
-	if(Ar.CustomVer(FUE5MainStreamObjectVersion::GUID) < FUE5MainStreamObjectVersion::AnimNodeConstantDataRefactorPhase0)
-	{
-		FStructProperty* NodeProperty = GetFNodeProperty();
-		if(NodeProperty->Struct->IsChildOf(FAnimNode_AssetPlayerBase::StaticStruct()))
-		{
-			FAnimNode_AssetPlayerBase* Node = NodeProperty->ContainerPtrToValuePtr<FAnimNode_AssetPlayerBase>(this);
-			Node->SetGroupName(SyncGroup_DEPRECATED.GroupName);
-			Node->GetGroupRole(SyncGroup_DEPRECATED.GroupRole);
-			Node->SetGroupMethod(SyncGroup_DEPRECATED.Method);
+			SyncGroup.Method = EAnimSyncMethod::SyncGroup;
 		}
 	}
 }
@@ -56,17 +46,12 @@ void UAnimGraphNode_AssetPlayerBase::PostEditChangeProperty(FPropertyChangedEven
 {
 	Super::PostEditChangeProperty(PropertyChangedEvent);
 
-	FStructProperty* NodeProperty = GetFNodeProperty();
-	if(NodeProperty->Struct->IsChildOf(FAnimNode_AssetPlayerBase::StaticStruct()))
+	if(PropertyChangedEvent.GetPropertyName() == GET_MEMBER_NAME_CHECKED(FAnimationGroupReference, Method))
 	{
-		if(PropertyChangedEvent.GetPropertyName() == GET_MEMBER_NAME_CHECKED(FAnimNode_AssetPlayerBase, Method))
+		if(SyncGroup.Method != EAnimSyncMethod::SyncGroup)
 		{
-			FAnimNode_AssetPlayerBase* Node = NodeProperty->ContainerPtrToValuePtr<FAnimNode_AssetPlayerBase>(this);
-			if(Node->GetGroupMethod() != EAnimSyncMethod::SyncGroup)
-			{
-				Node->SetGroupName(NAME_None);
-				Node->GetGroupRole(EAnimGroupRole::CanBeLeader);
-			}
+			SyncGroup.GroupName = NAME_None;
+			SyncGroup.GroupRole = EAnimGroupRole::CanBeLeader;
 		}
 	}
 }
@@ -150,15 +135,9 @@ void UAnimGraphNode_AssetPlayerBase::ValidateAnimNodeDuringCompilation(USkeleton
 {
 	Super::ValidateAnimNodeDuringCompilation(ForSkeleton, MessageLog);
 
-	FStructProperty* NodeProperty = GetFNodeProperty();
-	if(NodeProperty->Struct->IsChildOf(FAnimNode_AssetPlayerBase::StaticStruct()))
+	if(SyncGroup.Method == EAnimSyncMethod::SyncGroup && SyncGroup.GroupName == NAME_None)
 	{
-		FAnimNode_AssetPlayerBase* Node = NodeProperty->ContainerPtrToValuePtr<FAnimNode_AssetPlayerBase>(this);
-
-		if(Node->GetGroupMethod() == EAnimSyncMethod::SyncGroup && Node->GetGroupName() == NAME_None)
-		{
-			MessageLog.Error(*LOCTEXT("NoSyncGroupSupplied", "Node @@ is set to use sync groups, but no sync group has been supplied").ToString(), this);
-		}
+		MessageLog.Error(*LOCTEXT("NoSyncGroupSupplied", "Node @@ is set to use sync groups, but no sync group has been supplied").ToString(), this);
 	}
 }
 
@@ -166,15 +145,9 @@ void UAnimGraphNode_AssetPlayerBase::GetOutputLinkAttributes(FNodeAttributeArray
 {
 	OutAttributes.Add(UE::Anim::FAttributes::Curves);
 	OutAttributes.Add(UE::Anim::FAttributes::Attributes);
-
-	FStructProperty* NodeProperty = GetFNodeProperty();
-	if(NodeProperty->Struct->IsChildOf(FAnimNode_AssetPlayerBase::StaticStruct()))
+	if(SyncGroup.Method == EAnimSyncMethod::Graph)
 	{
-		const FAnimNode_AssetPlayerBase* Node = NodeProperty->ContainerPtrToValuePtr<FAnimNode_AssetPlayerBase>(this);
-		if(Node->GetGroupMethod() == EAnimSyncMethod::Graph)
-		{
-			OutAttributes.Add(UE::Anim::FAnimSync::Attribute);
-		}
+		OutAttributes.Add(UE::Anim::FAnimSync::Attribute);
 	}
 }
 
diff --git a/Engine/Source/Editor/AnimGraph/Private/AnimGraphNode_Base.cpp b/Engine/Source/Editor/AnimGraph/Private/AnimGraphNode_Base.cpp
--- a/Engine/Source/Editor/AnimGraph/Private/AnimGraphNode_Base.cpp
+++ b/Engine/Source/Editor/AnimGraph/Private/AnimGraphNode_Base.cpp
@@ -18,12 +18,9 @@
 #include "Kismet2/CompilerResultsLog.h"
 #include "Subsystems/AssetEditorSubsystem.h"
 #include "AnimBlueprintCompiler.h"
-#include "AnimBlueprintExtension_Base.h"
-#include "AnimBlueprintExtension_Attributes.h"
-#include "AnimBlueprintExtension_PropertyAccess.h"
+#include "AnimBlueprintCompilerHandler_Base.h"
 #include "IAnimBlueprintCompilationContext.h"
 #include "AnimBlueprintCompilationContext.h"
-#include "AnimBlueprintExtension.h"
 #include "FindInBlueprintManager.h"
 #include "UObject/ReleaseObjectVersion.h"
 
@@ -40,8 +37,8 @@ UAnimGraphNode_Base::UAnimGraphNode_Base(const FObjectInitializer& ObjectInitial
 void UAnimGraphNode_Base::ExpandNode(FKismetCompilerContext& CompilerContext, UEdGraph* SourceGraph)
 {
 	TUniquePtr<IAnimBlueprintCompilationContext> CompilationContext = IAnimBlueprintCompilationContext::Get(CompilerContext);
-	UAnimBlueprintExtension_Base* Extension = UAnimBlueprintExtension_Base::GetExtension<UAnimBlueprintExtension_Base>(GetAnimBlueprint());
-	Extension->CreateEvaluationHandlerForNode(*CompilationContext.Get(), this);
+	FAnimBlueprintCompilerHandler_Base* Handler = CompilationContext->GetHandler<FAnimBlueprintCompilerHandler_Base>("AnimBlueprintCompilerHandler_Base");
+	Handler->CreateEvaluationHandlerForNode(*CompilationContext.Get(), this);
 }
 
 void UAnimGraphNode_Base::PreEditChange(FProperty* PropertyThatWillChange)
@@ -101,22 +98,6 @@ void UAnimGraphNode_Base::Serialize(FArchive& Ar)
 	}
 }
 
-void UAnimGraphNode_Base::PostPlacedNewNode()
-{
-	Super::PostPlacedNewNode();
-
-	// This makes sure that all anim BP extensions are registered that this node needs
-	UAnimBlueprintExtension::RequestExtensionsForNode(this);
-}
-
-void UAnimGraphNode_Base::DestroyNode()
-{
-	// This node may have been the last using its extension, so refresh
-	GetAnimBlueprint()->RequestRefreshExtensions();
-
-	Super::DestroyNode();
-}
-
 void UAnimGraphNode_Base::CreateOutputPins()
 {
 	if (!IsSinkNode())
@@ -159,21 +140,6 @@ void UAnimGraphNode_Base::CopyTermDefaultsToDefaultObject(IAnimBlueprintCopyTerm
 	OnCopyTermDefaultsToDefaultObject(InCompilationContext, InPerNodeContext, OutCompiledData);
 }
 
-void UAnimGraphNode_Base::OverrideAssets(IAnimBlueprintNodeOverrideAssetsContext& InContext) const
-{
-	if(InContext.GetAssets().Num() > 0)
-	{
-		if(UAnimationAsset* AnimationAsset = Cast<UAnimationAsset>(InContext.GetAssets()[0]))
-		{
-			// Call the legacy implementation
-PRAGMA_DISABLE_DEPRECATION_WARNINGS
-			InContext.GetAnimNode<FAnimNode_Base>().OverrideAsset(AnimationAsset);
-PRAGMA_ENABLE_DEPRECATION_WARNINGS
-		}
-	}
-	OnOverrideAssets(InContext);
-}
-
 void UAnimGraphNode_Base::InternalPinCreation(TArray<UEdGraphPin*>* OldPins)
 {
 	// preload required assets first before creating pins
@@ -417,10 +383,10 @@ void UAnimGraphNode_Base::GetPinHoverText(const UEdGraphPin& Pin, FString& Hover
 
 void UAnimGraphNode_Base::ProcessDuringCompilation(IAnimBlueprintCompilationContext& InCompilationContext, IAnimBlueprintGeneratedClassCompiledData& OutCompiledData)
 {
-	UAnimBlueprintExtension_Base* Extension = UAnimBlueprintExtension::GetExtension<UAnimBlueprintExtension_Base>(GetAnimBlueprint());
+	FAnimBlueprintCompilerHandler_Base* HandlerBase = InCompilationContext.GetHandler<FAnimBlueprintCompilerHandler_Base>("AnimBlueprintCompilerHandler_Base");
 
 	// Record pose pins for later patchup and gather pins that have an associated evaluation handler
-	Extension->AddStructEvalHandlers(this, InCompilationContext, OutCompiledData);
+	HandlerBase->AddStructEvalHandlers(this, InCompilationContext, OutCompiledData);
 
 	// Call the override point
 	OnProcessDuringCompilation(InCompilationContext, OutCompiledData);
@@ -546,12 +512,6 @@ bool UAnimGraphNode_Base::IsPinExposedAndLinked(const FString& InPinName, const
 	return Pin != nullptr && Pin->LinkedTo.Num() > 0 && Pin->LinkedTo[0] != nullptr;
 }
 
-bool UAnimGraphNode_Base::IsPinExposedAndBound(const FString& InPinName, const EEdGraphPinDirection InDirection) const
-{
-	UEdGraphPin* Pin = FindPin(InPinName, InDirection);
-	return Pin != nullptr && Pin->LinkedTo.Num() == 0 && PropertyBindings.Find(Pin->GetFName()) != nullptr;
-}
-
 void UAnimGraphNode_Base::PinConnectionListChanged(UEdGraphPin* Pin)
 {
 	if(Pin->LinkedTo.Num() > 0)
diff --git a/Engine/Source/Editor/AnimGraph/Private/AnimGraphNode_BlendListBase.cpp b/Engine/Source/Editor/AnimGraph/Private/AnimGraphNode_BlendListBase.cpp
--- a/Engine/Source/Editor/AnimGraph/Private/AnimGraphNode_BlendListBase.cpp
+++ b/Engine/Source/Editor/AnimGraph/Private/AnimGraphNode_BlendListBase.cpp
@@ -113,7 +113,7 @@ void UAnimGraphNode_BlendListBase::GetOutputLinkAttributes(FNodeAttributeArray&
 	{
 		const FAnimNode_BlendListBase* BaseNode = GetFNodeProperty()->ContainerPtrToValuePtr<FAnimNode_BlendListBase>(this);
 
-		if(BaseNode->GetTransitionType() == EBlendListTransitionType::Inertialization)
+		if(BaseNode->TransitionType == EBlendListTransitionType::Inertialization)
 		{
 			OutAttributes.Add(UE::Anim::IInertializationRequester::Attribute);
 		}
diff --git a/Engine/Source/Editor/AnimGraph/Private/AnimGraphNode_BlendListByEnum.cpp b/Engine/Source/Editor/AnimGraph/Private/AnimGraphNode_BlendListByEnum.cpp
--- a/Engine/Source/Editor/AnimGraph/Private/AnimGraphNode_BlendListByEnum.cpp
+++ b/Engine/Source/Editor/AnimGraph/Private/AnimGraphNode_BlendListByEnum.cpp
@@ -307,9 +307,8 @@ void UAnimGraphNode_BlendListByEnum::BakeDataDuringCompilation(class FCompilerRe
 		BoundEnum->ConditionalPostLoad();
 
 		// Zero the array out so it looks up the default value, and stat counting at index 1
-		TArray<int32> EnumToPoseIndex;
-		EnumToPoseIndex.Empty();
-		EnumToPoseIndex.AddZeroed(BoundEnum->NumEnums());
+		Node.EnumToPoseIndex.Empty();
+		Node.EnumToPoseIndex.AddZeroed(BoundEnum->NumEnums());
 		int32 PinIndex = 1;
 
 		// Run thru the enum entries
@@ -320,7 +319,7 @@ void UAnimGraphNode_BlendListByEnum::BakeDataDuringCompilation(class FCompilerRe
 
 			if (EnumIndex != INDEX_NONE)
 			{
-				EnumToPoseIndex[EnumIndex] = PinIndex;
+				Node.EnumToPoseIndex[EnumIndex] = PinIndex;
 			}
 			else
 			{
@@ -329,8 +328,6 @@ void UAnimGraphNode_BlendListByEnum::BakeDataDuringCompilation(class FCompilerRe
 
 			++PinIndex;
 		}
-
-		Node.SetEnumToPoseIndex(EnumToPoseIndex);
 	}
 }
 
diff --git a/Engine/Source/Editor/AnimGraph/Private/AnimGraphNode_BlendListByInt.cpp b/Engine/Source/Editor/AnimGraph/Private/AnimGraphNode_BlendListByInt.cpp
--- a/Engine/Source/Editor/AnimGraph/Private/AnimGraphNode_BlendListByInt.cpp
+++ b/Engine/Source/Editor/AnimGraph/Private/AnimGraphNode_BlendListByInt.cpp
@@ -32,8 +32,6 @@ FText UAnimGraphNode_BlendListByInt::GetNodeTitle(ENodeTitleType::Type TitleType
 
 void UAnimGraphNode_BlendListByInt::PostPlacedNewNode()
 {
-	Super::PostPlacedNewNode();
-
 	Node.AddPose();
 	ReconstructNode();
 }
diff --git a/Engine/Source/Editor/AnimGraph/Private/AnimGraphNode_BlendSpaceEvaluator.cpp b/Engine/Source/Editor/AnimGraph/Private/AnimGraphNode_BlendSpaceEvaluator.cpp
--- a/Engine/Source/Editor/AnimGraph/Private/AnimGraphNode_BlendSpaceEvaluator.cpp
+++ b/Engine/Source/Editor/AnimGraph/Private/AnimGraphNode_BlendSpaceEvaluator.cpp
@@ -4,7 +4,6 @@
 #include "ToolMenus.h"
 #include "AnimGraphCommands.h"
 #include "Kismet2/CompilerResultsLog.h"
-#include "IAnimBlueprintNodeOverrideAssetsContext.h"
 
 /////////////////////////////////////////////////////
 // UAnimGraphNode_BlendSpaceEvaluator
@@ -40,11 +39,11 @@ FText UAnimGraphNode_BlendSpaceEvaluator::GetNodeTitleForBlendSpace(ENodeTitleTy
 		TitleArgs.Add(TEXT("BlendSpaceName"), BlendSpaceName);
 		FText Title = FText::Format(LOCTEXT("BlendSpaceEvaluatorFullTitle", "{BlendSpaceName}\nBlendspace Evaluator"), TitleArgs);
 
-		if ((TitleType == ENodeTitleType::FullTitle) && (Node.GetGroupName() != NAME_None))
+		if ((TitleType == ENodeTitleType::FullTitle) && (SyncGroup.GroupName != NAME_None))
 		{
 			FFormatNamedArguments Args;
 			Args.Add(TEXT("Title"), Title);
-			Args.Add(TEXT("SyncGroupName"), FText::FromName(Node.GetGroupName()));
+			Args.Add(TEXT("SyncGroupName"), FText::FromName(SyncGroup.GroupName));
 			Title = FText::Format(LOCTEXT("BlendSpaceNodeGroupSubtitle", "{Title}\nSync group {SyncGroupName}"), Args);
 		}
 		// FText::Format() is slow, so we cache this to save on performance
@@ -56,7 +55,7 @@ FText UAnimGraphNode_BlendSpaceEvaluator::GetNodeTitleForBlendSpace(ENodeTitleTy
 
 FText UAnimGraphNode_BlendSpaceEvaluator::GetNodeTitle(ENodeTitleType::Type TitleType) const
 {
-	if (Node.GetBlendSpace() == nullptr)
+	if (Node.BlendSpace == nullptr)
 	{
 		// we may have a valid variable connected or default pin value
 		UEdGraphPin* BlendSpacePin = FindPin(GET_MEMBER_NAME_STRING_CHECKED(FAnimNode_BlendSpacePlayer, BlendSpace));
@@ -84,7 +83,7 @@ FText UAnimGraphNode_BlendSpaceEvaluator::GetNodeTitle(ENodeTitleType::Type Titl
 	//        choose to mark this dirty when that happens for this to properly work
 	else //if (!CachedNodeTitles.IsTitleCached(TitleType, this))
 	{
-		return GetNodeTitleForBlendSpace(TitleType, Node.GetBlendSpace());
+		return GetNodeTitleForBlendSpace(TitleType, Node.BlendSpace);
 	}
 }
 
@@ -98,7 +97,7 @@ void UAnimGraphNode_BlendSpaceEvaluator::ValidateAnimNodeDuringCompilation(class
 {
 	Super::ValidateAnimNodeDuringCompilation(ForSkeleton, MessageLog);
 
-	UBlendSpace* BlendSpaceToCheck = Node.GetBlendSpace();
+	UBlendSpace* BlendSpaceToCheck = Node.BlendSpace;
 	UEdGraphPin* BlendSpacePin = FindPin(GET_MEMBER_NAME_STRING_CHECKED(FAnimNode_BlendSpaceEvaluator, BlendSpace));
 	if (BlendSpacePin != nullptr && BlendSpaceToCheck == nullptr)
 	{
@@ -137,7 +136,10 @@ void UAnimGraphNode_BlendSpaceEvaluator::ValidateAnimNodeDuringCompilation(class
 void UAnimGraphNode_BlendSpaceEvaluator::BakeDataDuringCompilation(class FCompilerResultsLog& MessageLog)
 {
 	UAnimBlueprint* AnimBlueprint = GetAnimBlueprint();
-	AnimBlueprint->FindOrAddGroup(Node.GetGroupName());
+	AnimBlueprint->FindOrAddGroup(SyncGroup.GroupName);
+	Node.GroupName = SyncGroup.GroupName;
+	Node.GroupRole = SyncGroup.GroupRole;
+	Node.Method = SyncGroup.Method;
 }
 
 void UAnimGraphNode_BlendSpaceEvaluator::GetNodeContextMenuActions(UToolMenu* Menu, UGraphNodeContextMenuContext* Context) const
@@ -157,19 +159,7 @@ void UAnimGraphNode_BlendSpaceEvaluator::SetAnimationAsset(UAnimationAsset* Asse
 {
 	if (UBlendSpace* BlendSpace = Cast<UBlendSpace>(Asset))
 	{
-		Node.SetBlendSpace(BlendSpace);
-	}
-}
-
-void UAnimGraphNode_BlendSpaceEvaluator::OnOverrideAssets(IAnimBlueprintNodeOverrideAssetsContext& InContext) const
-{
-	if(InContext.GetAssets().Num() > 0)
-	{
-		if (UBlendSpace* BlendSpace = Cast<UBlendSpace>(InContext.GetAssets()[0]))
-		{
-			FAnimNode_BlendSpaceEvaluator& AnimNode = InContext.GetAnimNode<FAnimNode_BlendSpaceEvaluator>();
-			AnimNode.SetBlendSpace(BlendSpace);
-		}
+		Node.BlendSpace = BlendSpace;
 	}
 }
 
@@ -180,7 +170,7 @@ bool UAnimGraphNode_BlendSpaceEvaluator::DoesSupportTimeForTransitionGetter() co
 
 UAnimationAsset* UAnimGraphNode_BlendSpaceEvaluator::GetAnimationAsset() const 
 {
-	UBlendSpace* BlendSpace = Node.GetBlendSpace();
+	UBlendSpace* BlendSpace = Node.BlendSpace;
 	UEdGraphPin* BlendSpacePin = FindPin(GET_MEMBER_NAME_STRING_CHECKED(FAnimNode_BlendSpaceEvaluator, BlendSpace));
 	if (BlendSpacePin != nullptr && BlendSpace == nullptr)
 	{
diff --git a/Engine/Source/Editor/AnimGraph/Private/AnimGraphNode_BlendSpaceGraphBase.cpp b/Engine/Source/Editor/AnimGraph/Private/AnimGraphNode_BlendSpaceGraphBase.cpp
--- a/Engine/Source/Editor/AnimGraph/Private/AnimGraphNode_BlendSpaceGraphBase.cpp
+++ b/Engine/Source/Editor/AnimGraph/Private/AnimGraphNode_BlendSpaceGraphBase.cpp
@@ -20,14 +20,12 @@
 #include "EdGraphUtilities.h"
 #include "IAnimBlueprintCompilationContext.h"
 #include "AnimationBlendSpaceSampleGraph.h"
-#include "AnimBlueprintExtension.h"
 #include "AnimGraphNode_BlendSpaceSampleResult.h"
 #include "AnimGraphNode_SequencePlayer.h"
 #include "Kismet2/Kismet2NameValidators.h"
 #include "AnimNodes/AnimNode_BlendSpaceGraphBase.h"
 #include "BlueprintEditor.h"
 #include "Animation/AnimSequence.h"
-#include "AnimBlueprintExtension_BlendSpaceGraph.h"
 
 #define LOCTEXT_NAMESPACE "UAnimGraphNode_BlendSpaceGraphBase"
 
@@ -124,9 +122,10 @@ void UAnimGraphNode_BlendSpaceGraphBase::OnProcessDuringCompilation(IAnimBluepri
 
 void UAnimGraphNode_BlendSpaceGraphBase::OnCopyTermDefaultsToDefaultObject(IAnimBlueprintCopyTermDefaultsContext& InCompilationContext, IAnimBlueprintNodeCopyTermDefaultsContext& InPerNodeContext, IAnimBlueprintGeneratedClassCompiledData& OutCompiledData)
 {
+	UAnimGraphNode_BlendSpaceGraphBase* TrueNode = InCompilationContext.GetMessageLog().FindSourceObjectTypeChecked<UAnimGraphNode_BlendSpaceGraphBase>(this);
+
 	FAnimNode_BlendSpaceGraphBase* DestinationNode = reinterpret_cast<FAnimNode_BlendSpaceGraphBase*>(InPerNodeContext.GetDestinationPtr());
-	UAnimBlueprintExtension_BlendSpaceGraph* Extension = UAnimBlueprintExtension::GetExtension<UAnimBlueprintExtension_BlendSpaceGraph>(GetAnimBlueprint());
-	DestinationNode->BlendSpace = Extension->AddBlendSpace(BlendSpace);
+	DestinationNode->BlendSpace = OutCompiledData.AddBlendSpace(BlendSpace);
 }
 
 void UAnimGraphNode_BlendSpaceGraphBase::SetupFromAsset(UBlendSpace* InBlendSpace, bool bInIsTemplateNode)
@@ -323,8 +322,6 @@ void UAnimGraphNode_BlendSpaceGraphBase::PostPasteNode()
 
 void UAnimGraphNode_BlendSpaceGraphBase::PostPlacedNewNode()
 {
-	Super::PostPlacedNewNode();
-	
 	// Create a new graph & blendspace if we havent been set up already
 	if(BlendSpaceGraph == nullptr)
 	{
@@ -424,7 +421,7 @@ UAnimationBlendSpaceSampleGraph* UAnimGraphNode_BlendSpaceGraphBase::AddGraphInt
 		FGraphNodeCreator<UAnimGraphNode_SequencePlayer> SequencePlayerNodeCreator(*SampleGraph);
 		UAnimGraphNode_SequencePlayer* SequencePlayer = SequencePlayerNodeCreator.CreateNode();
 		SequencePlayer->SetAnimationAsset(InSequence);
-		SequencePlayer->Node.SetGroupMethod(EAnimSyncMethod::Graph);
+		SequencePlayer->SyncGroup.Method = EAnimSyncMethod::Graph;
 		SequencePlayerNodeCreator.Finalize();
 
 		// Offset node in X
@@ -496,9 +493,4 @@ void UAnimGraphNode_BlendSpaceGraphBase::GetInputLinkAttributes(FNodeAttributeAr
 	}
 }
 
-void UAnimGraphNode_BlendSpaceGraphBase::GetRequiredExtensions(TArray<TSubclassOf<UAnimBlueprintExtension>>& OutExtensions) const
-{
-	OutExtensions.Add(UAnimBlueprintExtension_BlendSpaceGraph::StaticClass());
-}
-
 #undef LOCTEXT_NAMESPACE
diff --git a/Engine/Source/Editor/AnimGraph/Private/AnimGraphNode_BlendSpacePlayer.cpp b/Engine/Source/Editor/AnimGraph/Private/AnimGraphNode_BlendSpacePlayer.cpp
--- a/Engine/Source/Editor/AnimGraph/Private/AnimGraphNode_BlendSpacePlayer.cpp
+++ b/Engine/Source/Editor/AnimGraph/Private/AnimGraphNode_BlendSpacePlayer.cpp
@@ -11,8 +11,6 @@
 #include "Animation/AimOffsetBlendSpace.h"
 #include "Animation/AimOffsetBlendSpace1D.h"
 #include "AnimGraphCommands.h"
-#include "IAnimBlueprintCopyTermDefaultsContext.h"
-#include "IAnimBlueprintNodeOverrideAssetsContext.h"
 
 /////////////////////////////////////////////////////
 // UAnimGraphNode_BlendSpacePlayer
@@ -52,12 +50,12 @@ FText UAnimGraphNode_BlendSpacePlayer::GetNodeTitleForBlendSpace(ENodeTitleType:
 			FFormatNamedArguments Args;
 			Args.Add(TEXT("Title"), Title);
 
-			if(Node.GetGroupMethod() == EAnimSyncMethod::SyncGroup)
+			if(SyncGroup.Method == EAnimSyncMethod::SyncGroup)
 			{
-				Args.Add(TEXT("SyncGroupName"), FText::FromName(Node.GetGroupName()));
+				Args.Add(TEXT("SyncGroupName"), FText::FromName(SyncGroup.GroupName));
 				Title = FText::Format(LOCTEXT("BlendSpaceNodeGroupSubtitle", "{Title}\nSync group {SyncGroupName}"), Args);
 			}
-			else if(Node.GetGroupMethod() == EAnimSyncMethod::Graph)
+			else if(SyncGroup.Method == EAnimSyncMethod::Graph)
 			{
 				Title = FText::Format(LOCTEXT("BlendSpaceNodeGroupSubtitle", "{Title}\nGraph sync group"), Args);
 
@@ -86,7 +84,7 @@ FText UAnimGraphNode_BlendSpacePlayer::GetNodeTitleForBlendSpace(ENodeTitleType:
 
 FText UAnimGraphNode_BlendSpacePlayer::GetNodeTitle(ENodeTitleType::Type TitleType) const
 {	
-	if (Node.GetBlendSpace() == nullptr)
+	if (Node.BlendSpace == nullptr)
 	{
 		// we may have a valid variable connected or default pin value
 		UEdGraphPin* BlendSpacePin = FindPin(GET_MEMBER_NAME_STRING_CHECKED(FAnimNode_BlendSpacePlayer, BlendSpace));
@@ -114,7 +112,7 @@ FText UAnimGraphNode_BlendSpacePlayer::GetNodeTitle(ENodeTitleType::Type TitleTy
 	//        choose to mark this dirty when that happens for this to properly work
 	else //if (!CachedNodeTitles.IsTitleCached(TitleType, this))
 	{
-		return GetNodeTitleForBlendSpace(TitleType, Node.GetBlendSpace());
+		return GetNodeTitleForBlendSpace(TitleType, Node.BlendSpace);
 	}
 }
 
@@ -122,7 +120,7 @@ void UAnimGraphNode_BlendSpacePlayer::ValidateAnimNodeDuringCompilation(class US
 {
 	Super::ValidateAnimNodeDuringCompilation(ForSkeleton, MessageLog);
 
-	UBlendSpace* BlendSpaceToCheck = Node.GetBlendSpace();
+	UBlendSpace* BlendSpaceToCheck = Node.BlendSpace;
 	UEdGraphPin* BlendSpacePin = FindPin(GET_MEMBER_NAME_STRING_CHECKED(FAnimNode_BlendSpacePlayer, BlendSpace));
 	if (BlendSpacePin != nullptr && BlendSpaceToCheck == nullptr)
 	{
@@ -161,7 +159,10 @@ void UAnimGraphNode_BlendSpacePlayer::ValidateAnimNodeDuringCompilation(class US
 void UAnimGraphNode_BlendSpacePlayer::BakeDataDuringCompilation(class FCompilerResultsLog& MessageLog)
 {
 	UAnimBlueprint* AnimBlueprint = GetAnimBlueprint();
-	AnimBlueprint->FindOrAddGroup(Node.GetGroupName());
+	AnimBlueprint->FindOrAddGroup(SyncGroup.GroupName);
+	Node.GroupName = SyncGroup.GroupName;
+	Node.GroupRole = SyncGroup.GroupRole;
+	Node.Method = SyncGroup.Method;
 }
 
 void UAnimGraphNode_BlendSpacePlayer::GetNodeContextMenuActions(UToolMenu* Menu, UGraphNodeContextMenuContext* Context) const
@@ -185,7 +186,7 @@ void UAnimGraphNode_BlendSpacePlayer::GetMenuActions(FBlueprintActionDatabaseReg
 		static void SetNodeBlendSpace(UEdGraphNode* NewNode, bool /*bIsTemplateNode*/, TWeakObjectPtr<UBlendSpace> BlendSpace)
 		{
 			UAnimGraphNode_BlendSpacePlayer* BlendSpaceNode = CastChecked<UAnimGraphNode_BlendSpacePlayer>(NewNode);
-			BlendSpaceNode->Node.SetBlendSpace(BlendSpace.Get());
+			BlendSpaceNode->Node.BlendSpace = BlendSpace.Get();
 		}
 
 		static UBlueprintNodeSpawner* MakeBlendSpaceAction(TSubclassOf<UEdGraphNode> const NodeClass, const UBlendSpace* BlendSpace)
@@ -240,7 +241,7 @@ void UAnimGraphNode_BlendSpacePlayer::GetMenuActions(FBlueprintActionDatabaseReg
 FBlueprintNodeSignature UAnimGraphNode_BlendSpacePlayer::GetSignature() const
 {
 	FBlueprintNodeSignature NodeSignature = Super::GetSignature();
-	NodeSignature.AddSubObject(Node.GetBlendSpace());
+	NodeSignature.AddSubObject(Node.BlendSpace);
 
 	return NodeSignature;
 }
@@ -249,25 +250,13 @@ void UAnimGraphNode_BlendSpacePlayer::SetAnimationAsset(UAnimationAsset* Asset)
 {
 	if (UBlendSpace* BlendSpace = Cast<UBlendSpace>(Asset))
 	{
-		Node.SetBlendSpace(BlendSpace);
-	}
-}
-
-void UAnimGraphNode_BlendSpacePlayer::OnOverrideAssets(IAnimBlueprintNodeOverrideAssetsContext& InContext) const
-{
-	if(InContext.GetAssets().Num() > 0)
-	{
-		if (UBlendSpace* BlendSpace = Cast<UBlendSpace>(InContext.GetAssets()[0]))
-		{
-			FAnimNode_BlendSpacePlayer& AnimNode = InContext.GetAnimNode<FAnimNode_BlendSpacePlayer>();
-			AnimNode.SetBlendSpace(BlendSpace);
-		}
+		Node.BlendSpace = BlendSpace;
 	}
 }
 
 void UAnimGraphNode_BlendSpacePlayer::GetAllAnimationSequencesReferred(TArray<UAnimationAsset*>& AnimationAssets) const
 {
-	if(Node.GetBlendSpace())
+	if(Node.BlendSpace)
 	{
 		HandleAnimReferenceCollection(Node.BlendSpace, AnimationAssets);
 	}
@@ -285,7 +274,7 @@ bool UAnimGraphNode_BlendSpacePlayer::DoesSupportTimeForTransitionGetter() const
 
 UAnimationAsset* UAnimGraphNode_BlendSpacePlayer::GetAnimationAsset() const 
 {
-	UBlendSpace* BlendSpace = Node.GetBlendSpace();
+	UBlendSpace* BlendSpace = Node.BlendSpace;
 	UEdGraphPin* BlendSpacePin = FindPin(GET_MEMBER_NAME_STRING_CHECKED(FAnimNode_BlendSpacePlayer, BlendSpace));
 	if (BlendSpacePin != nullptr && BlendSpace == nullptr)
 	{
diff --git a/Engine/Source/Editor/AnimGraph/Private/AnimGraphNode_CustomProperty.cpp b/Engine/Source/Editor/AnimGraph/Private/AnimGraphNode_CustomProperty.cpp
--- a/Engine/Source/Editor/AnimGraph/Private/AnimGraphNode_CustomProperty.cpp
+++ b/Engine/Source/Editor/AnimGraph/Private/AnimGraphNode_CustomProperty.cpp
@@ -6,8 +6,18 @@
 #include "Kismet2/BlueprintEditorUtils.h"
 #include "DetailLayoutBuilder.h"
 #include "AnimationGraphSchema.h"
+#include "Animation/AnimInstance.h"
+#include "Animation/AnimLayerInterface.h"
 #include "DetailCategoryBuilder.h"
+#include "Widgets/Text/STextBlock.h"
+#include "Widgets/Input/SCheckBox.h"
+#include "Widgets/SBoxPanel.h"
+#include "DetailWidgetRow.h"
 #include "PropertyCustomizationHelpers.h"
+#include "KismetCompilerMisc.h"
+#include "AnimBlueprintCompilerHandler_Base.h"
+#include "IAnimBlueprintCompilerHandlerCollection.h"
+#include "IAnimBlueprintCompilerHandler.h"
 #include "KismetCompiler.h"
 #include "IAnimBlueprintCompilationContext.h"
 
diff --git a/Engine/Source/Editor/AnimGraph/Private/AnimGraphNode_IdentityPose.cpp b/Engine/Source/Editor/AnimGraph/Private/AnimGraphNode_IdentityPose.cpp
--- a/Engine/Source/Editor/AnimGraph/Private/AnimGraphNode_IdentityPose.cpp
+++ b/Engine/Source/Editor/AnimGraph/Private/AnimGraphNode_IdentityPose.cpp
@@ -10,7 +10,7 @@
 UAnimGraphNode_IdentityPose::UAnimGraphNode_IdentityPose(const FObjectInitializer& ObjectInitializer)
 	: Super(ObjectInitializer)
 {
-	Node.SetRefPoseType(EIT_Additive);
+	Node.RefPoseType = EIT_Additive;
 }
 
 FText UAnimGraphNode_IdentityPose::GetTooltipText() const
diff --git a/Engine/Source/Editor/AnimGraph/Private/AnimGraphNode_LinkedAnimGraphBase.cpp b/Engine/Source/Editor/AnimGraph/Private/AnimGraphNode_LinkedAnimGraphBase.cpp
--- a/Engine/Source/Editor/AnimGraph/Private/AnimGraphNode_LinkedAnimGraphBase.cpp
+++ b/Engine/Source/Editor/AnimGraph/Private/AnimGraphNode_LinkedAnimGraphBase.cpp
@@ -18,7 +18,6 @@
 #include "Animation/AnimNode_LinkedAnimGraph.h"
 #include "AnimGraphAttributes.h"
 #include "IAnimBlueprintCopyTermDefaultsContext.h"
-#include "AnimBlueprintExtension_LinkedAnimGraph.h"
 
 #define LOCTEXT_NAMESPACE "LinkedAnimGraph"
 
@@ -490,9 +489,4 @@ void UAnimGraphNode_LinkedAnimGraphBase::OnCopyTermDefaultsToDefaultObject(IAnim
 	DestinationNode->NodeIndex = InPerNodeContext.GetNodePropertyIndex();
 }
 
-void UAnimGraphNode_LinkedAnimGraphBase::GetRequiredExtensions(TArray<TSubclassOf<UAnimBlueprintExtension>>& OutExtensions) const
-{
-	OutExtensions.Add(UAnimBlueprintExtension_LinkedAnimGraph::StaticClass());
-}
-
 #undef LOCTEXT_NAMESPACE
diff --git a/Engine/Source/Editor/AnimGraph/Private/AnimGraphNode_LocalRefPose.cpp b/Engine/Source/Editor/AnimGraph/Private/AnimGraphNode_LocalRefPose.cpp
--- a/Engine/Source/Editor/AnimGraph/Private/AnimGraphNode_LocalRefPose.cpp
+++ b/Engine/Source/Editor/AnimGraph/Private/AnimGraphNode_LocalRefPose.cpp
@@ -10,7 +10,7 @@
 UAnimGraphNode_LocalRefPose::UAnimGraphNode_LocalRefPose(const FObjectInitializer& ObjectInitializer)
 	: Super(ObjectInitializer)
 {
-	Node.SetRefPoseType(EIT_LocalSpace);
+	Node.RefPoseType = EIT_LocalSpace;
 }
 
 FText UAnimGraphNode_LocalRefPose::GetTooltipText() const
diff --git a/Engine/Source/Editor/AnimGraph/Private/AnimGraphNode_MultiWayBlend.cpp b/Engine/Source/Editor/AnimGraph/Private/AnimGraphNode_MultiWayBlend.cpp
--- a/Engine/Source/Editor/AnimGraph/Private/AnimGraphNode_MultiWayBlend.cpp
+++ b/Engine/Source/Editor/AnimGraph/Private/AnimGraphNode_MultiWayBlend.cpp
@@ -94,8 +94,6 @@ void UAnimGraphNode_MultiWayBlend::RemovePinFromBlendNode(UEdGraphPin* Pin)
 
 void UAnimGraphNode_MultiWayBlend::PostPlacedNewNode()
 {
-	Super::PostPlacedNewNode();
-	
 	// Make sure we start out with two inputs
 	Node.AddPose();
 	Node.AddPose();
diff --git a/Engine/Source/Editor/AnimGraph/Private/AnimGraphNode_PoseHandler.cpp b/Engine/Source/Editor/AnimGraph/Private/AnimGraphNode_PoseHandler.cpp
--- a/Engine/Source/Editor/AnimGraph/Private/AnimGraphNode_PoseHandler.cpp
+++ b/Engine/Source/Editor/AnimGraph/Private/AnimGraphNode_PoseHandler.cpp
@@ -4,7 +4,6 @@
 #include "Animation/PoseAsset.h"
 #include "AnimNodes/AnimNode_PoseHandler.h"
 #include "Kismet2/CompilerResultsLog.h"
-#include "IAnimBlueprintNodeOverrideAssetsContext.h"
 
 UAnimGraphNode_PoseHandler::UAnimGraphNode_PoseHandler(const FObjectInitializer& ObjectInitializer)
 	:Super(ObjectInitializer)
@@ -49,18 +48,6 @@ void UAnimGraphNode_PoseHandler::SetAnimationAsset(UAnimationAsset* Asset)
 	}
 }
 
-void UAnimGraphNode_PoseHandler::OnOverrideAssets(IAnimBlueprintNodeOverrideAssetsContext& InContext) const
-{
-	if(InContext.GetAssets().Num() > 0)
-	{
-		if (UPoseAsset* PoseAsset = Cast<UPoseAsset>(InContext.GetAssets()[0]))
-		{
-			FAnimNode_PoseHandler& AnimNode = InContext.GetAnimNode<FAnimNode_PoseHandler>();
-			AnimNode.SetPoseAsset(PoseAsset);
-		}
-	}
-}
-
 void UAnimGraphNode_PoseHandler::PreloadRequiredAssets()
 {
 	PreloadObject(GetPoseHandlerNode()->PoseAsset);
diff --git a/Engine/Source/Editor/AnimGraph/Private/AnimGraphNode_RefPoseBase.cpp b/Engine/Source/Editor/AnimGraph/Private/AnimGraphNode_RefPoseBase.cpp
--- a/Engine/Source/Editor/AnimGraph/Private/AnimGraphNode_RefPoseBase.cpp
+++ b/Engine/Source/Editor/AnimGraph/Private/AnimGraphNode_RefPoseBase.cpp
@@ -17,7 +17,7 @@ FString UAnimGraphNode_RefPoseBase::GetNodeCategory() const
 
 FLinearColor UAnimGraphNode_RefPoseBase::GetNodeTitleColor() const
 {
-	if ( Node.GetRefPoseType() == EIT_Additive )
+	if ( Node.RefPoseType == EIT_Additive )
 	{
 		return FLinearColor(0.10f, 0.60f, 0.12f);
 	}
diff --git a/Engine/Source/Editor/AnimGraph/Private/AnimGraphNode_Root.cpp b/Engine/Source/Editor/AnimGraph/Private/AnimGraphNode_Root.cpp
--- a/Engine/Source/Editor/AnimGraph/Private/AnimGraphNode_Root.cpp
+++ b/Engine/Source/Editor/AnimGraph/Private/AnimGraphNode_Root.cpp
@@ -3,6 +3,7 @@
 #include "AnimGraphNode_Root.h"
 #include "GraphEditorSettings.h"
 #include "AnimBlueprintCompiler.h"
+#include "AnimBlueprintCompilerHandler_Base.h"
 #include "IAnimBlueprintCompilationContext.h"
 #include "IAnimBlueprintCopyTermDefaultsContext.h"
 
@@ -91,7 +92,7 @@ void UAnimGraphNode_Root::OnProcessDuringCompilation(IAnimBlueprintCompilationCo
 {
 	UAnimGraphNode_Root* TrueNode = InCompilationContext.GetMessageLog().FindSourceObjectTypeChecked<UAnimGraphNode_Root>(this);
 
-	Node.SetName(TrueNode->GetGraph()->GetFName());
+	Node.Name = TrueNode->GetGraph()->GetFName();
 }
 
 void UAnimGraphNode_Root::OnCopyTermDefaultsToDefaultObject(IAnimBlueprintCopyTermDefaultsContext& InCompilationContext, IAnimBlueprintNodeCopyTermDefaultsContext& InPerNodeContext, IAnimBlueprintGeneratedClassCompiledData& OutCompiledData)
@@ -99,7 +100,7 @@ void UAnimGraphNode_Root::OnCopyTermDefaultsToDefaultObject(IAnimBlueprintCopyTe
 	UAnimGraphNode_Root* TrueNode = InCompilationContext.GetMessageLog().FindSourceObjectTypeChecked<UAnimGraphNode_Root>(this);
 
 	FAnimNode_Root* DestinationNode = reinterpret_cast<FAnimNode_Root*>(InPerNodeContext.GetDestinationPtr());
-	DestinationNode->SetName(TrueNode->GetGraph()->GetFName());
+	DestinationNode->Name = TrueNode->GetGraph()->GetFName();
 }
 
 #undef LOCTEXT_NAMESPACE
diff --git a/Engine/Source/Editor/AnimGraph/Private/AnimGraphNode_RotationOffsetBlendSpace.cpp b/Engine/Source/Editor/AnimGraph/Private/AnimGraphNode_RotationOffsetBlendSpace.cpp
--- a/Engine/Source/Editor/AnimGraph/Private/AnimGraphNode_RotationOffsetBlendSpace.cpp
+++ b/Engine/Source/Editor/AnimGraph/Private/AnimGraphNode_RotationOffsetBlendSpace.cpp
@@ -34,7 +34,7 @@ FText UAnimGraphNode_RotationOffsetBlendSpace::GetTooltipText() const
 
 FText UAnimGraphNode_RotationOffsetBlendSpace::GetNodeTitle(ENodeTitleType::Type TitleType) const
 {
-	UBlendSpace* BlendSpaceToCheck = Node.GetBlendSpace();
+	UBlendSpace* BlendSpaceToCheck = Node.BlendSpace;
 	UEdGraphPin* BlendSpacePin = FindPin(GET_MEMBER_NAME_STRING_CHECKED(FAnimNode_RotationOffsetBlendSpace, BlendSpace));
 	if (BlendSpacePin != nullptr && BlendSpaceToCheck == nullptr)
 	{
@@ -81,7 +81,7 @@ void UAnimGraphNode_RotationOffsetBlendSpace::GetMenuActions(FBlueprintActionDat
 		static void SetNodeBlendSpace(UEdGraphNode* NewNode, bool /*bIsTemplateNode*/, TWeakObjectPtr<UBlendSpace> BlendSpace)
 		{
 			UAnimGraphNode_RotationOffsetBlendSpace* BlendSpaceNode = CastChecked<UAnimGraphNode_RotationOffsetBlendSpace>(NewNode);
-			BlendSpaceNode->Node.SetBlendSpace(BlendSpace.Get());
+			BlendSpaceNode->Node.BlendSpace = BlendSpace.Get();
 		}
 
 		static UBlueprintNodeSpawner* MakeBlendSpaceAction(TSubclassOf<UEdGraphNode> const NodeClass, const UBlendSpace* BlendSpace)
@@ -136,7 +136,7 @@ void UAnimGraphNode_RotationOffsetBlendSpace::GetMenuActions(FBlueprintActionDat
 FBlueprintNodeSignature UAnimGraphNode_RotationOffsetBlendSpace::GetSignature() const
 {
 	FBlueprintNodeSignature NodeSignature = Super::GetSignature();
-	NodeSignature.AddSubObject(Node.GetBlendSpace());
+	NodeSignature.AddSubObject(Node.BlendSpace);
 
 	return NodeSignature;
 }
@@ -145,7 +145,7 @@ void UAnimGraphNode_RotationOffsetBlendSpace::SetAnimationAsset(UAnimationAsset*
 {
 	if (UBlendSpace* BlendSpace = Cast<UBlendSpace>(Asset))
 	{
-		Node.SetBlendSpace(BlendSpace);
+		Node.BlendSpace = BlendSpace;
 	}
 }
 
@@ -153,7 +153,7 @@ void UAnimGraphNode_RotationOffsetBlendSpace::ValidateAnimNodeDuringCompilation(
 {
 	Super::ValidateAnimNodeDuringCompilation(ForSkeleton, MessageLog);
 
-	UBlendSpace* BlendSpaceToCheck = Node.GetBlendSpace();
+	UBlendSpace* BlendSpaceToCheck = Node.BlendSpace;
 	UEdGraphPin* BlendSpacePin = FindPin(GET_MEMBER_NAME_STRING_CHECKED(FAnimNode_RotationOffsetBlendSpace, BlendSpace));
 	if (BlendSpacePin != nullptr && BlendSpaceToCheck == nullptr)
 	{
@@ -208,7 +208,7 @@ void UAnimGraphNode_RotationOffsetBlendSpace::GetNodeContextMenuActions(UToolMen
 
 void UAnimGraphNode_RotationOffsetBlendSpace::GetAllAnimationSequencesReferred(TArray<UAnimationAsset*>& AnimationAssets) const
 {
-	if(Node.GetBlendSpace())
+	if(Node.BlendSpace)
 	{
 		HandleAnimReferenceCollection(Node.BlendSpace, AnimationAssets);
 	}
diff --git a/Engine/Source/Editor/AnimGraph/Private/AnimGraphNode_SaveCachedPose.cpp b/Engine/Source/Editor/AnimGraph/Private/AnimGraphNode_SaveCachedPose.cpp
--- a/Engine/Source/Editor/AnimGraph/Private/AnimGraphNode_SaveCachedPose.cpp
+++ b/Engine/Source/Editor/AnimGraph/Private/AnimGraphNode_SaveCachedPose.cpp
@@ -5,7 +5,6 @@
 
 #include "BlueprintNodeSpawner.h"
 #include "BlueprintActionDatabaseRegistrar.h"
-#include "AnimBlueprintExtension_CachedPose.h"
 
 /////////////////////////////////////////////////////
 // FCachedPoseNameValidator
@@ -117,9 +116,4 @@ TSharedPtr<class INameValidatorInterface> UAnimGraphNode_SaveCachedPose::MakeNam
 	return MakeShareable(new FCachedPoseNameValidator(GetBlueprint(), CacheName));
 }
 
-void UAnimGraphNode_SaveCachedPose::GetRequiredExtensions(TArray<TSubclassOf<UAnimBlueprintExtension>>& OutExtensions) const
-{
-	OutExtensions.Add(UAnimBlueprintExtension_CachedPose::StaticClass());
-}
-
 #undef LOCTEXT_NAMESPACE
diff --git a/Engine/Source/Editor/AnimGraph/Private/AnimGraphNode_SequenceEvaluator.cpp b/Engine/Source/Editor/AnimGraph/Private/AnimGraphNode_SequenceEvaluator.cpp
--- a/Engine/Source/Editor/AnimGraph/Private/AnimGraphNode_SequenceEvaluator.cpp
+++ b/Engine/Source/Editor/AnimGraph/Private/AnimGraphNode_SequenceEvaluator.cpp
@@ -8,7 +8,6 @@
 #include "Animation/AnimComposite.h"
 #include "Animation/AnimSequence.h"
 #include "AnimGraphNode_SequenceEvaluator.h"
-#include "IAnimBlueprintNodeOverrideAssetsContext.h"
 
 /////////////////////////////////////////////////////
 // UAnimGraphNode_SequenceEvaluator
@@ -22,7 +21,7 @@ UAnimGraphNode_SequenceEvaluator::UAnimGraphNode_SequenceEvaluator(const FObject
 
 void UAnimGraphNode_SequenceEvaluator::PreloadRequiredAssets()
 {
-	PreloadObject(Node.GetSequence());
+	PreloadObject(Node.Sequence);
 
 	Super::PreloadRequiredAssets();
 }
@@ -30,12 +29,15 @@ void UAnimGraphNode_SequenceEvaluator::PreloadRequiredAssets()
 void UAnimGraphNode_SequenceEvaluator::BakeDataDuringCompilation(class FCompilerResultsLog& MessageLog)
 {
 	UAnimBlueprint* AnimBlueprint = GetAnimBlueprint();
-	AnimBlueprint->FindOrAddGroup(Node.GetGroupName());
+	AnimBlueprint->FindOrAddGroup(SyncGroup.GroupName);
+	Node.GroupName = SyncGroup.GroupName;
+	Node.GroupRole = SyncGroup.GroupRole;
+	Node.Method = SyncGroup.Method;
 }
 
 void UAnimGraphNode_SequenceEvaluator::GetAllAnimationSequencesReferred(TArray<UAnimationAsset*>& AnimationAssets) const
 {
-	if(Node.GetSequence())
+	if(Node.Sequence)
 	{
 		HandleAnimReferenceCollection(Node.Sequence, AnimationAssets);
 	}
@@ -74,8 +76,7 @@ FText UAnimGraphNode_SequenceEvaluator::GetNodeTitleForSequence(ENodeTitleType::
 
 FText UAnimGraphNode_SequenceEvaluator::GetNodeTitle(ENodeTitleType::Type TitleType) const
 {
-	UAnimSequenceBase* Sequence = Node.GetSequence();
-	if (Sequence == nullptr)
+	if (Node.Sequence == nullptr)
 	{
 		// we may have a valid variable connected or default pin value
 		UEdGraphPin* SequencePin = FindPin(GET_MEMBER_NAME_STRING_CHECKED(FAnimNode_SequenceEvaluator, Sequence));
@@ -97,7 +98,7 @@ FText UAnimGraphNode_SequenceEvaluator::GetNodeTitle(ENodeTitleType::Type TitleT
 	//        then, we'll leave this optimization off
 	else //if (CachedNodeTitle.IsOutOfDate(this))
 	{
-		GetNodeTitleForSequence(TitleType, Sequence);
+		GetNodeTitleForSequence(TitleType, Node.Sequence);
 	}
 
 	return CachedNodeTitle;
@@ -112,19 +113,7 @@ void UAnimGraphNode_SequenceEvaluator::SetAnimationAsset(UAnimationAsset* Asset)
 {
 	if (UAnimSequenceBase* Seq =  Cast<UAnimSequence>(Asset))
 	{
-		Node.SetSequence(Seq);
-	}
-}
-
-void UAnimGraphNode_SequenceEvaluator::OnOverrideAssets(IAnimBlueprintNodeOverrideAssetsContext& InContext) const
-{
-	if(InContext.GetAssets().Num() > 0)
-	{
-		if (UAnimSequenceBase* Sequence = Cast<UAnimSequenceBase>(InContext.GetAssets()[0]))
-		{
-			FAnimNode_SequenceEvaluator& AnimNode = InContext.GetAnimNode<FAnimNode_SequenceEvaluator>();
-			AnimNode.SetSequence(Sequence);
-		}
+		Node.Sequence = Seq;
 	}
 }
 
@@ -132,7 +121,7 @@ void UAnimGraphNode_SequenceEvaluator::ValidateAnimNodeDuringCompilation(class U
 {
 	Super::ValidateAnimNodeDuringCompilation(ForSkeleton, MessageLog);
 
-	UAnimSequenceBase* SequenceToCheck = Node.GetSequence();
+	UAnimSequenceBase* SequenceToCheck = Node.Sequence;
 	UEdGraphPin* SequencePin = FindPin(GET_MEMBER_NAME_STRING_CHECKED(FAnimNode_SequenceEvaluator, Sequence));
 	if (SequencePin != nullptr && SequenceToCheck == nullptr)
 	{
@@ -188,14 +177,14 @@ bool UAnimGraphNode_SequenceEvaluator::DoesSupportTimeForTransitionGetter() cons
 
 UAnimationAsset* UAnimGraphNode_SequenceEvaluator::GetAnimationAsset() const 
 {
-	UAnimSequenceBase* Sequence = Node.GetSequence();
+	UAnimSequenceBase* Sequence = Node.Sequence;
 	UEdGraphPin* SequencePin = FindPin(GET_MEMBER_NAME_STRING_CHECKED(FAnimNode_SequenceEvaluator, Sequence));
 	if (SequencePin != nullptr && Sequence == nullptr)
 	{
 		Sequence = Cast<UAnimSequenceBase>(SequencePin->DefaultObject);
 	}
 
-	return Sequence;
+	return Node.Sequence;
 }
 
 const TCHAR* UAnimGraphNode_SequenceEvaluator::GetTimePropertyName() const 
diff --git a/Engine/Source/Editor/AnimGraph/Private/AnimGraphNode_SequencePlayer.cpp b/Engine/Source/Editor/AnimGraph/Private/AnimGraphNode_SequencePlayer.cpp
--- a/Engine/Source/Editor/AnimGraph/Private/AnimGraphNode_SequencePlayer.cpp
+++ b/Engine/Source/Editor/AnimGraph/Private/AnimGraphNode_SequencePlayer.cpp
@@ -17,8 +17,6 @@
 #include "Animation/AnimComposite.h"
 #include "Animation/AnimSequence.h"
 #include "Animation/AnimPoseSearchProvider.h"
-#include "UObject/UE5MainStreamObjectVersion.h"
-#include "IAnimBlueprintNodeOverrideAssetsContext.h"
 
 #define LOCTEXT_NAMESPACE "A3Nodes"
 
@@ -43,7 +41,7 @@ public:
 	virtual UEdGraphNode* PerformAction(class UEdGraph* ParentGraph, UEdGraphPin* FromPin, const FVector2D Location, bool bSelectNewNode = true) override
 	{
 		UAnimGraphNode_SequencePlayer* SpawnedNode = CastChecked<UAnimGraphNode_SequencePlayer>(FEdGraphSchemaAction_K2NewNode::PerformAction(ParentGraph, FromPin, Location, bSelectNewNode));
-		SpawnedNode->Node.SetSequence(Cast<UAnimSequence>(AssetInfo.GetAsset()));
+		SpawnedNode->Node.Sequence = Cast<UAnimSequence>(AssetInfo.GetAsset());
 
 		return SpawnedNode;
 	}
@@ -57,29 +55,16 @@ UAnimGraphNode_SequencePlayer::UAnimGraphNode_SequencePlayer(const FObjectInitia
 {
 }
 
-void UAnimGraphNode_SequencePlayer::Serialize(FArchive& Ar)
-{
-	Super::Serialize(Ar);
-
-	Ar.UsingCustomVersion(FUE5MainStreamObjectVersion::GUID);
-
-	if(Ar.IsLoading() && Ar.CustomVer(FUE5MainStreamObjectVersion::GUID) < FUE5MainStreamObjectVersion::AnimNodeConstantDataRefactorPhase0)
-	{
-		Node.PlayRateScaleBiasClampConstants.CopyFromLegacy(Node.PlayRateScaleBiasClamp_DEPRECATED);
-	}
-}
-
 void UAnimGraphNode_SequencePlayer::PreloadRequiredAssets()
 {
-	PreloadObject(Node.GetSequence());
+	PreloadObject(Node.Sequence);
 
 	Super::PreloadRequiredAssets();
 }
 
 FLinearColor UAnimGraphNode_SequencePlayer::GetNodeTitleColor() const
 {
-	UAnimSequenceBase* Sequence = Node.GetSequence();
-	if ((Sequence != NULL) && Sequence->IsValidAdditive())
+	if ((Node.Sequence != NULL) && Node.Sequence->IsValidAdditive())
 	{
 		return FLinearColor(0.10f, 0.60f, 0.12f);
 	}
@@ -91,14 +76,13 @@ FLinearColor UAnimGraphNode_SequencePlayer::GetNodeTitleColor() const
 
 FText UAnimGraphNode_SequencePlayer::GetTooltipText() const
 {
-	UAnimSequenceBase* Sequence = Node.GetSequence();
-	if (!Sequence)
+	if (!Node.Sequence)
 	{
 		return FText();
 	}
 
-	const bool bAdditive = Sequence->IsValidAdditive();
-	return GetTitleGivenAssetInfo(FText::FromString(Sequence->GetPathName()), bAdditive);
+	const bool bAdditive = Node.Sequence->IsValidAdditive();
+	return GetTitleGivenAssetInfo(FText::FromString(Node.Sequence->GetPathName()), bAdditive);
 }
 
 FText UAnimGraphNode_SequencePlayer::GetNodeTitleForSequence(ENodeTitleType::Type TitleType, UAnimSequenceBase* InSequence) const
@@ -106,9 +90,9 @@ FText UAnimGraphNode_SequencePlayer::GetNodeTitleForSequence(ENodeTitleType::Typ
 	const bool bAdditive = InSequence->IsValidAdditive();
 	const FText BasicTitle = GetTitleGivenAssetInfo(FText::FromName(InSequence->GetFName()), bAdditive);
 
-	if(Node.GetGroupMethod() == EAnimSyncMethod::SyncGroup)
+	if(SyncGroup.Method == EAnimSyncMethod::SyncGroup)
 	{
-		const FText SyncGroupName = FText::FromName(Node.GetGroupName());
+		const FText SyncGroupName = FText::FromName(SyncGroup.GroupName);
 
 		FFormatNamedArguments Args;
 		Args.Add(TEXT("Title"), BasicTitle);
@@ -123,7 +107,7 @@ FText UAnimGraphNode_SequencePlayer::GetNodeTitleForSequence(ENodeTitleType::Typ
 			return FText::Format(LOCTEXT("SequenceNodeGroupWithSubtitleList", "{Title} (Sync group {SyncGroup})"), Args);
 		}
 	}
-	else if(Node.GetGroupMethod() == EAnimSyncMethod::Graph)
+	else if(SyncGroup.Method == EAnimSyncMethod::Graph)
 	{
 		FFormatNamedArguments Args;
 		Args.Add(TEXT("Title"), BasicTitle);
@@ -169,8 +153,7 @@ FText UAnimGraphNode_SequencePlayer::GetNodeTitleForSequence(ENodeTitleType::Typ
 
 FText UAnimGraphNode_SequencePlayer::GetNodeTitle(ENodeTitleType::Type TitleType) const
 {
-	UAnimSequenceBase* Sequence = Node.GetSequence();
-	if (Sequence == nullptr)
+	if (Node.Sequence == nullptr)
 	{
 		// we may have a valid variable connected or default pin value
 		UEdGraphPin* SequencePin = FindPin(GET_MEMBER_NAME_STRING_CHECKED(FAnimNode_SequencePlayer, Sequence));
@@ -189,7 +172,7 @@ FText UAnimGraphNode_SequencePlayer::GetNodeTitle(ENodeTitleType::Type TitleType
 	}
 	else
 	{
-		return GetNodeTitleForSequence(TitleType, Sequence);
+		return GetNodeTitleForSequence(TitleType, Node.Sequence);
 	}
 }
 
@@ -218,7 +201,7 @@ void UAnimGraphNode_SequencePlayer::GetMenuActions(FBlueprintActionDatabaseRegis
 	auto LoadedAssetSetup = [](UEdGraphNode* NewNode, bool /*bIsTemplateNode*/, TWeakObjectPtr<UAnimSequence> SequencePtr)
 	{
 		UAnimGraphNode_SequencePlayer* SequencePlayerNode = CastChecked<UAnimGraphNode_SequencePlayer>(NewNode);
-		SequencePlayerNode->Node.SetSequence(SequencePtr.Get());
+		SequencePlayerNode->Node.Sequence = SequencePtr.Get();
 	};
 
 	auto UnloadedAssetSetup = [](UEdGraphNode* NewNode, bool bIsTemplateNode, const FAssetData AssetData)
@@ -232,7 +215,7 @@ void UAnimGraphNode_SequencePlayer::GetMenuActions(FBlueprintActionDatabaseRegis
 		{
 			UAnimSequence* Sequence = Cast<UAnimSequence>(AssetData.GetAsset());
 			check(Sequence != nullptr);
-			SequencePlayerNode->Node.SetSequence(Sequence);
+			SequencePlayerNode->Node.Sequence = Sequence;
 		}
 	};	
 
@@ -328,10 +311,9 @@ bool UAnimGraphNode_SequencePlayer::IsActionFilteredOut(class FBlueprintActionFi
 		UAnimBlueprint* AnimBlueprint = Cast<UAnimBlueprint>(Blueprint);
 		if (AnimBlueprint && AnimBlueprint->TargetSkeleton)
 		{
-			UAnimSequenceBase* Sequence = Node.GetSequence();
-			if(Sequence)
+			if(Node.Sequence)
 			{
-				if (!AnimBlueprint->TargetSkeleton->IsCompatible(Sequence->GetSkeleton()))
+				if (!AnimBlueprint->TargetSkeleton->IsCompatible(Node.Sequence->GetSkeleton()))
 				{
 					// Asset does not use a compatible skeleton with the Blueprint, cannot use
 					bIsFilteredOut = true;
@@ -373,7 +355,7 @@ void UAnimGraphNode_SequencePlayer::ValidateAnimNodeDuringCompilation(class USke
 {
 	Super::ValidateAnimNodeDuringCompilation(ForSkeleton, MessageLog);
 
-	UAnimSequenceBase* SequenceToCheck = Node.GetSequence();
+	UAnimSequenceBase* SequenceToCheck = Node.Sequence;
 	UEdGraphPin* SequencePin = FindPin(GET_MEMBER_NAME_STRING_CHECKED(FAnimNode_SequencePlayer, Sequence));
 	if (SequencePin != nullptr && SequenceToCheck == nullptr)
 	{
@@ -430,31 +412,22 @@ void UAnimGraphNode_SequencePlayer::SetAnimationAsset(UAnimationAsset* Asset)
 {
 	if (UAnimSequenceBase* Seq = Cast<UAnimSequenceBase>(Asset))
 	{
-		Node.SetSequence(Seq);
-	}
-}
-
-void UAnimGraphNode_SequencePlayer::OnOverrideAssets(IAnimBlueprintNodeOverrideAssetsContext& InContext) const
-{
-	if(InContext.GetAssets().Num() > 0)
-	{
-		if (UAnimSequenceBase* Sequence = Cast<UAnimSequenceBase>(InContext.GetAssets()[0]))
-		{
-			FAnimNode_SequencePlayer& AnimNode = InContext.GetAnimNode<FAnimNode_SequencePlayer>();
-			AnimNode.SetSequence(Sequence);
-		}
+		Node.Sequence = Seq;
 	}
 }
 
 void UAnimGraphNode_SequencePlayer::BakeDataDuringCompilation(class FCompilerResultsLog& MessageLog)
 {
 	UAnimBlueprint* AnimBlueprint = GetAnimBlueprint();
-	AnimBlueprint->FindOrAddGroup(Node.GetGroupName());
+	AnimBlueprint->FindOrAddGroup(SyncGroup.GroupName);
+	Node.GroupName = SyncGroup.GroupName;
+	Node.GroupRole = SyncGroup.GroupRole;
+	Node.Method = SyncGroup.Method;
 }
 
 void UAnimGraphNode_SequencePlayer::GetAllAnimationSequencesReferred(TArray<UAnimationAsset*>& AnimationAssets) const
 {
-	if(Node.GetSequence())
+	if(Node.Sequence)
 	{
 		HandleAnimReferenceCollection(Node.Sequence, AnimationAssets);
 	}
@@ -473,7 +446,7 @@ bool UAnimGraphNode_SequencePlayer::DoesSupportTimeForTransitionGetter() const
 
 UAnimationAsset* UAnimGraphNode_SequencePlayer::GetAnimationAsset() const 
 {
-	UAnimSequenceBase* Sequence = Node.GetSequence();
+	UAnimSequenceBase* Sequence = Node.Sequence;
 	UEdGraphPin* SequencePin = FindPin(GET_MEMBER_NAME_STRING_CHECKED(FAnimNode_SequencePlayer, Sequence));
 	if (SequencePin != nullptr && Sequence == nullptr)
 	{
@@ -515,11 +488,11 @@ void UAnimGraphNode_SequencePlayer::CustomizePinData(UEdGraphPin* Pin, FName Sou
 			UEdGraphPin* PlayRateBasisPin = FindPin(GET_MEMBER_NAME_STRING_CHECKED(FAnimNode_SequencePlayer, PlayRateBasis));
 			if (!PlayRateBasisPin || PlayRateBasisPin->bHidden)
 			{
-				if (Node.GetPlayRateBasis() != 1.f)
+				if (Node.PlayRateBasis != 1.f)
 				{
 					FFormatNamedArguments Args;
 					Args.Add(TEXT("PinFriendlyName"), Pin->PinFriendlyName);
-					Args.Add(TEXT("PlayRateBasis"), FText::AsNumber(Node.GetPlayRateBasis()));
+					Args.Add(TEXT("PlayRateBasis"), FText::AsNumber(Node.PlayRateBasis));
 					Pin->PinFriendlyName = FText::Format(LOCTEXT("FAnimNode_SequencePlayer_PlayRateBasis_Value", "({PinFriendlyName} / {PlayRateBasis})"), Args);
 				}
 			}
@@ -530,7 +503,7 @@ void UAnimGraphNode_SequencePlayer::CustomizePinData(UEdGraphPin* Pin, FName Sou
 				Pin->PinFriendlyName = FText::Format(LOCTEXT("FAnimNode_SequencePlayer_PlayRateBasis_Name", "({PinFriendlyName} / PlayRateBasis)"), Args);
 			}
 
-			Pin->PinFriendlyName = Node.GetPlayRateScaleBiasClampConstants().GetFriendlyName(Pin->PinFriendlyName);
+			Pin->PinFriendlyName = Node.PlayRateScaleBiasClamp.GetFriendlyName(Pin->PinFriendlyName);
 		}
 	}
 }
@@ -541,17 +514,17 @@ void UAnimGraphNode_SequencePlayer::PostEditChangeProperty(struct FPropertyChang
 
 	// Reconstruct node to show updates to PinFriendlyNames.
 	if ((PropertyName == GET_MEMBER_NAME_STRING_CHECKED(FAnimNode_SequencePlayer, PlayRateBasis))
-		|| (PropertyName == GET_MEMBER_NAME_STRING_CHECKED(FInputScaleBiasClampConstants, bMapRange))
+		|| (PropertyName == GET_MEMBER_NAME_STRING_CHECKED(FInputScaleBiasClamp, bMapRange))
 		|| (PropertyName == GET_MEMBER_NAME_STRING_CHECKED(FInputRange, Min))
 		|| (PropertyName == GET_MEMBER_NAME_STRING_CHECKED(FInputRange, Max))
-		|| (PropertyName == GET_MEMBER_NAME_STRING_CHECKED(FInputScaleBiasClampConstants, Scale))
-		|| (PropertyName == GET_MEMBER_NAME_STRING_CHECKED(FInputScaleBiasClampConstants, Bias))
-		|| (PropertyName == GET_MEMBER_NAME_STRING_CHECKED(FInputScaleBiasClampConstants, bClampResult))
-		|| (PropertyName == GET_MEMBER_NAME_STRING_CHECKED(FInputScaleBiasClampConstants, ClampMin))
-		|| (PropertyName == GET_MEMBER_NAME_STRING_CHECKED(FInputScaleBiasClampConstants, ClampMax))
-		|| (PropertyName == GET_MEMBER_NAME_STRING_CHECKED(FInputScaleBiasClampConstants, bInterpResult))
-		|| (PropertyName == GET_MEMBER_NAME_STRING_CHECKED(FInputScaleBiasClampConstants, InterpSpeedIncreasing))
-		|| (PropertyName == GET_MEMBER_NAME_STRING_CHECKED(FInputScaleBiasClampConstants, InterpSpeedDecreasing)))
+		|| (PropertyName == GET_MEMBER_NAME_STRING_CHECKED(FInputScaleBiasClamp, Scale))
+		|| (PropertyName == GET_MEMBER_NAME_STRING_CHECKED(FInputScaleBiasClamp, Bias))
+		|| (PropertyName == GET_MEMBER_NAME_STRING_CHECKED(FInputScaleBiasClamp, bClampResult))
+		|| (PropertyName == GET_MEMBER_NAME_STRING_CHECKED(FInputScaleBiasClamp, ClampMin))
+		|| (PropertyName == GET_MEMBER_NAME_STRING_CHECKED(FInputScaleBiasClamp, ClampMax))
+		|| (PropertyName == GET_MEMBER_NAME_STRING_CHECKED(FInputScaleBiasClamp, bInterpResult))
+		|| (PropertyName == GET_MEMBER_NAME_STRING_CHECKED(FInputScaleBiasClamp, InterpSpeedIncreasing))
+		|| (PropertyName == GET_MEMBER_NAME_STRING_CHECKED(FInputScaleBiasClamp, InterpSpeedDecreasing)))
 	{
 		ReconstructNode();
 	}
diff --git a/Engine/Source/Editor/AnimGraph/Private/AnimGraphNode_StateMachineBase.cpp b/Engine/Source/Editor/AnimGraph/Private/AnimGraphNode_StateMachineBase.cpp
--- a/Engine/Source/Editor/AnimGraph/Private/AnimGraphNode_StateMachineBase.cpp
+++ b/Engine/Source/Editor/AnimGraph/Private/AnimGraphNode_StateMachineBase.cpp
@@ -25,7 +25,7 @@
 #include "AnimGraphNode_LinkedAnimLayer.h"
 #include "AnimGraphNode_TransitionPoseEvaluator.h"
 #include "AnimGraphNode_CustomTransitionResult.h"
-#include "AnimBlueprintExtension_StateMachine.h"
+#include "AnimBlueprintCompilerHandler_StateMachine.h"
 #include "IAnimBlueprintGeneratedClassCompiledData.h"
 #include "IAnimBlueprintCompilationContext.h"
 #include "Animation/AnimNode_Inertialization.h"
@@ -114,8 +114,6 @@ FString UAnimGraphNode_StateMachineBase::GetNodeCategory() const
 
 void UAnimGraphNode_StateMachineBase::PostPlacedNewNode()
 {
-	Super::PostPlacedNewNode();
-
 	// Create a new animation graph
 	check(EditorStateMachineGraph == NULL);
 	EditorStateMachineGraph = CastChecked<UAnimationStateMachineGraph>(FBlueprintEditorUtils::CreateNewGraph(this, NAME_None, UAnimationStateMachineGraph::StaticClass(), UAnimationStateMachineSchema::StaticClass()));
@@ -365,8 +363,8 @@ void UAnimGraphNode_StateMachineBase::OnProcessDuringCompilation(IAnimBlueprintC
 		}
 	};
 	
-	UAnimBlueprintExtension_StateMachine* Extension = UAnimBlueprintExtension::GetExtension<UAnimBlueprintExtension_StateMachine>(GetAnimBlueprint());
-	check(Extension);
+	FAnimBlueprintCompilerHandler_StateMachine* CompilerHandler = InCompilationContext.GetHandler<FAnimBlueprintCompilerHandler_StateMachine>("AnimBlueprintCompilerHandler_StateMachine");
+	check(CompilerHandler);
 
 	if (EditorStateMachineGraph == NULL)
 	{
@@ -433,7 +431,7 @@ void UAnimGraphNode_StateMachineBase::OnProcessDuringCompilation(IAnimBlueprintC
 				{
 					InCompilationContext.ValidateGraphIsWellFormed(StateNode->BoundGraph);
 
-					BakedState.StateRootNodeIndex = Extension->ExpandGraphAndProcessNodes(StateNode->BoundGraph, AnimGraphResultNode, InCompilationContext, OutCompiledData);
+					BakedState.StateRootNodeIndex = CompilerHandler->ExpandGraphAndProcessNodes(StateNode->BoundGraph, AnimGraphResultNode, InCompilationContext, OutCompiledData);
 
 					// See if the state consists of a single sequence player node, and remember the index if so
 					for (UEdGraphPin* TestPin : AnimGraphResultNode->Pins)
@@ -475,7 +473,7 @@ void UAnimGraphNode_StateMachineBase::OnProcessDuringCompilation(IAnimBlueprintC
 			{
 				if (UAnimGraphNode_TransitionResult* EntryRuleResultNode = CastChecked<UAnimationTransitionGraph>(ConduitNode->BoundGraph)->GetResultNode())
 				{
-					BakedState.EntryRuleNodeIndex = Extension->ExpandGraphAndProcessNodes(ConduitNode->BoundGraph, EntryRuleResultNode, InCompilationContext, OutCompiledData);
+					BakedState.EntryRuleNodeIndex = CompilerHandler->ExpandGraphAndProcessNodes(ConduitNode->BoundGraph, EntryRuleResultNode, InCompilationContext, OutCompiledData);
 				}
 			}
 
@@ -652,7 +650,7 @@ void UAnimGraphNode_StateMachineBase::OnProcessDuringCompilation(IAnimBlueprintC
 				}
 				else
 				{
-					Rule.CanTakeDelegateIndex = Extension->ExpandGraphAndProcessNodes(TransitionNode->BoundGraph, TransitionResultNode, InCompilationContext, OutCompiledData, TransitionNode);
+					Rule.CanTakeDelegateIndex = CompilerHandler->ExpandGraphAndProcessNodes(TransitionNode->BoundGraph, TransitionResultNode, InCompilationContext, OutCompiledData, TransitionNode);
 					AlreadyMergedTransitionList.Add(TransitionResultNode, Rule.CanTakeDelegateIndex);
 				}
 			}
@@ -673,7 +671,7 @@ void UAnimGraphNode_StateMachineBase::OnProcessDuringCompilation(IAnimBlueprintC
 				TArray<UEdGraphNode*> ClonedNodes;
 				if (CustomTransitionGraph->GetResultNode())
 				{
-					Rule.CustomResultNodeIndex = Extension->ExpandGraphAndProcessNodes(TransitionNode->CustomTransitionGraph, CustomTransitionGraph->GetResultNode(), InCompilationContext, OutCompiledData, nullptr, &ClonedNodes);
+					Rule.CustomResultNodeIndex = CompilerHandler->ExpandGraphAndProcessNodes(TransitionNode->CustomTransitionGraph, CustomTransitionGraph->GetResultNode(), InCompilationContext, OutCompiledData, nullptr, &ClonedNodes);
 				}
 
 				// Find all the pose evaluators used in this transition, save handles to them because we need to populate some pose data before executing
@@ -716,9 +714,4 @@ void UAnimGraphNode_StateMachineBase::GetOutputLinkAttributes(FNodeAttributeArra
 	}
 }
 
-void UAnimGraphNode_StateMachineBase::GetRequiredExtensions(TArray<TSubclassOf<UAnimBlueprintExtension>>& OutExtensions) const
-{
-	OutExtensions.Add(UAnimBlueprintExtension_StateMachine::StaticClass());
-}
-
 #undef LOCTEXT_NAMESPACE
diff --git a/Engine/Source/Editor/AnimGraph/Private/AnimGraphNode_StateResult.cpp b/Engine/Source/Editor/AnimGraph/Private/AnimGraphNode_StateResult.cpp
--- a/Engine/Source/Editor/AnimGraph/Private/AnimGraphNode_StateResult.cpp
+++ b/Engine/Source/Editor/AnimGraph/Private/AnimGraphNode_StateResult.cpp
@@ -2,6 +2,8 @@
 
 #include "AnimGraphNode_StateResult.h"
 #include "GraphEditorSettings.h"
+#include "AnimBlueprintCompiler.h"
+#include "AnimBlueprintCompilerHandler_Base.h"
 #include "IAnimBlueprintCompilationContext.h"
 
 #define LOCTEXT_NAMESPACE "A3Nodes"
@@ -48,7 +50,7 @@ void UAnimGraphNode_StateResult::OnProcessDuringCompilation(IAnimBlueprintCompil
 {
 	UAnimGraphNode_StateResult* TrueNode = InCompilationContext.GetMessageLog().FindSourceObjectTypeChecked<UAnimGraphNode_StateResult>(this);
 
-	Node.SetName(TrueNode->GetGraph()->GetFName());
+	Node.Name = TrueNode->GetGraph()->GetFName();
 }
 
 #undef LOCTEXT_NAMESPACE
diff --git a/Engine/Source/Editor/AnimGraph/Private/AnimGraphNode_UseCachedPose.cpp b/Engine/Source/Editor/AnimGraph/Private/AnimGraphNode_UseCachedPose.cpp
--- a/Engine/Source/Editor/AnimGraph/Private/AnimGraphNode_UseCachedPose.cpp
+++ b/Engine/Source/Editor/AnimGraph/Private/AnimGraphNode_UseCachedPose.cpp
@@ -8,7 +8,7 @@
 #include "BlueprintActionFilter.h"
 #include "BlueprintNodeSpawner.h"
 #include "BlueprintActionDatabaseRegistrar.h"
-#include "AnimBlueprintExtension_CachedPose.h"
+#include "AnimBlueprintCompilerHandler_CachedPose.h"
 #include "IAnimBlueprintCompilationContext.h"
 #include "Kismet2/KismetEditorUtilities.h"
 #include "BlueprintEditorModule.h"
@@ -200,15 +200,15 @@ bool UAnimGraphNode_UseCachedPose::IsActionFilteredOut(class FBlueprintActionFil
 
 void UAnimGraphNode_UseCachedPose::OnProcessDuringCompilation(IAnimBlueprintCompilationContext& InCompilationContext, IAnimBlueprintGeneratedClassCompiledData& OutCompiledData)
 {	
-	UAnimBlueprintExtension_CachedPose* CompilerExtension = UAnimBlueprintExtension::GetExtension<UAnimBlueprintExtension_CachedPose>(GetAnimBlueprint());
-	check(CompilerExtension);
+	FAnimBlueprintCompilerHandler_CachedPose* CompilerHandler = InCompilationContext.GetHandler<FAnimBlueprintCompilerHandler_CachedPose>("AnimBlueprintCompilerHandler_CachedPose");
+	check(CompilerHandler);
 
 	bool bSuccessful = false;
 
 	// Link to the saved cached pose
 	if(SaveCachedPoseNode.IsValid())
 	{
-		if (UAnimGraphNode_SaveCachedPose* AssociatedSaveNode = CompilerExtension->GetSaveCachedPoseNodes().FindRef(SaveCachedPoseNode->CacheName))
+		if (UAnimGraphNode_SaveCachedPose* AssociatedSaveNode = CompilerHandler->GetSaveCachedPoseNodes().FindRef(SaveCachedPoseNode->CacheName))
 		{
 			FStructProperty* LinkProperty = FindFProperty<FStructProperty>(FAnimNode_UseCachedPose::StaticStruct(), TEXT("LinkToCachingNode"));
 			check(LinkProperty);
@@ -233,9 +233,4 @@ void UAnimGraphNode_UseCachedPose::OnProcessDuringCompilation(IAnimBlueprintComp
 	}
 }
 
-void UAnimGraphNode_UseCachedPose::GetRequiredExtensions(TArray<TSubclassOf<UAnimBlueprintExtension>>& OutExtensions) const
-{
-	OutExtensions.Add(UAnimBlueprintExtension_CachedPose::StaticClass());
-}
-
 #undef LOCTEXT_NAMESPACE
diff --git a/Engine/Source/Editor/AnimGraph/Private/AnimationGraphSchema.cpp b/Engine/Source/Editor/AnimGraph/Private/AnimationGraphSchema.cpp
--- a/Engine/Source/Editor/AnimGraph/Private/AnimationGraphSchema.cpp
+++ b/Engine/Source/Editor/AnimGraph/Private/AnimationGraphSchema.cpp
@@ -293,7 +293,7 @@ void UAnimationGraphSchema::CreateFunctionGraphTerminators(UEdGraph& Graph, UCla
 			Graph.GetNodesOfClass<UAnimGraphNode_Root>(RootNodes);
 
 			check(RootNodes.Num() == 1);
-			RootNodes[0]->Node.SetGroup(*FObjectEditorUtils::GetCategoryText(InterfaceToImplement).ToString());
+			RootNodes[0]->Node.Group = *FObjectEditorUtils::GetCategoryText(InterfaceToImplement).ToString();
 
 			int32 CurrentPoseIndex = 0;
 			for (TFieldIterator<FProperty> PropIt(InterfaceToImplement); PropIt && (PropIt->PropertyFlags & CPF_Parm); ++PropIt)
@@ -730,7 +730,7 @@ void UAnimationGraphSchema::ConformAnimGraphToInterface(UBlueprint* InBlueprint,
 		InGraph.GetNodesOfClass<UAnimGraphNode_Root>(RootNodes);
 
 		check(RootNodes.Num() == 1);
-		RootNodes[0]->Node.SetGroup(*FObjectEditorUtils::GetCategoryText(InFunction).ToString());
+		RootNodes[0]->Node.Group = *FObjectEditorUtils::GetCategoryText(InFunction).ToString();
 
 		TArray<UAnimGraphNode_LinkedInputPose*> LinkedInputPoseNodes;
 		InGraph.GetNodesOfClass<UAnimGraphNode_LinkedInputPose>(LinkedInputPoseNodes);
diff --git a/Engine/Source/Editor/AnimGraph/Public/AnimGraphNode_AimOffsetLookAt.h b/Engine/Source/Editor/AnimGraph/Public/AnimGraphNode_AimOffsetLookAt.h
--- a/Engine/Source/Editor/AnimGraph/Public/AnimGraphNode_AimOffsetLookAt.h
+++ b/Engine/Source/Editor/AnimGraph/Public/AnimGraphNode_AimOffsetLookAt.h
@@ -29,7 +29,7 @@ class UAnimGraphNode_AimOffsetLookAt : public UAnimGraphNode_BlendSpaceBase
 
 	// UAnimGraphNode_Base interface
 	virtual void ValidateAnimNodeDuringCompilation(class USkeleton* ForSkeleton, class FCompilerResultsLog& MessageLog) override;
-	virtual UAnimationAsset* GetAnimationAsset() const override { return Node.GetBlendSpace(); }
+	virtual UAnimationAsset* GetAnimationAsset() const override { return Node.BlendSpace; }
 	virtual void GetAllAnimationSequencesReferred(TArray<UAnimationAsset*>& AnimationAssets) const override;
 	virtual void ReplaceReferredAnimations(const TMap<UAnimationAsset*, UAnimationAsset*>& AnimAssetReplacementMap) override;
 	ANIMGRAPH_API virtual void CustomizePinData(UEdGraphPin* Pin, FName SourcePropertyName, int32 ArrayIndex) const override;
diff --git a/Engine/Source/Editor/AnimGraph/Public/AnimGraphNode_AssetPlayerBase.h b/Engine/Source/Editor/AnimGraph/Public/AnimGraphNode_AssetPlayerBase.h
--- a/Engine/Source/Editor/AnimGraph/Public/AnimGraphNode_AssetPlayerBase.h
+++ b/Engine/Source/Editor/AnimGraph/Public/AnimGraphNode_AssetPlayerBase.h
@@ -20,9 +20,9 @@ class ANIMGRAPH_API UAnimGraphNode_AssetPlayerBase : public UAnimGraphNode_Base
 {
 	GENERATED_BODY()
 public:
-	// Deprecated - sync group data is held on the contained FAnimNode_Base
-	UPROPERTY()
-	FAnimationGroupReference SyncGroup_DEPRECATED;
+	// Sync group settings for this player.  Sync groups keep related animations with different lengths synchronized.
+	UPROPERTY(EditAnywhere, Category = Settings)
+	FAnimationGroupReference SyncGroup;
 
 	/** UObject interface */
 	void Serialize(FArchive& Ar) override;
diff --git a/Engine/Source/Editor/AnimGraph/Public/AnimGraphNode_Base.h b/Engine/Source/Editor/AnimGraph/Public/AnimGraphNode_Base.h
--- a/Engine/Source/Editor/AnimGraph/Public/AnimGraphNode_Base.h
+++ b/Engine/Source/Editor/AnimGraph/Public/AnimGraphNode_Base.h
@@ -27,8 +27,6 @@ class IAnimBlueprintGeneratedClassCompiledData;
 class IAnimBlueprintCompilationContext;
 class IAnimBlueprintCopyTermDefaultsContext;
 class IAnimBlueprintNodeCopyTermDefaultsContext;
-class IAnimBlueprintNodeOverrideAssetsContext;
-class UAnimBlueprintExtension;
 
 struct FPoseLinkMappingRecord
 {
@@ -207,8 +205,6 @@ class ANIMGRAPH_API UAnimGraphNode_Base : public UK2Node
 	virtual void AddSearchMetaDataInfo(TArray<struct FSearchTagDataPair>& OutTaggedMetaData) const override;
 	virtual void PinConnectionListChanged(UEdGraphPin* Pin) override;
 	virtual void AutowireNewNode(UEdGraphPin* FromPin) override;
-	virtual void PostPlacedNewNode() override;
-	virtual void DestroyNode() override;
 	// End of UEdGraphNode interface
 
 	// UK2Node interface
@@ -356,14 +352,6 @@ class ANIMGRAPH_API UAnimGraphNode_Base : public UK2Node
 	 */
 	bool IsPinExposedAndLinked(const FString& InPinName, const EEdGraphPinDirection Direction = EGPD_MAX) const;
 
-	/**
-	 * Helper function to check whether a pin is valid and bound via property access
-	 * @param	InPinName		The name of the pin @see UEdGraphNode::FindPin
-	 * @param	InPinDirection	The direction of the pin we are looking for. If this is EGPD_MAX, all directions are considered
-	 * @return true if the pin is present and bound
-	 */
-	bool IsPinExposedAndBound(const FString& InPinName, const EEdGraphPinDirection InDirection = EGPD_MAX) const;
-
 	// Event that is broadcast to inform observers that the node title has changed
 	// The default SAnimationGraphNode uses this to invalidate cached node title text
 	DECLARE_EVENT(UAnimGraphNode_Base, FOnNodeTitleChangedEvent);
@@ -388,8 +376,7 @@ class ANIMGRAPH_API UAnimGraphNode_Base : public UK2Node
 protected:
 	friend class FAnimBlueprintCompilerContext;
 	friend class FAnimGraphNodeDetails;
-	friend class UAnimBlueprintExtension;
-	friend class UAnimBlueprintExtension_Base;
+	friend class FAnimBlueprintCompilerHandler_Base;
 
 	// Gets the animation FNode type represented by this ed graph node
 	UScriptStruct* GetFNodeType() const;
@@ -397,9 +384,6 @@ protected:
 	// Gets the animation FNode property represented by this ed graph node
 	FStructProperty* GetFNodeProperty() const;
 
-	// Get the extension types that this node type holds on the anim blueprint. Some extension types are always requested by the system
-	virtual void GetRequiredExtensions(TArray<TSubclassOf<UAnimBlueprintExtension>>& OutExtensions) const {}
-	
 	// This will be called when a pose link is found, and can be called with PoseProperty being either of:
 	//  - an array property (ArrayIndex >= 0)
 	//  - a single pose property (ArrayIndex == INDEX_NONE)
@@ -424,12 +408,6 @@ protected:
 	// Copy this node's data during the last phase of compilation where term defaults are copied to the new CDO
 	virtual void OnCopyTermDefaultsToDefaultObject(IAnimBlueprintCopyTermDefaultsContext& InCompilationContext, IAnimBlueprintNodeCopyTermDefaultsContext& InPerNodeContext, IAnimBlueprintGeneratedClassCompiledData& OutCompiledData) {}
 
-	// Called to override the assets held on a runtime anim node. Implements per-node logic for child anim blueprints.
-	void OverrideAssets(IAnimBlueprintNodeOverrideAssetsContext& InContext) const;
-
-	// Override point for OverrideAssets
-	virtual void OnOverrideAssets(IAnimBlueprintNodeOverrideAssetsContext& InContext) const {}
-	
 	// Allocates or reallocates pins
 	void InternalPinCreation(TArray<UEdGraphPin*>* OldPins);
 
diff --git a/Engine/Source/Editor/AnimGraph/Public/AnimGraphNode_BlendSpaceEvaluator.h b/Engine/Source/Editor/AnimGraph/Public/AnimGraphNode_BlendSpaceEvaluator.h
--- a/Engine/Source/Editor/AnimGraph/Public/AnimGraphNode_BlendSpaceEvaluator.h
+++ b/Engine/Source/Editor/AnimGraph/Public/AnimGraphNode_BlendSpaceEvaluator.h
@@ -34,7 +34,6 @@ class UAnimGraphNode_BlendSpaceEvaluator : public UAnimGraphNode_BlendSpaceBase
 	virtual const TCHAR* GetTimePropertyName() const override;
 	virtual UScriptStruct* GetTimePropertyStruct() const override;
 	virtual EAnimAssetHandlerType SupportsAssetClass(const UClass* AssetClass) const override;
-	virtual void OnOverrideAssets(IAnimBlueprintNodeOverrideAssetsContext& InContext) const override;
 	// End of UAnimGraphNode_Base interface
 
 	// UK2Node interface
diff --git a/Engine/Source/Editor/AnimGraph/Public/AnimGraphNode_BlendSpaceGraphBase.h b/Engine/Source/Editor/AnimGraph/Public/AnimGraphNode_BlendSpaceGraphBase.h
--- a/Engine/Source/Editor/AnimGraph/Public/AnimGraphNode_BlendSpaceGraphBase.h
+++ b/Engine/Source/Editor/AnimGraph/Public/AnimGraphNode_BlendSpaceGraphBase.h
@@ -96,7 +96,6 @@ protected:
 	virtual void OnCopyTermDefaultsToDefaultObject(IAnimBlueprintCopyTermDefaultsContext& InCompilationContext, IAnimBlueprintNodeCopyTermDefaultsContext& InPerNodeContext, IAnimBlueprintGeneratedClassCompiledData& OutCompiledData) override;
 	virtual void CustomizeDetails(IDetailLayoutBuilder& InDetailBuilder) override;
 	virtual void GetInputLinkAttributes(FNodeAttributeArray& OutAttributes) const override;
-	virtual void GetRequiredExtensions(TArray<TSubclassOf<UAnimBlueprintExtension>>& OutExtensions) const override;
 
 	// UK2Node interface
 	virtual void PreloadRequiredAssets() override;
diff --git a/Engine/Source/Editor/AnimGraph/Public/AnimGraphNode_BlendSpacePlayer.h b/Engine/Source/Editor/AnimGraph/Public/AnimGraphNode_BlendSpacePlayer.h
--- a/Engine/Source/Editor/AnimGraph/Public/AnimGraphNode_BlendSpacePlayer.h
+++ b/Engine/Source/Editor/AnimGraph/Public/AnimGraphNode_BlendSpacePlayer.h
@@ -35,7 +35,6 @@ class UAnimGraphNode_BlendSpacePlayer : public UAnimGraphNode_BlendSpaceBase
 	virtual void GetAllAnimationSequencesReferred(TArray<UAnimationAsset*>& AnimationAssets) const override;
 	virtual void ReplaceReferredAnimations(const TMap<UAnimationAsset*, UAnimationAsset*>& AnimAssetReplacementMap) override;
 	virtual EAnimAssetHandlerType SupportsAssetClass(const UClass* AssetClass) const override;
-	virtual void OnOverrideAssets(IAnimBlueprintNodeOverrideAssetsContext& InContext) const override;
 	// End of UAnimGraphNode_Base interface
 
 	// UK2Node interface
diff --git a/Engine/Source/Editor/AnimGraph/Public/AnimGraphNode_CustomProperty.h b/Engine/Source/Editor/AnimGraph/Public/AnimGraphNode_CustomProperty.h
--- a/Engine/Source/Editor/AnimGraph/Public/AnimGraphNode_CustomProperty.h
+++ b/Engine/Source/Editor/AnimGraph/Public/AnimGraphNode_CustomProperty.h
@@ -14,6 +14,7 @@
 class FCompilerResultsLog;
 class IDetailLayoutBuilder;
 class IPropertyHandle;
+class FAnimBlueprintCompilerHandler;
 
 UCLASS(Abstract)
 class ANIMGRAPH_API UAnimGraphNode_CustomProperty : public UAnimGraphNode_Base, public IClassVariableCreator
diff --git a/Engine/Source/Editor/AnimGraph/Public/AnimGraphNode_LinkedAnimGraphBase.h b/Engine/Source/Editor/AnimGraph/Public/AnimGraphNode_LinkedAnimGraphBase.h
--- a/Engine/Source/Editor/AnimGraph/Public/AnimGraphNode_LinkedAnimGraphBase.h
+++ b/Engine/Source/Editor/AnimGraph/Public/AnimGraphNode_LinkedAnimGraphBase.h
@@ -37,7 +37,6 @@ public:
 	virtual void GetOutputLinkAttributes(FNodeAttributeArray& OutAttributes) const override;
 	virtual bool ShouldShowAttributesOnPins() const override { return false; }
 	virtual void OnCopyTermDefaultsToDefaultObject(IAnimBlueprintCopyTermDefaultsContext& InCompilationContext, IAnimBlueprintNodeCopyTermDefaultsContext& InPerNodeContext, IAnimBlueprintGeneratedClassCompiledData& OutCompiledData) override;
-	virtual void GetRequiredExtensions(TArray<TSubclassOf<UAnimBlueprintExtension>>& OutExtensions) const override;
 
 	// UAnimGraphNode_CustomProperty interface
 	virtual bool IsStructuralProperty(FProperty* InProperty) const override;
@@ -47,7 +46,7 @@ public:
 	virtual const FAnimNode_LinkedAnimGraph* GetLinkedAnimGraphNode() const PURE_VIRTUAL(UAnimGraphNode_LinkedAnimGraphBase::GetLinkedAnimGraphNode, return nullptr;);
 
 protected:
-	friend class UAnimBlueprintExtension_LinkedAnimGraph;
+	friend class FAnimBlueprintCompilerHandler_LinkedAnimGraph;
 
 	// Called pre-compilation to allocate pose links
 	void AllocatePoseLinks();
diff --git a/Engine/Source/Editor/AnimGraph/Public/AnimGraphNode_RotationOffsetBlendSpace.h b/Engine/Source/Editor/AnimGraph/Public/AnimGraphNode_RotationOffsetBlendSpace.h
--- a/Engine/Source/Editor/AnimGraph/Public/AnimGraphNode_RotationOffsetBlendSpace.h
+++ b/Engine/Source/Editor/AnimGraph/Public/AnimGraphNode_RotationOffsetBlendSpace.h
@@ -28,7 +28,7 @@ class UAnimGraphNode_RotationOffsetBlendSpace: public UAnimGraphNode_BlendSpaceB
 
 	// UAnimGraphNode_Base interface
 	virtual void ValidateAnimNodeDuringCompilation(class USkeleton* ForSkeleton, class FCompilerResultsLog& MessageLog) override;
-	virtual UAnimationAsset* GetAnimationAsset() const override { return Node.GetBlendSpace(); }
+	virtual UAnimationAsset* GetAnimationAsset() const override { return Node.BlendSpace; }
 	virtual void GetAllAnimationSequencesReferred(TArray<UAnimationAsset*>& AnimationAssets) const override;
 	virtual void ReplaceReferredAnimations(const TMap<UAnimationAsset*, UAnimationAsset*>& AnimAssetReplacementMap) override;
 	virtual EAnimAssetHandlerType SupportsAssetClass(const UClass* AssetClass) const override;
diff --git a/Engine/Source/Editor/AnimGraph/Public/AnimGraphNode_SaveCachedPose.h b/Engine/Source/Editor/AnimGraph/Public/AnimGraphNode_SaveCachedPose.h
--- a/Engine/Source/Editor/AnimGraph/Public/AnimGraphNode_SaveCachedPose.h
+++ b/Engine/Source/Editor/AnimGraph/Public/AnimGraphNode_SaveCachedPose.h
@@ -40,7 +40,6 @@ class UAnimGraphNode_SaveCachedPose : public UAnimGraphNode_Base
 	// UAnimGraphNode_Base interface
 	virtual FString GetNodeCategory() const override;
 	virtual bool IsSinkNode() const override { return true; }
-	virtual void GetRequiredExtensions(TArray<TSubclassOf<UAnimBlueprintExtension>>& OutExtensions) const override;
 	// End of UAnimGraphNode_Base interface
 
 private:
diff --git a/Engine/Source/Editor/AnimGraph/Public/AnimGraphNode_SequencePlayer.h b/Engine/Source/Editor/AnimGraph/Public/AnimGraphNode_SequencePlayer.h
--- a/Engine/Source/Editor/AnimGraph/Public/AnimGraphNode_SequencePlayer.h
+++ b/Engine/Source/Editor/AnimGraph/Public/AnimGraphNode_SequencePlayer.h
@@ -19,9 +19,6 @@ class UAnimGraphNode_SequencePlayer : public UAnimGraphNode_AssetPlayerBase
 	UPROPERTY(EditAnywhere, Category=Settings)
 	FAnimNode_SequencePlayer Node;
 
-	// UObject interface
-	void Serialize(FArchive& Ar) override;
-
 	// UEdGraphNode interface
 	virtual FLinearColor GetNodeTitleColor() const override;
 	virtual FText GetTooltipText() const override;
@@ -45,7 +42,6 @@ class UAnimGraphNode_SequencePlayer : public UAnimGraphNode_AssetPlayerBase
 	virtual void GetMenuActions(FBlueprintActionDatabaseRegistrar& ActionRegistrar) const override;
 	virtual bool IsActionFilteredOut(class FBlueprintActionFilter const& Filter) override;
 	virtual EAnimAssetHandlerType SupportsAssetClass(const UClass* AssetClass) const override;
-	virtual void OnOverrideAssets(IAnimBlueprintNodeOverrideAssetsContext& InContext) const override;
 	// End of UAnimGraphNode_Base interface
 
 	// UK2Node interface
diff --git a/Engine/Source/Editor/AnimGraph/Public/IAnimBlueprintCompilationBracketContext.h b/Engine/Source/Editor/AnimGraph/Public/IAnimBlueprintCompilationBracketContext.h
--- a/Engine/Source/Editor/AnimGraph/Public/IAnimBlueprintCompilationBracketContext.h
+++ b/Engine/Source/Editor/AnimGraph/Public/IAnimBlueprintCompilationBracketContext.h
@@ -5,6 +5,7 @@
 #include "CoreMinimal.h"
 
 class FCompilerResultsLog;
+class IAnimBlueprintCompilerHandler;
 
 // Interface passed to start/end compilation delegates
 class ANIMGRAPH_API IAnimBlueprintCompilationBracketContext
@@ -12,16 +13,21 @@ class ANIMGRAPH_API IAnimBlueprintCompilationBracketContext
 public:	
 	virtual ~IAnimBlueprintCompilationBracketContext() {}
 
+	// Get a handler of the specified type and name (i.e. via simple name-based RTTI)
+	// Handlers are registered via IAnimBlueprintCompilerHandlerCollection::RegisterHandler
+	template <typename THandlerClass>
+	THandlerClass* GetHandler(FName InName) const
+	{
+		return static_cast<THandlerClass*>(GetHandlerInternal(InName));
+	}
+
 	// Get the message log for the current compilation
 	FCompilerResultsLog& GetMessageLog() const { return GetMessageLogImpl(); }
 
-	// Index of the nodes (must match up with the runtime discovery process of nodes, which runs thru the property chain)
-	const TMap<UAnimGraphNode_Base*, int32>& GetAllocatedAnimNodeIndices() const { return GetAllocatedAnimNodeIndicesImpl(); }
+protected:	
+	// GetHandler helper function
+	virtual IAnimBlueprintCompilerHandler* GetHandlerInternal(FName InName) const = 0;
 
-protected:
 	// Get the message log for the current compilation
 	virtual FCompilerResultsLog& GetMessageLogImpl() const = 0;
-
-	// Map of anim node properties to original anim graph node
-	virtual const TMap<UAnimGraphNode_Base*, int32>& GetAllocatedAnimNodeIndicesImpl() const = 0;
 };
diff --git a/Engine/Source/Editor/AnimGraph/Public/IAnimBlueprintCompilationContext.h b/Engine/Source/Editor/AnimGraph/Public/IAnimBlueprintCompilationContext.h
--- a/Engine/Source/Editor/AnimGraph/Public/IAnimBlueprintCompilationContext.h
+++ b/Engine/Source/Editor/AnimGraph/Public/IAnimBlueprintCompilationContext.h
@@ -16,6 +16,7 @@ class UBlueprint;
 class UAnimBlueprint;
 class FProperty;
 struct FKismetCompilerOptions;
+class IAnimBlueprintCompilerHandler;
 
 #define ANIM_FUNC_DECORATOR	TEXT("__AnimFunc")
 
@@ -23,48 +24,19 @@ struct FKismetCompilerOptions;
 class ANIMGRAPH_API IAnimBlueprintCompilationContext
 {
 public:
-	// Record of a property that can be folded into the class members/constant blocks
-	struct FFoldedPropertyRecord
-	{
-		FFoldedPropertyRecord(UAnimGraphNode_Base* InAnimGraphNode, FStructProperty* InAnimNodeProperty, FProperty* InProperty, bool bInIsOnClass)
-			: AnimGraphNode(InAnimGraphNode)
-			, AnimNodeProperty(InAnimNodeProperty)
-			, Property(InProperty)
-			, GeneratedProperty(nullptr)
-			, FoldIndex(INDEX_NONE)
-			, PropertyIndex(INDEX_NONE)
-			, bIsOnClass(bInIsOnClass)
-		{
-		}
-
-		// The anim graph node that this property record came from
-		UAnimGraphNode_Base* AnimGraphNode = nullptr;
-
-		// The property of the FAnimNode_Base-derived structure within the anim graph node
-		FStructProperty* AnimNodeProperty = nullptr;
-
-		// The original property within the FAnimNode_Base
-		FProperty* Property = nullptr;
-
-		// The generated property within the respective data area (either constants or mutables)
-		FProperty* GeneratedProperty = nullptr;
-
-		// The index that this property was folded to. INDEX_NONE if it was not folded.
-		int32 FoldIndex = INDEX_NONE;
-
-		// The index of the property in its respective data area (either constant or mutable)
-		// This will be INDEX_NONE if the property was folded.
-		int32 PropertyIndex = INDEX_NONE;
-
-		// Whether this property will be held on the class (constants on sparse class data), or on the instance (mutables struct)
-		bool bIsOnClass = false;
-	};
-
 	virtual ~IAnimBlueprintCompilationContext() {}
 
 	// Get a compilation context from a kismet compiler context assuming that it is an FAnimBlueprintCompilerContext
 	static TUniquePtr<IAnimBlueprintCompilationContext> Get(FKismetCompilerContext& InKismetCompiler);
 
+	// Get a handler of the specified type and name (i.e. via simple name-based RTTI)
+	// Handlers are registered via IAnimBlueprintCompilerHandlerCollection::RegisterHandler
+	template <typename THandlerClass>
+	THandlerClass* GetHandler(FName InName) const
+	{
+		return static_cast<THandlerClass*>(GetHandlerInternal(InName));
+	}
+
 	// Spawns an intermediate node associated with the source node (for error purposes)
 	template <typename NodeType>
 	NodeType* SpawnIntermediateNode(UEdGraphNode* SourceNode, UEdGraph* ParentGraph = nullptr)
@@ -137,15 +109,6 @@ public:
 	// Get the current compiled-in attributes uniquely assigned to the specified node
 	TArrayView<const FName> GetAttributesFromNode(UAnimGraphNode_Base* InNode) const { return GetAttributesFromNodeImpl(InNode); }
 
-	// Check whether an anim node participates in constant folding
-	bool IsAnimGraphNodeFolded(UAnimGraphNode_Base* InNode) const { return IsAnimGraphNodeFoldedImpl(InNode); }
-
-	// Get the folded property record, if any, for the supplied node & named property
-	const FFoldedPropertyRecord* GetFoldedPropertyRecord(UAnimGraphNode_Base* InNode, FName InPropertyName) const { return GetFoldedPropertyRecordImpl(InNode, InPropertyName); }
-
-	// Get the generated property of the class that mutable data is added to
-	virtual const FStructProperty* GetMutableDataProperty() const { return GetMutableDataPropertyImpl(); }
-	
 protected:
 	// Adds a pose link mapping record
 	virtual void AddPoseLinkMappingRecordImpl(const FPoseLinkMappingRecord& InRecord) = 0;
@@ -198,14 +161,8 @@ protected:
 	// Get the current compiled-in attributes uniquely assigned to the specified node
 	virtual TArrayView<const FName> GetAttributesFromNodeImpl(UAnimGraphNode_Base* InNode) const = 0;
 
-	// Check whether an anim node participates in constant folding
-	virtual bool IsAnimGraphNodeFoldedImpl(UAnimGraphNode_Base* InNode) const = 0;
-
-	// Get the folded property record, if any, for the supplied node & named property
-	virtual const FFoldedPropertyRecord* GetFoldedPropertyRecordImpl(UAnimGraphNode_Base* InNode, FName InPropertyName) const = 0;
-
-	// Get the generated property of the class that mutable data is added to
-	virtual const FStructProperty* GetMutableDataPropertyImpl() const = 0;
+	// GetHandler helper function
+	virtual IAnimBlueprintCompilerHandler* GetHandlerInternal(FName InName) const = 0;
 
 	// Get the compiler as a base class to avoid circular include issues with templated functions/classes
 	virtual FKismetCompilerContext* GetKismetCompiler() const = 0;
