Fix crash in debug where delegate self-unregistration removed lambda underneath itself

diff --git a/Engine/Source/Editor/AnimGraph/Private/AnimBlueprintExtension_Base.cpp b/Engine/Source/Editor/AnimGraph/Private/AnimBlueprintExtension_Base.cpp
--- a/Engine/Source/Editor/AnimGraph/Private/AnimBlueprintExtension_Base.cpp
+++ b/Engine/Source/Editor/AnimGraph/Private/AnimBlueprintExtension_Base.cpp
@@ -152,71 +152,65 @@ void UAnimBlueprintExtension_Base::HandleStartCompilingClass(const UClass* InCla
 	ValidEvaluationHandlerList.Empty();
 	ValidEvaluationHandlerMap.Empty();
 	HandlerFunctionNames.Empty();
-	
+	PreLibraryCompiledDelegateHandle.Reset();
+	PostLibraryCompiledDelegateHandle.Reset();	
+
 	UAnimBlueprintExtension_PropertyAccess* PropertyAccessExtension = UAnimBlueprintExtension::GetExtension<UAnimBlueprintExtension_PropertyAccess>(GetAnimBlueprint());
 	if(PropertyAccessExtension)
 	{
-		if(!PreLibraryCompiledDelegateHandle.IsValid())
+		PreLibraryCompiledDelegateHandle = PropertyAccessExtension->OnPreLibraryCompiled().AddLambda([this, PropertyAccessExtension, InClass]()
 		{
-			PreLibraryCompiledDelegateHandle = PropertyAccessExtension->OnPreLibraryCompiled().AddLambda([this, PropertyAccessExtension, InClass]()
+			if(IModularFeatures::Get().IsModularFeatureAvailable("PropertyAccessEditor"))
 			{
-				if(IModularFeatures::Get().IsModularFeatureAvailable("PropertyAccessEditor"))
-				{
-					IPropertyAccessEditor& PropertyAccessEditor = IModularFeatures::Get().GetModularFeature<IPropertyAccessEditor>("PropertyAccessEditor");
+				IPropertyAccessEditor& PropertyAccessEditor = IModularFeatures::Get().GetModularFeature<IPropertyAccessEditor>("PropertyAccessEditor");
 
-					// Build the classes property access library before the library is compiled
-					for(FEvaluationHandlerRecord& HandlerRecord : ValidEvaluationHandlerList)
+				// Build the classes property access library before the library is compiled
+				for(FEvaluationHandlerRecord& HandlerRecord : ValidEvaluationHandlerList)
+				{
+					for(TPair<FName, FAnimNodeSinglePropertyHandler>& PropertyHandler : HandlerRecord.ServicedProperties)
 					{
-						for(TPair<FName, FAnimNodeSinglePropertyHandler>& PropertyHandler : HandlerRecord.ServicedProperties)
+						for(FPropertyCopyRecord& Record : PropertyHandler.Value.CopyRecords)
 						{
-							for(FPropertyCopyRecord& Record : PropertyHandler.Value.CopyRecords)
+							if(Record.IsFastPath())
 							{
-								if(Record.IsFastPath())
-								{
-									// Check if the resolved copy
-									FProperty* LeafProperty = nullptr;
-									int32 ArrayIndex = INDEX_NONE;
-									EPropertyAccessResolveResult Result = PropertyAccessEditor.ResolveLeafProperty(InClass, Record.SourcePropertyPath, LeafProperty, ArrayIndex);
-
-									// Batch all external accesses, we cant call them safely from a worker thread.
-									Record.LibraryBatchType = Result == EPropertyAccessResolveResult::SucceededExternal ? EPropertyAccessBatchType::Batched : EPropertyAccessBatchType::Unbatched;
-									Record.LibraryCopyIndex = PropertyAccessExtension->AddCopy(Record.SourcePropertyPath, Record.DestPropertyPath, Record.LibraryBatchType, HandlerRecord.AnimGraphNode);
-								}
+								// Check if the resolved copy
+								FProperty* LeafProperty = nullptr;
+								int32 ArrayIndex = INDEX_NONE;
+								EPropertyAccessResolveResult Result = PropertyAccessEditor.ResolveLeafProperty(InClass, Record.SourcePropertyPath, LeafProperty, ArrayIndex);
+
+								// Batch all external accesses, we cant call them safely from a worker thread.
+								Record.LibraryBatchType = Result == EPropertyAccessResolveResult::SucceededExternal ? EPropertyAccessBatchType::Batched : EPropertyAccessBatchType::Unbatched;
+								Record.LibraryCopyIndex = PropertyAccessExtension->AddCopy(Record.SourcePropertyPath, Record.DestPropertyPath, Record.LibraryBatchType, HandlerRecord.AnimGraphNode);
 							}
 						}
 					}
 				}
+			}
 
-				PropertyAccessExtension->OnPreLibraryCompiled().Remove(PreLibraryCompiledDelegateHandle);
-				PreLibraryCompiledDelegateHandle.Reset();
-			});
-		}
+			PropertyAccessExtension->OnPreLibraryCompiled().Remove(PreLibraryCompiledDelegateHandle);
+		});
 
-		if(!PostLibraryCompiledDelegateHandle.IsValid())
+		PostLibraryCompiledDelegateHandle = PropertyAccessExtension->OnPostLibraryCompiled().AddLambda([this, PropertyAccessExtension](IAnimBlueprintCompilationBracketContext& InCompilationContext, IAnimBlueprintGeneratedClassCompiledData& OutCompiledData)
 		{
-			PostLibraryCompiledDelegateHandle = PropertyAccessExtension->OnPostLibraryCompiled().AddLambda([this, PropertyAccessExtension](IAnimBlueprintCompilationBracketContext& InCompilationContext, IAnimBlueprintGeneratedClassCompiledData& OutCompiledData)
+			for(FEvaluationHandlerRecord& HandlerRecord : ValidEvaluationHandlerList)
 			{
-				for(FEvaluationHandlerRecord& HandlerRecord : ValidEvaluationHandlerList)
+				// Map global copy index to batched indices
+				for(TPair<FName, FAnimNodeSinglePropertyHandler>& PropertyHandler : HandlerRecord.ServicedProperties)
 				{
-					// Map global copy index to batched indices
-					for(TPair<FName, FAnimNodeSinglePropertyHandler>& PropertyHandler : HandlerRecord.ServicedProperties)
+					for(FPropertyCopyRecord& CopyRecord : PropertyHandler.Value.CopyRecords)
 					{
-						for(FPropertyCopyRecord& CopyRecord : PropertyHandler.Value.CopyRecords)
+						if(CopyRecord.IsFastPath())
 						{
-							if(CopyRecord.IsFastPath())
-							{
-								CopyRecord.LibraryCopyIndex = PropertyAccessExtension->MapCopyIndex(CopyRecord.LibraryCopyIndex);
-							}
+							CopyRecord.LibraryCopyIndex = PropertyAccessExtension->MapCopyIndex(CopyRecord.LibraryCopyIndex);
 						}
 					}
 				}
+			}
 
-				PatchEvaluationHandlers(InCompilationContext, OutCompiledData);
+			PatchEvaluationHandlers(InCompilationContext, OutCompiledData);
 
-				PropertyAccessExtension->OnPostLibraryCompiled().Remove(PostLibraryCompiledDelegateHandle);
-				PostLibraryCompiledDelegateHandle.Reset();	
-			});
-		}
+			PropertyAccessExtension->OnPostLibraryCompiled().Remove(PostLibraryCompiledDelegateHandle);
+		});
 	}
 }
 
