Fix anim BP internal structs showing up in the content browser

diff --git a/Engine/Source/Editor/AnimGraph/Private/AnimBlueprintCompiler.cpp b/Engine/Source/Editor/AnimGraph/Private/AnimBlueprintCompiler.cpp
--- a/Engine/Source/Editor/AnimGraph/Private/AnimBlueprintCompiler.cpp
+++ b/Engine/Source/Editor/AnimGraph/Private/AnimBlueprintCompiler.cpp
@@ -1187,7 +1187,7 @@ void FAnimBlueprintCompilerContext::RecreateSparseClassData()
 		check(ParentAnimClass->GetSparseClassDataStruct());
 		
 		// Derive sparse class data from parent class
-		NewAnimBlueprintConstants = NewObject<UScriptStruct>(NewAnimBlueprintClass->GetOuter(), *ConstantsStructTypeName, RF_Public);
+		NewAnimBlueprintConstants = NewObject<UAnimBlueprintGeneratedStruct>(NewAnimBlueprintClass->GetOuter(), *ConstantsStructTypeName, RF_Public);
 		NewAnimBlueprintConstants->SetSuperStruct(ParentAnimClass->GetSparseClassDataStruct());
 
 		// Just link & assign sparse class data struct here, no additional members are added
@@ -1202,7 +1202,7 @@ void FAnimBlueprintCompilerContext::RecreateSparseClassData()
 		check(ParentClass);
 		
 	    // Create new sparse class data struct
-		NewAnimBlueprintConstants = NewObject<UScriptStruct>(NewAnimBlueprintClass->GetOuter(), *ConstantsStructTypeName, RF_Public);
+		NewAnimBlueprintConstants = NewObject<UAnimBlueprintGeneratedStruct>(NewAnimBlueprintClass->GetOuter(), *ConstantsStructTypeName, RF_Public);
 
 		// Inherit from archetype struct if there is any 
 		UScriptStruct* ArchetypeStruct = ParentClass->GetSparseClassDataArchetypeStruct();
@@ -1224,7 +1224,7 @@ void FAnimBlueprintCompilerContext::RecreateMutables()
 		OldAnimBlueprintMutables->Rename(nullptr, GetTransientPackage(), REN_DontCreateRedirectors | REN_NonTransactional | REN_ForceNoResetLoaders);
 	}
 
-	NewAnimBlueprintMutables = NewObject<UScriptStruct>(NewAnimBlueprintClass->GetOuter(), *MutablesStructTypeName);
+	NewAnimBlueprintMutables = NewObject<UAnimBlueprintGeneratedStruct>(NewAnimBlueprintClass->GetOuter(), *MutablesStructTypeName);
 	NewAnimBlueprintMutables->SetSuperStruct(FAnimBlueprintMutableData::StaticStruct());
 }
 
diff --git a/Engine/Source/Runtime/Engine/Classes/Animation/AnimBlueprintGeneratedClass.h b/Engine/Source/Runtime/Engine/Classes/Animation/AnimBlueprintGeneratedClass.h
--- a/Engine/Source/Runtime/Engine/Classes/Animation/AnimBlueprintGeneratedClass.h
+++ b/Engine/Source/Runtime/Engine/Classes/Animation/AnimBlueprintGeneratedClass.h
@@ -338,6 +338,17 @@ namespace EPropertySearchMode
 }
 #endif
 
+// Struct type generated by the anim BP compiler. Used for sparse class data and mutable data area.
+// Only really needed to hide the struct from the content browser (via IsAsset override)
+UCLASS()
+class ENGINE_API UAnimBlueprintGeneratedStruct : public UScriptStruct
+{
+	GENERATED_BODY()
+
+	// UObject interface
+	virtual bool IsAsset() const override { return false; }
+};
+
 UCLASS()
 class ENGINE_API UAnimBlueprintGeneratedClass : public UBlueprintGeneratedClass, public IAnimClassInterface
 {
