Fix for renaming issues with animation blueprints

Switches generated structs to be subobjects of the generated class.
Key to this is supporting re-instancing via re-linking the mutable struct before the generated class is linked to allow the correct class layout to be used by serializers. This is because the property links are not preserved with StaticDuplicateObject.

diff --git a/Engine/Source/Editor/AnimGraph/Private/AnimBlueprintCompiler.cpp b/Engine/Source/Editor/AnimGraph/Private/AnimBlueprintCompiler.cpp
--- a/Engine/Source/Editor/AnimGraph/Private/AnimBlueprintCompiler.cpp
+++ b/Engine/Source/Editor/AnimGraph/Private/AnimBlueprintCompiler.cpp
@@ -1169,14 +1169,6 @@ void FAnimBlueprintCompilerContext::RecreateSparseClassData()
 	UAnimBlueprintGeneratedClass* NewAnimBlueprintClass = GetNewAnimBlueprintClass();
 	
 	// Set up our sparse class data struct
-	const FString ConstantsStructTypeName = NewAnimBlueprintClass->GetName() + TEXT("_ConstantData");
-
-	// Rename old struct out of the way
-	if(UScriptStruct* OldAnimBlueprintConstants = FindObject<UScriptStruct>(NewAnimBlueprintClass->GetOuter(), *ConstantsStructTypeName))
-	{
-		OldAnimBlueprintConstants->Rename(nullptr, GetTransientPackage(), REN_DontCreateRedirectors | REN_NonTransactional | REN_ForceNoResetLoaders);
-	}
-
 	if (bIsDerivedAnimBlueprint)
 	{
 		// Get parent class
@@ -1187,7 +1179,7 @@ void FAnimBlueprintCompilerContext::RecreateSparseClassData()
 		check(ParentAnimClass->GetSparseClassDataStruct());
 		
 		// Derive sparse class data from parent class
-		NewAnimBlueprintConstants = NewObject<UAnimBlueprintGeneratedStruct>(NewAnimBlueprintClass->GetOuter(), *ConstantsStructTypeName, RF_Public);
+		NewAnimBlueprintConstants = NewObject<UScriptStruct>(NewAnimBlueprintClass, UAnimBlueprintGeneratedClass::GetConstantsStructName(), RF_Public);
 		NewAnimBlueprintConstants->SetSuperStruct(ParentAnimClass->GetSparseClassDataStruct());
 
 		// Just link & assign sparse class data struct here, no additional members are added
@@ -1202,7 +1194,7 @@ void FAnimBlueprintCompilerContext::RecreateSparseClassData()
 		check(ParentClass);
 		
 	    // Create new sparse class data struct
-		NewAnimBlueprintConstants = NewObject<UAnimBlueprintGeneratedStruct>(NewAnimBlueprintClass->GetOuter(), *ConstantsStructTypeName, RF_Public);
+		NewAnimBlueprintConstants = NewObject<UScriptStruct>(NewAnimBlueprintClass, UAnimBlueprintGeneratedClass::GetConstantsStructName(), RF_Public);
 
 		// Inherit from archetype struct if there is any 
 		UScriptStruct* ArchetypeStruct = ParentClass->GetSparseClassDataArchetypeStruct();
@@ -1216,15 +1208,7 @@ void FAnimBlueprintCompilerContext::RecreateMutables()
 	UAnimBlueprintGeneratedClass* NewAnimBlueprintClass = GetNewAnimBlueprintClass();
 
 	// Set up our mutables struct
-	const FString MutablesStructTypeName = NewAnimBlueprintClass->GetName() + TEXT("_MutableData");
-
-	// Rename old struct out of the way
-	if(UScriptStruct* OldAnimBlueprintMutables = FindObject<UScriptStruct>(NewAnimBlueprintClass->GetOuter(), *MutablesStructTypeName))
-	{
-		OldAnimBlueprintMutables->Rename(nullptr, GetTransientPackage(), REN_DontCreateRedirectors | REN_NonTransactional | REN_ForceNoResetLoaders);
-	}
-
-	NewAnimBlueprintMutables = NewObject<UAnimBlueprintGeneratedStruct>(NewAnimBlueprintClass->GetOuter(), *MutablesStructTypeName);
+	NewAnimBlueprintMutables = NewObject<UScriptStruct>(NewAnimBlueprintClass, UAnimBlueprintGeneratedClass::GetMutablesStructName());
 	NewAnimBlueprintMutables->SetSuperStruct(FAnimBlueprintMutableData::StaticStruct());
 }
 
@@ -1269,7 +1253,9 @@ void FAnimBlueprintCompilerContext::CleanAndSanitizeClass(UBlueprintGeneratedCla
 	Super::CleanAndSanitizeClass(ClassToClean, InOldCDO);
 
 	UAnimBlueprintGeneratedClass* AnimBlueprintClassToClean= CastChecked<UAnimBlueprintGeneratedClass>(ClassToClean);
-	
+
+	AnimBlueprintClassToClean->SetSparseClassDataStruct(nullptr);
+
 	AnimBlueprintClassToClean->AnimBlueprintDebugData = FAnimBlueprintDebugData();
 
 	// Reset the baked data
diff --git a/Engine/Source/Runtime/Engine/Classes/Animation/AnimBlueprintGeneratedClass.h b/Engine/Source/Runtime/Engine/Classes/Animation/AnimBlueprintGeneratedClass.h
--- a/Engine/Source/Runtime/Engine/Classes/Animation/AnimBlueprintGeneratedClass.h
+++ b/Engine/Source/Runtime/Engine/Classes/Animation/AnimBlueprintGeneratedClass.h
@@ -599,6 +599,10 @@ public:
 	// Build the properties that we cache for our constant data
 	void BuildConstantProperties();
 
+	// Get the fixed names of our generated structs
+	static FName GetConstantsStructName();
+	static FName GetMutablesStructName();
+	
 	// UObject interface
 	virtual void Serialize(FArchive& Ar) override;
 	// End of UObject interface
diff --git a/Engine/Source/Runtime/Engine/Private/Animation/AnimBlueprintGeneratedClass.cpp b/Engine/Source/Runtime/Engine/Private/Animation/AnimBlueprintGeneratedClass.cpp
--- a/Engine/Source/Runtime/Engine/Private/Animation/AnimBlueprintGeneratedClass.cpp
+++ b/Engine/Source/Runtime/Engine/Private/Animation/AnimBlueprintGeneratedClass.cpp
@@ -448,8 +448,17 @@ void UAnimBlueprintGeneratedClass::Serialize(FArchive& Ar)
 
 void UAnimBlueprintGeneratedClass::Link(FArchive& Ar, bool bRelinkExistingProperties)
 {
+#if WITH_EDITOR
+	// Re-link the mutables struct if we have one as we need to correctly account for its size
+	// This is mostly needed during re-instancing as post duplication the class requires a re-link
+	if(UScriptStruct* MutablesStruct = FindObject<UScriptStruct>(this, *GetMutablesStructName().ToString()))
+	{
+		MutablesStruct->Link(Ar, bRelinkExistingProperties);
+	}
+#endif	
+	
 	Super::Link(Ar, bRelinkExistingProperties);
-
+	
 	// @TODO: Shouldn't be necessary to clear these, but currently the class gets linked twice during compilation
 	AnimNodeProperties.Empty();
 	LinkedAnimGraphNodeProperties.Empty();
@@ -889,4 +898,16 @@ int32 UAnimBlueprintGeneratedClass::GetAnimNodePropertyCount(const UScriptStruct
 	}
 	
 	return 0;
+}
+
+FName UAnimBlueprintGeneratedClass::GetConstantsStructName()
+{
+	static const FName Name(TEXT("AnimBlueprintGeneratedConstantData"));
+	return Name;
+}
+
+FName UAnimBlueprintGeneratedClass::GetMutablesStructName()
+{
+	static const FName Name(TEXT("AnimBlueprintGeneratedMutableData"));
+	return Name;
 }
\ No newline at end of file
