Warn if MirrorDataTable used in a MirrorNode is not compatible with the current skeleton

diff --git a/Engine/Source/Editor/AnimGraph/Private/AnimGraphNode_Mirror.cpp b/Engine/Source/Editor/AnimGraph/Private/AnimGraphNode_Mirror.cpp
--- a/Engine/Source/Editor/AnimGraph/Private/AnimGraphNode_Mirror.cpp
+++ b/Engine/Source/Editor/AnimGraph/Private/AnimGraphNode_Mirror.cpp
@@ -59,6 +59,14 @@ void UAnimGraphNode_Mirror::ValidateAnimNodeDuringCompilation(class USkeleton* F
 	{
 		MessageLog.Error(TEXT("@@ does not have a mirror data table selected.  Please select a table or delete the node."), this);
 	}
+	else if (ForSkeleton)
+	{
+		if (!ForSkeleton->IsCompatible(Node.GetMirrorDataTable()->Skeleton))
+		{
+			MessageLog.Error(TEXT("@@ has a mirror data table that is not compatible with the current skeleton. Please update the table or create a new table for the skeleton."), this);
+		}
+	}
+	
 	BindMirrorDataTableChangedDelegate();
 }
 
