Fix reference possibly being invalidated in while loop

diff --git a/Engine/Source/Editor/AnimGraph/Private/AnimBlueprintCompiler.cpp b/Engine/Source/Editor/AnimGraph/Private/AnimBlueprintCompiler.cpp
--- a/Engine/Source/Editor/AnimGraph/Private/AnimBlueprintCompiler.cpp
+++ b/Engine/Source/Editor/AnimGraph/Private/AnimBlueprintCompiler.cpp
@@ -1850,7 +1850,8 @@ void FAnimBlueprintCompilerContext::ProcessFoldedPropertyRecords()
 		const int32 NodePropertyIndex = AllocatedAnimNodeIndices.Num() - 1 - NodeIndex;
 		FStructProperty* NodeStructProperty = CastFieldChecked<FStructProperty>(AllocatedPropertiesByIndex.FindChecked(NodePropertyIndex));
 		const UScriptStruct* Struct = NodeStructProperty->Struct;
-		const FAnimNodeStructData& AnimNodeStructData = NewAnimBlueprintClass->NodeTypeMap.Add(Struct, FAnimNodeStructData(Struct));
+		const FAnimNodeStructData AnimNodeStructData = NewAnimBlueprintClass->NodeTypeMap.Add(Struct, FAnimNodeStructData(Struct));
+		const int32 NumProperties = AnimNodeStructData.GetNumProperties();
 
 		// Add any super-structs as values can be accessed via base classes
 		const UScriptStruct* SuperStruct = Cast<UScriptStruct>(Struct->GetSuperStruct());
@@ -1862,7 +1863,6 @@ void FAnimBlueprintCompilerContext::ProcessFoldedPropertyRecords()
 		
 		NodeData.AnimClassInterface = NewAnimBlueprintClass;
 		NodeData.NodeIndex = NodeIndex++;
-		const int32 NumProperties = AnimNodeStructData.GetNumProperties();
 		check(NumProperties >= 0);
 		NodeData.Entries.SetNum(NumProperties);
 		for(uint32& Entry : NodeData.Entries)
@@ -1928,4 +1928,4 @@ const IAnimBlueprintCompilationContext::FFoldedPropertyRecord* FAnimBlueprintCom
 
 //////////////////////////////////////////////////////////////////////////
 
-#undef LOCTEXT_NAMESPACE
\ No newline at end of file
+#undef LOCTEXT_NAMESPACE
