Fix for double click not working with shared transition rules
#rnx

diff --git a/Engine/Source/Editor/AnimGraph/Private/AnimStateTransitionNode.cpp b/Engine/Source/Editor/AnimGraph/Private/AnimStateTransitionNode.cpp
--- a/Engine/Source/Editor/AnimGraph/Private/AnimStateTransitionNode.cpp
+++ b/Engine/Source/Editor/AnimGraph/Private/AnimStateTransitionNode.cpp
@@ -641,6 +641,12 @@ void UAnimStateTransitionNode::ValidateNodeDuringCompilation(class FCompilerResu
 	}
 }
 
+UObject* UAnimStateTransitionNode::GetJumpTargetForDoubleClick() const
+{
+	// Our base class uses GetSubGraphs. Since we explicitly ignore a shared bound graph, use BoundGraph directly instead.
+	return BoundGraph;
+}
+
 TArray<UEdGraph*> UAnimStateTransitionNode::GetSubGraphs() const
 { 
 	TArray<UEdGraph*> SubGraphs;
diff --git a/Engine/Source/Editor/AnimGraph/Public/AnimStateTransitionNode.h b/Engine/Source/Editor/AnimGraph/Public/AnimStateTransitionNode.h
--- a/Engine/Source/Editor/AnimGraph/Public/AnimStateTransitionNode.h
+++ b/Engine/Source/Editor/AnimGraph/Public/AnimStateTransitionNode.h
@@ -124,6 +124,7 @@ class UAnimStateTransitionNode : public UAnimStateNodeBase
 	virtual void PostPlacedNewNode() override;
 	virtual void DestroyNode() override;
 	virtual void ValidateNodeDuringCompilation(class FCompilerResultsLog& MessageLog) const override;
+	virtual UObject* GetJumpTargetForDoubleClick() const override;
 	virtual TArray<UEdGraph*> GetSubGraphs() const override;
 	//~ End UEdGraphNode Interface
 
