Fix a mirror node compilation error encountered with the mirror data table is not yet loaded

diff --git a/Engine/Source/Editor/AnimGraph/Classes/AnimGraphNode_Mirror.h b/Engine/Source/Editor/AnimGraph/Classes/AnimGraphNode_Mirror.h
--- a/Engine/Source/Editor/AnimGraph/Classes/AnimGraphNode_Mirror.h
+++ b/Engine/Source/Editor/AnimGraph/Classes/AnimGraphNode_Mirror.h
@@ -29,6 +29,7 @@ public:
 	virtual EAnimAssetHandlerType SupportsAssetClass(const UClass* AssetClass) const override;
 	virtual void PreEditChange(FProperty* PropertyThatWillChange) override;
 	virtual void PostEditChangeProperty(struct FPropertyChangedEvent& PropertyChangedEvent) override;
+	virtual void PreloadRequiredAssets() override;
 	void ValidateAnimNodeDuringCompilation(class USkeleton* ForSkeleton, class FCompilerResultsLog& MessageLog) override;
 
 private:
diff --git a/Engine/Source/Editor/AnimGraph/Private/AnimGraphNode_Mirror.cpp b/Engine/Source/Editor/AnimGraph/Private/AnimGraphNode_Mirror.cpp
--- a/Engine/Source/Editor/AnimGraph/Private/AnimGraphNode_Mirror.cpp
+++ b/Engine/Source/Editor/AnimGraph/Private/AnimGraphNode_Mirror.cpp
@@ -51,6 +51,16 @@ void UAnimGraphNode_Mirror::GetOutputLinkAttributes(FNodeAttributeArray& OutAttr
 	}
 }
 
+void UAnimGraphNode_Mirror::PreloadRequiredAssets()
+{
+	if (Node.GetMirrorDataTable())
+	{
+		PreloadObject(Node.GetMirrorDataTable());
+		PreloadObject(Node.GetMirrorDataTable()->Skeleton);
+	}
+	Super::PreloadRequiredAssets();
+}
+
 void UAnimGraphNode_Mirror::ValidateAnimNodeDuringCompilation(class USkeleton* ForSkeleton, class FCompilerResultsLog& MessageLog)
 {
 	Super::ValidateAnimNodeDuringCompilation(ForSkeleton, MessageLog);
@@ -61,7 +71,8 @@ void UAnimGraphNode_Mirror::ValidateAnimNodeDuringCompilation(class USkeleton* F
 	}
 	else if (ForSkeleton)
 	{
-		if (!ForSkeleton->IsCompatible(Node.GetMirrorDataTable()->Skeleton))
+		const USkeleton* MirrorTableSkeleton = Node.GetMirrorDataTable()->Skeleton;
+		if (!ForSkeleton->IsCompatible(MirrorTableSkeleton))
 		{
 			MessageLog.Error(TEXT("@@ has a mirror data table that is not compatible with the current skeleton. Please update the table or create a new table for the skeleton."), this);
 		}
