Fix various issues with child anim BPs

If a child anim BP is made and is saved without saving its parent, then when the class is reloaded later it will not be able to find the parent's sparse class data in the linker (as it will have been generated again from scratch on load, which doesnt patch the linker until the asset is resaved). This leads to serialization size mismatches.
To address this (and because the sparse class data will always be regenerated as part of compile-on-load), we skip the portion of the CDO's archive that represents the serialized sparse class data if it is null after serialization. It can of course be legitimately null in a lot of cases (say a regular actor Blueprint), but this should have no ill effects and effectively be a seek to the current archive offset.

This CL also fixes a number of child anim BP-related issues. First of all it prevents a crash when accessing out-of-bounds AnimNodeData by ensuring it is correctly duplicated in all circumstances. It also correctly propogates values from root->child sparse class data (in the same manner that anim node data used to be propogated before). This deals with compile-on-load issues where the child classes sparse class data is regenerated without serializing the parent classes sparse class data (see crash fix above).

diff --git a/Engine/Source/Editor/AnimGraph/Private/AnimBlueprintCompiler.cpp b/Engine/Source/Editor/AnimGraph/Private/AnimBlueprintCompiler.cpp
--- a/Engine/Source/Editor/AnimGraph/Private/AnimBlueprintCompiler.cpp
+++ b/Engine/Source/Editor/AnimGraph/Private/AnimBlueprintCompiler.cpp
@@ -705,8 +705,9 @@ void FAnimBlueprintCompilerContext::CopyTermDefaultsToDefaultObject(UObject* Def
 
 	if (bIsDerivedAnimBlueprint && DefaultAnimInstance)
 	{
-		// Make sure our sparse class data is re-created
-		RecreateSparseClassData();
+		// Ensure we have constant properties & anim node data rebuilt
+		NewAnimBlueprintClass->BuildConstantProperties();
+		CopyAnimNodeDataFromRoot();
 		
 		// If we are a derived animation graph; apply any stored overrides.
 		// Restore values from the root class to catch values where the override has been removed.
@@ -715,14 +716,6 @@ void FAnimBlueprintCompilerContext::CopyTermDefaultsToDefaultObject(UObject* Def
 		{
 			RootAnimClass = NextClass;
 		}
-
-		// Copy constant/folded data from root class, remapping the class
-		NewAnimBlueprintClass->AnimNodeData = RootAnimClass->AnimNodeData;
-		NewAnimBlueprintClass->NodeTypeMap = RootAnimClass->NodeTypeMap;
-		for(FAnimNodeData& NodeData : NewAnimBlueprintClass->AnimNodeData)
-		{
-			NodeData.AnimClassInterface = NewAnimBlueprintClass;
-		}
 		
 		UObject* RootDefaultObject = RootAnimClass->GetDefaultObject();
 
@@ -744,6 +737,24 @@ void FAnimBlueprintCompilerContext::CopyTermDefaultsToDefaultObject(UObject* Def
 			}
 		}
 
+		// Copy from root sparse class data to our new class
+		if(NewAnimBlueprintClass->GetConstantNodeData() && RootAnimClass->GetConstantNodeData())
+		{
+			check(NewAnimBlueprintClass->GetSparseClassDataStruct()->IsChildOf(RootAnimClass->GetSparseClassDataStruct()));
+
+			for (TFieldIterator<FProperty> PropertyIt(RootAnimClass->GetSparseClassDataStruct()); PropertyIt; ++PropertyIt)
+			{
+				FProperty* RootProperty = *PropertyIt;
+				FProperty* ChildProperty = FindFProperty<FProperty>(NewAnimBlueprintClass->GetSparseClassDataStruct(), *RootProperty->GetName());
+				check(ChildProperty);
+				
+				const uint8* SourcePtr = RootProperty->ContainerPtrToValuePtr<uint8>(RootAnimClass->GetConstantNodeData());
+				uint8* DestPtr = const_cast<uint8*>(ChildProperty->ContainerPtrToValuePtr<uint8>(NewAnimBlueprintClass->GetConstantNodeData()));
+				check(SourcePtr && DestPtr);
+				RootProperty->CopyCompleteValue(DestPtr, SourcePtr);
+			}
+		}
+		
 		// Re-initialize node data tables (they would be overwritten in the loop above)
 		NewAnimBlueprintClass->InitializeAnimNodeData(DefaultObject);
 	}
@@ -1083,6 +1094,9 @@ void FAnimBlueprintCompilerContext::MergeUbergraphPagesIn(UEdGraph* Ubergraph)
 	if (bIsDerivedAnimBlueprint)
 	{
 		// Skip any work related to an anim graph, it's all done by the parent class
+		// We do need to make sure that anim node data is correctly copied & remapped to this class
+		RecreateSparseClassData();
+		CopyAnimNodeDataFromRoot();
 	}
 	else
 	{
@@ -1143,6 +1157,26 @@ void FAnimBlueprintCompilerContext::MergeUbergraphPagesIn(UEdGraph* Ubergraph)
 	}
 }
 
+void FAnimBlueprintCompilerContext::CopyAnimNodeDataFromRoot() const
+{
+	check(bIsDerivedAnimBlueprint);
+
+	UAnimBlueprintGeneratedClass* NewAnimBlueprintClass = GetNewAnimBlueprintClass();
+	UAnimBlueprintGeneratedClass* RootAnimClass = NewAnimBlueprintClass;
+	while (UAnimBlueprintGeneratedClass* NextClass = Cast<UAnimBlueprintGeneratedClass>(RootAnimClass->GetSuperClass()))
+	{
+		RootAnimClass = NextClass;
+	}
+	
+	// Copy constant/folded data from root class, remapping the class
+	NewAnimBlueprintClass->AnimNodeData = RootAnimClass->AnimNodeData;
+	NewAnimBlueprintClass->NodeTypeMap = RootAnimClass->NodeTypeMap;
+	for(FAnimNodeData& NodeData : NewAnimBlueprintClass->AnimNodeData)
+	{
+		NodeData.AnimClassInterface = NewAnimBlueprintClass;
+	}
+}
+
 void FAnimBlueprintCompilerContext::ProcessOneFunctionGraph(UEdGraph* SourceGraph, bool bInternalFunction)
 {
 	if(!KnownGraphSchemas.FindByPredicate([SourceGraph](const TSubclassOf<UEdGraphSchema>& InSchemaClass)
@@ -1167,6 +1201,8 @@ void FAnimBlueprintCompilerContext::EnsureProperGeneratedClass(UClass*& TargetUC
 void FAnimBlueprintCompilerContext::RecreateSparseClassData()
 {
 	UAnimBlueprintGeneratedClass* NewAnimBlueprintClass = GetNewAnimBlueprintClass();
+
+	UScriptStruct* OldSparseClassDataStruct = NewAnimBlueprintClass->GetSparseClassDataStruct();
 	
 	// Set up our sparse class data struct
 	if (bIsDerivedAnimBlueprint)
@@ -1201,15 +1237,27 @@ void FAnimBlueprintCompilerContext::RecreateSparseClassData()
 		UScriptStruct* SuperStruct = ArchetypeStruct ? ArchetypeStruct : FAnimBlueprintConstantData::StaticStruct();
 		NewAnimBlueprintConstants->SetSuperStruct(SuperStruct);
 	}
+
+	if(OldSparseClassDataStruct)
+	{
+		FLinkerLoad::PRIVATE_PatchNewObjectIntoExport(OldSparseClassDataStruct, NewAnimBlueprintConstants);
+	}
 }
 
 void FAnimBlueprintCompilerContext::RecreateMutables()
 {
 	UAnimBlueprintGeneratedClass* NewAnimBlueprintClass = GetNewAnimBlueprintClass();
 
+	UScriptStruct* OldMutablesStruct = NewAnimBlueprintClass->MutableNodeDataProperty ? NewAnimBlueprintClass->MutableNodeDataProperty->Struct : nullptr;
+	
 	// Set up our mutables struct
 	NewAnimBlueprintMutables = NewObject<UScriptStruct>(NewAnimBlueprintClass, UAnimBlueprintGeneratedClass::GetMutablesStructName());
 	NewAnimBlueprintMutables->SetSuperStruct(FAnimBlueprintMutableData::StaticStruct());
+
+	if(OldMutablesStruct)
+	{
+		FLinkerLoad::PRIVATE_PatchNewObjectIntoExport(OldMutablesStruct, NewAnimBlueprintMutables);
+	}
 }
 
 void FAnimBlueprintCompilerContext::SpawnNewClass(const FString& NewClassName)
@@ -1254,8 +1302,6 @@ void FAnimBlueprintCompilerContext::CleanAndSanitizeClass(UBlueprintGeneratedCla
 
 	UAnimBlueprintGeneratedClass* AnimBlueprintClassToClean= CastChecked<UAnimBlueprintGeneratedClass>(ClassToClean);
 
-	AnimBlueprintClassToClean->SetSparseClassDataStruct(nullptr);
-
 	AnimBlueprintClassToClean->AnimBlueprintDebugData = FAnimBlueprintDebugData();
 
 	// Reset the baked data
diff --git a/Engine/Source/Editor/AnimGraph/Private/AnimBlueprintCompiler.h b/Engine/Source/Editor/AnimGraph/Private/AnimBlueprintCompiler.h
--- a/Engine/Source/Editor/AnimGraph/Private/AnimBlueprintCompiler.h
+++ b/Engine/Source/Editor/AnimGraph/Private/AnimBlueprintCompiler.h
@@ -202,6 +202,9 @@ private:
 	// Check whether an anim node participates in constant folding
 	bool IsAnimGraphNodeFolded(UAnimGraphNode_Base* InNode) const;
 
+	// Copy the AnimNodeData array etc. from root anim BP class to a derived anim BP 
+	void CopyAnimNodeDataFromRoot() const;
+	
 	// Get the folded property record, if any, for the supplied node & named property
 	const IAnimBlueprintCompilationContext::FFoldedPropertyRecord* GetFoldedPropertyRecord(UAnimGraphNode_Base* InNode, FName InPropertyName) const;
 };
diff --git a/Engine/Source/Runtime/Engine/Private/Animation/AnimBlueprintGeneratedClass.cpp b/Engine/Source/Runtime/Engine/Private/Animation/AnimBlueprintGeneratedClass.cpp
--- a/Engine/Source/Runtime/Engine/Private/Animation/AnimBlueprintGeneratedClass.cpp
+++ b/Engine/Source/Runtime/Engine/Private/Animation/AnimBlueprintGeneratedClass.cpp
@@ -681,7 +681,7 @@ void UAnimBlueprintGeneratedClass::InitializeAnimNodeData(UObject* DefaultObject
 		if(StructProperty->Struct->IsChildOf(FAnimNode_Base::StaticStruct()))
 		{
 			FAnimNode_Base* Node = StructProperty->ContainerPtrToValuePtr<FAnimNode_Base>(DefaultObject);
-			Node->SetNodeData(GetNodeData()[AnimNodeIndex]);
+			Node->SetNodeData(AnimNodeData[AnimNodeIndex]);
 		}
 	}
 }
@@ -700,7 +700,13 @@ void UAnimBlueprintGeneratedClass::LinkFunctionsToDefaultObjectNodes(UObject* De
 		if(StructProperty->Struct->IsChildOf(FAnimNode_Base::StaticStruct()))
 		{
 			FAnimNode_Base* Node = StructProperty->ContainerPtrToValuePtr<FAnimNode_Base>(DefaultObject);
-			Node->SetNodeData(GetNodeData()[AnimNodeIndex]);
+
+			// This function is called on child->parent hierarchies in postload. We dont want to overwrite child data
+			// with parent data so skip if the node data has already been set up by a previous call.
+			if(Node->NodeData == nullptr)
+			{
+				Node->SetNodeData(AnimNodeData[AnimNodeIndex]);
+			}
 
 			if(Node->NeedsDynamicReset())
 			{
diff --git a/Engine/Source/Runtime/Engine/Private/BlueprintGeneratedClass.cpp b/Engine/Source/Runtime/Engine/Private/BlueprintGeneratedClass.cpp
--- a/Engine/Source/Runtime/Engine/Private/BlueprintGeneratedClass.cpp
+++ b/Engine/Source/Runtime/Engine/Private/BlueprintGeneratedClass.cpp
@@ -415,16 +415,34 @@ void UBlueprintGeneratedClass::SerializeDefaultObject(UObject* Object, FStructur
 
 	UnderlyingArchive.UsingCustomVersion(FUE5MainStreamObjectVersion::GUID);
 
+	bool bSkipSparseClassDataSerialization = false;
 	if(UnderlyingArchive.CustomVer(FUE5MainStreamObjectVersion::GUID) >= FUE5MainStreamObjectVersion::SparseClassDataStructSerialization)
 	{
 		UnderlyingArchive << SparseClassDataStruct;
+
+		if(UnderlyingArchive.IsLoading() && SparseClassDataStruct == nullptr && Object->GetLinker())
+		{
+			// Missing or failed to load sparse class data struct - possible that the parent class was deleted, or regenerated on load
+			// so seek past where this CDO was serialized as we cannot read the serialized sparse class data any more.
+			// Note this happens in majority of cases with classes that have no sparse class data attached too. In that case 
+			// this 'skip' over the remaining part of the archive should have no ill effects as the seek should effectively be zero
+			bSkipSparseClassDataSerialization = true;
+		}
 	}
 
 #if WITH_EDITORONLY_DATA
 	if (bIsSparseClassDataSerializable)
 #endif
 	{
-		if (Object->GetSparseClassDataStruct())
+		if(bSkipSparseClassDataSerialization)
+		{
+			// Seek to after sparse class data rather than load it.
+			const FLinkerLoad* Linker = Object->GetLinker();
+			const int32 LinkerIndex = Object->GetLinkerIndex();
+			const FObjectExport& Export = Linker->ExportMap[LinkerIndex];
+			UnderlyingArchive.Seek(Export.SerialOffset + Export.SerialSize);
+		}
+		else if (Object->GetSparseClassDataStruct())
 		{
 			SerializeSparseClassData(FStructuredArchiveFromArchive(UnderlyingArchive).GetSlot());
 		}
