Fix internal blendspaces not being saved/copied in certain circumstances

As the blendspace was duplicated from a previously duplicated node, it was marked as transient and hence was not saved along with the class. This is not a problem in the editor as the blend space is always regenerated on load, but in cooked games the blendspace was never saved in the cook and thus not present at runtime.

diff --git a/Engine/Source/Editor/AnimGraph/Private/AnimBlueprintExtension_BlendSpaceGraph.cpp b/Engine/Source/Editor/AnimGraph/Private/AnimBlueprintExtension_BlendSpaceGraph.cpp
--- a/Engine/Source/Editor/AnimGraph/Private/AnimBlueprintExtension_BlendSpaceGraph.cpp
+++ b/Engine/Source/Editor/AnimGraph/Private/AnimBlueprintExtension_BlendSpaceGraph.cpp
@@ -12,6 +12,7 @@ void UAnimBlueprintExtension_BlendSpaceGraph::HandleStartCompilingClass(const UC
 UBlendSpace* UAnimBlueprintExtension_BlendSpaceGraph::AddBlendSpace(UBlendSpace* InSourceBlendSpace)
 {
 	UBlendSpace* CopiedBlendSpace = DuplicateObject(InSourceBlendSpace, Class);
+	CopiedBlendSpace->ClearFlags(RF_Transient);
 	Subsystem.BlendSpaces.Add(CopiedBlendSpace);
 	return CopiedBlendSpace;
 }
\ No newline at end of file
diff --git a/Engine/Source/Editor/AnimGraph/Private/AnimGraphNode_BlendSpaceGraphBase.cpp b/Engine/Source/Editor/AnimGraph/Private/AnimGraphNode_BlendSpaceGraphBase.cpp
--- a/Engine/Source/Editor/AnimGraph/Private/AnimGraphNode_BlendSpaceGraphBase.cpp
+++ b/Engine/Source/Editor/AnimGraph/Private/AnimGraphNode_BlendSpaceGraphBase.cpp
@@ -142,7 +142,8 @@ void UAnimGraphNode_BlendSpaceGraphBase::SetupFromAsset(UBlendSpace* InBlendSpac
 
 		BlendSpaceGraph = CastChecked<UBlendSpaceGraph>(FBlueprintEditorUtils::CreateNewGraph(this, InBlendSpace->GetFName(), UBlendSpaceGraph::StaticClass(), UEdGraphSchema::StaticClass()));
 		BlendSpaceGraph->BlendSpace = BlendSpace = DuplicateObject(InBlendSpace, BlendSpaceGraph);
-		BlendSpace->ClearFlags(RF_Public | RF_Standalone);
+		BlendSpaceGraph->BlendSpace->ClearFlags(RF_Public | RF_Standalone | RF_Transient);
+		BlendSpaceGraph->BlendSpace->SetFlags(RF_Transactional);
 		BlendSpaceGraph->BlendSpace->ResetSkeleton(nullptr);
 		BlendSpaceGraph->BlendSpace->EmptyMetaData();
 		BlendSpaceGraph->BlendSpace->RemoveUserDataOfClass(UAssetUserData::StaticClass());
@@ -182,6 +183,7 @@ void UAnimGraphNode_BlendSpaceGraphBase::SetupFromClass(TSubclassOf<UBlendSpace>
 	{
 		BlendSpaceGraph = CastChecked<UBlendSpaceGraph>(FBlueprintEditorUtils::CreateNewGraph(this, InBlendSpaceClass.Get()->GetFName(), UBlendSpaceGraph::StaticClass(), UEdGraphSchema::StaticClass()));
 		BlendSpaceGraph->BlendSpace = BlendSpace = NewObject<UBlendSpace>(BlendSpaceGraph, InBlendSpaceClass.Get());
+		BlendSpaceGraph->BlendSpace->ClearFlags(RF_Public | RF_Standalone | RF_Transient);
 		BlendSpaceGraph->BlendSpace->SetFlags(RF_Transactional);
 		BlendSpaceGraph->bAllowDeletion = false;
 		BlendSpaceGraph->bAllowRenaming = true;
