Fixed crash deleting, undoing then compiling a blend space graph

diff --git a/Engine/Source/Editor/AnimGraph/Private/AnimGraphNode_Base.cpp b/Engine/Source/Editor/AnimGraph/Private/AnimGraphNode_Base.cpp
--- a/Engine/Source/Editor/AnimGraph/Private/AnimGraphNode_Base.cpp
+++ b/Engine/Source/Editor/AnimGraph/Private/AnimGraphNode_Base.cpp
@@ -54,6 +54,17 @@ void UAnimGraphNode_Base::PreEditChange(FProperty* PropertyThatWillChange)
 	}
 }
 
+void UAnimGraphNode_Base::PostEditUndo()
+{
+	Super::PostEditUndo();
+
+	if(HasValidBlueprint())
+	{
+		// Node may have been removed or added by undo/redo so make sure extensions are refreshed
+		GetAnimBlueprint()->RequestRefreshExtensions();
+	}
+}
+
 void UAnimGraphNode_Base::PostEditChangeProperty(FPropertyChangedEvent& PropertyChangedEvent)
 {
 	FName PropertyName = (PropertyChangedEvent.Property != NULL) ? PropertyChangedEvent.Property->GetFName() : NAME_None;
diff --git a/Engine/Source/Editor/AnimGraph/Public/AnimGraphNode_Base.h b/Engine/Source/Editor/AnimGraph/Public/AnimGraphNode_Base.h
--- a/Engine/Source/Editor/AnimGraph/Public/AnimGraphNode_Base.h
+++ b/Engine/Source/Editor/AnimGraph/Public/AnimGraphNode_Base.h
@@ -193,6 +193,7 @@ class ANIMGRAPH_API UAnimGraphNode_Base : public UK2Node
 	// UObject interface
 	virtual void PostEditChangeProperty(FPropertyChangedEvent& PropertyChangedEvent) override;
 	virtual void PreEditChange(FProperty* PropertyAboutToChange) override;
+	virtual void PostEditUndo() override;
 	virtual void Serialize(FArchive& Ar) override;
 	// End of UObject interface
 
