Animation: thread-safe execution and property access improvements

"Call function" anim node:
- Adds the ability to call functions at different points in the anim graph execution (and under different conditions, e.g. when a branch has started blending in). Only thread-safe functions are allowed to be called.
- Adds a thread-safe override point BlueprintThreadSafeUpdateAnimation, called on worker threads for the main instance and for linked instances when they are relevant in the graph (only one call per frame for linked layer instances).

Subsystems:
- Added new override points pre/post event graph(s) (moved override point for worker thread work to around the thread safe update function call).

Improves property access integration:
- Property access now shows (and allows the user to override) the call site of accesses. This is to allow users to see when their property access calls will be made, hopefully making its use less confusing for power users.
- Tweaked UX for property access nodes and dropdowns.
- Anim node pins now have property access bindings in-line on the pin.

Also adds the abilility for the anim graph to opt-in (via a config flag) to more stringent thread safety checks. Disabled by default for now as this requires content fixup.

diff --git a/Engine/Config/BaseEngine.ini b/Engine/Config/BaseEngine.ini
--- a/Engine/Config/BaseEngine.ini
+++ b/Engine/Config/BaseEngine.ini
@@ -1582,6 +1582,7 @@ bCanSuppressAccessViolation=false
 bLoadNativeConvertedBPClassInEditor=false
 bIgnoreCompileOnLoadErrorsOnBuildMachine=false
 bAllowNativeComponentClassOverrides=true
+bUseLegacyAnimBlueprintThreadSafetyChecks=true
 
 [/Script/Engine.Blueprint]
 bRecompileOnLoad=true
diff --git a/Engine/Plugins/Developer/PropertyAccessNode/Source/PropertyAccessNode/Private/K2Node_PropertyAccess.cpp b/Engine/Plugins/Developer/PropertyAccessNode/Source/PropertyAccessNode/Private/K2Node_PropertyAccess.cpp
--- a/Engine/Plugins/Developer/PropertyAccessNode/Source/PropertyAccessNode/Private/K2Node_PropertyAccess.cpp
+++ b/Engine/Plugins/Developer/PropertyAccessNode/Source/PropertyAccessNode/Private/K2Node_PropertyAccess.cpp
@@ -9,18 +9,16 @@
 #include "BlueprintActionDatabaseRegistrar.h"
 #include "EditorCategoryUtils.h"
 #include "BlueprintNodeSpawner.h"
-#include "KismetCompilerMisc.h"
-#include "EdGraphUtilities.h"
 #include "KismetCompiler.h"
 #include "K2Node_VariableGet.h"
 #include "AnimationGraph.h"
 #include "Kismet2/BlueprintEditorUtils.h"
 #include "AnimBlueprintExtension_PropertyAccess.h"
-#include "Modules/ModuleManager.h"
 #include "IPropertyAccessEditor.h"
 #include "IPropertyAccessCompiler.h"
 #include "Features/IModularFeatures.h"
 #include "IAnimBlueprintCompilationContext.h"
+#include "IPropertyAccessBlueprintBinding.h"
 #include "Animation/AnimBlueprint.h"
 
 #define LOCTEXT_NAMESPACE "K2Node_PropertyAccess"
@@ -29,9 +27,11 @@ void UK2Node_PropertyAccess::CreateClassVariablesFromBlueprint(IAnimBlueprintVar
 {
 	GeneratedPropertyName = NAME_None;
 
-	if(ResolvedPinType != FEdGraphPinType() && ResolvedPinType.PinCategory != UEdGraphSchema_K2::PC_Wildcard)
+	const bool bRequiresCachedVariable = !bWasResolvedThreadSafe || UAnimBlueprintExtension_PropertyAccess::ContextRequiresCachedVariable(ContextId);
+	
+	if(ResolvedPinType != FEdGraphPinType() && ResolvedPinType.PinCategory != UEdGraphSchema_K2::PC_Wildcard && bRequiresCachedVariable)
 	{
-		// Create internal generated destination property
+		// Create internal generated destination property (only if we were not identified as thread safe)
 		if(FProperty* DestProperty = InCreationContext.CreateUniqueVariable(this, ResolvedPinType))
 		{
 			GeneratedPropertyName = DestProperty->GetFName();
@@ -41,24 +41,42 @@ void UK2Node_PropertyAccess::CreateClassVariablesFromBlueprint(IAnimBlueprintVar
 
 void UK2Node_PropertyAccess::ExpandNode(FKismetCompilerContext& InCompilerContext, UEdGraph* InSourceGraph)
 {
-	ResolveLeafProperty();
+	ResolvePropertyAccess();
+
+	UK2Node_PropertyAccess* OriginalNode = CastChecked<UK2Node_PropertyAccess>(InCompilerContext.MessageLog.FindSourceObject(this));
+	OriginalNode->CompiledContext = FText();
+	OriginalNode->CompiledContextDesc = FText();
 
+	TUniquePtr<IAnimBlueprintCompilationContext> CompilationContext = IAnimBlueprintCompilationContext::Get(InCompilerContext);
+	UAnimBlueprintExtension_PropertyAccess* PropertyAccessExtension = UAnimBlueprintExtension::GetExtension<UAnimBlueprintExtension_PropertyAccess>(CastChecked<UAnimBlueprint>(GetBlueprint()));
+	check(PropertyAccessExtension);
+	
 	if(GeneratedPropertyName != NAME_None)
 	{
 		TArray<FString> DestPropertyPath;
+
+		// We are using an intermediate object-level property (as we need to call this access and cache its result
+		// until later) 
 		DestPropertyPath.Add(GeneratedPropertyName.ToString());
 
 		// Create a copy event in the complied generated class
-		TUniquePtr<IAnimBlueprintCompilationContext> CompilationContext = IAnimBlueprintCompilationContext::Get(InCompilerContext);
-		UAnimBlueprintExtension_PropertyAccess* PropertyAccessExtension = UAnimBlueprintExtension::GetExtension<UAnimBlueprintExtension_PropertyAccess>(CastChecked<UAnimBlueprint>(GetBlueprint()));
-		check(PropertyAccessExtension);
-		PropertyAccessExtension->AddCopy(Path, DestPropertyPath, EPropertyAccessBatchType::Batched, this);
-
+		FPropertyAccessHandle Handle = PropertyAccessExtension->AddCopy(Path, DestPropertyPath, ContextId, this);
+		
+		PostLibraryCompiledHandle = PropertyAccessExtension->OnPostLibraryCompiled().AddLambda([this, OriginalNode, PropertyAccessExtension, Handle](IAnimBlueprintCompilationBracketContext& /*InCompilationContext*/, IAnimBlueprintGeneratedClassCompiledData& /*OutCompiledData*/)
+		{
+			const FCompiledPropertyAccessHandle CompiledHandle = PropertyAccessExtension->GetCompiledHandle(Handle);
+			if(CompiledHandle.IsValid())
+			{
+				OriginalNode->CompiledContext = UAnimBlueprintExtension_PropertyAccess::GetCompiledHandleContext(CompiledHandle);
+				OriginalNode->CompiledContextDesc = UAnimBlueprintExtension_PropertyAccess::GetCompiledHandleContextDesc(CompiledHandle);
+			}
+			PropertyAccessExtension->OnPostLibraryCompiled().Remove(PostLibraryCompiledHandle);
+		});
+		
 		// Replace us with a get node
 		UK2Node_VariableGet* VariableGetNode = InCompilerContext.SpawnIntermediateNode<UK2Node_VariableGet>(this, InSourceGraph);
 		VariableGetNode->VariableReference.SetSelfMember(GeneratedPropertyName);
 		VariableGetNode->AllocateDefaultPins();
-		InCompilerContext.MessageLog.NotifyIntermediateObjectCreation(VariableGetNode, this);
 
 		// Move pin links from Get node we are expanding, to the new pure one we've created
 		UEdGraphPin* VariableValuePin = VariableGetNode->GetValuePin();
@@ -67,11 +85,20 @@ void UK2Node_PropertyAccess::ExpandNode(FKismetCompilerContext& InCompilerContex
 	}
 	else
 	{
-		InCompilerContext.MessageLog.Error(*LOCTEXT("IntermediateProperty_Error", "Intermediate property could not be created on @@").ToString(), this);
+		const bool bRequiresCachedVariable = !bWasResolvedThreadSafe || UAnimBlueprintExtension_PropertyAccess::ContextRequiresCachedVariable(ContextId);
+		check(!bRequiresCachedVariable);
+
+		UEnum* EnumClass = StaticEnum<EAnimPropertyAccessCallSite>();
+		check(EnumClass != nullptr);
+
+		OriginalNode->CompiledContext = EnumClass->GetDisplayNameTextByValue((int32)EAnimPropertyAccessCallSite::WorkerThread_Unbatched);
+		OriginalNode->CompiledContextDesc = EnumClass->GetToolTipTextByIndex((int32)EAnimPropertyAccessCallSite::WorkerThread_Unbatched);
+		
+		PropertyAccessExtension->ExpandPropertyAccess(InCompilerContext, Path, InSourceGraph, GetOutputPin());
 	}
 }
 
-void UK2Node_PropertyAccess::ResolveLeafProperty()
+void UK2Node_PropertyAccess::ResolvePropertyAccess()
 {
 	if(UBlueprint* Blueprint = GetBlueprint())
 	{
@@ -81,7 +108,8 @@ void UK2Node_PropertyAccess::ResolveLeafProperty()
 		if(IModularFeatures::Get().IsModularFeatureAvailable("PropertyAccessEditor"))
 		{
 			IPropertyAccessEditor& PropertyAccessEditor = IModularFeatures::Get().GetModularFeature<IPropertyAccessEditor>("PropertyAccessEditor");
-			PropertyAccessEditor.ResolveLeafProperty(Blueprint->SkeletonGeneratedClass, Path, PropertyToResolve, ResolvedArrayIndex);
+			FPropertyAccessResolveResult Result = PropertyAccessEditor.ResolvePropertyAccess(Blueprint->SkeletonGeneratedClass, Path, PropertyToResolve, ResolvedArrayIndex);
+			bWasResolvedThreadSafe = Result.bIsThreadSafe;
 			ResolvedProperty = PropertyToResolve;
 		}
 	}
@@ -92,29 +120,25 @@ void UK2Node_PropertyAccess::ResolveLeafProperty()
 	}
 }
 
-static FText MakeTextPath(const TArray<FString>& InPath)
-{
-	return FText::FromString(Algo::Accumulate(InPath, FString(), [](const FString& InResult, const FString& InSegment)
-		{ 
-			return InResult.IsEmpty() ? InSegment : (InResult + TEXT(".") + InSegment);
-		}));	
-}
-
 void UK2Node_PropertyAccess::SetPath(const TArray<FString>& InPath)
 {
 	Path = InPath;
-	TextPath = MakeTextPath(Path);
+
+	IPropertyAccessEditor& PropertyAccessEditor = IModularFeatures::Get().GetModularFeature<IPropertyAccessEditor>("PropertyAccessEditor");
+	TextPath = PropertyAccessEditor.MakeTextPath(Path);
 	FBlueprintEditorUtils::MarkBlueprintAsStructurallyModified(GetBlueprint());
-	ResolveLeafProperty();
+	ResolvePropertyAccess();
 	ReconstructNode();
 }
 
 void UK2Node_PropertyAccess::SetPath(TArray<FString>&& InPath)
 {
 	Path = MoveTemp(InPath);
-	TextPath = MakeTextPath(Path);
+	
+	IPropertyAccessEditor& PropertyAccessEditor = IModularFeatures::Get().GetModularFeature<IPropertyAccessEditor>("PropertyAccessEditor");
+	TextPath = PropertyAccessEditor.MakeTextPath(Path);
 	FBlueprintEditorUtils::MarkBlueprintAsStructurallyModified(GetBlueprint());
-	ResolveLeafProperty();
+	ResolvePropertyAccess();
 	ReconstructNode();
 }
 
@@ -131,7 +155,7 @@ void UK2Node_PropertyAccess::ClearPath()
 void UK2Node_PropertyAccess::AllocatePins(UEdGraphPin* InOldOutputPin)
 {
 	// Resolve leaf to try to get a valid property type for an output pin
-	ResolveLeafProperty();
+	ResolvePropertyAccess();
 
 	if(UBlueprint* Blueprint = GetBlueprint())
 	{
@@ -194,27 +218,17 @@ void UK2Node_PropertyAccess::ReallocatePinsDuringReconstruction(TArray<UEdGraphP
 
 FText UK2Node_PropertyAccess::GetNodeTitle(ENodeTitleType::Type TitleType) const
 {
-	return LOCTEXT("PropertyAccess", "Property Access");
+	return !TextPath.IsEmpty() ? TextPath : LOCTEXT("PropertyAccess", "Property Access");
 }
 
 void UK2Node_PropertyAccess::AddSearchMetaDataInfo(TArray<FSearchTagDataPair>& OutTaggedMetaData) const
 {
-	auto MakeTextPath = [this]()
-	{
-		if(Path.Num())
-		{
-			return FText::FromString(Algo::Accumulate(Path, FString(), [](const FString& InResult, const FString& InSegment)
-				{ 
-					return InResult.IsEmpty() ? InSegment : (InResult + TEXT(".") + InSegment);
-				}));	
-		}
-		else
-		{
-			return LOCTEXT("None", "None");
-		}
-	};
+	Super::AddSearchMetaDataInfo(OutTaggedMetaData);
 
-	OutTaggedMetaData.Emplace(LOCTEXT("PropertyAccess", "Property Access"), MakeTextPath());
+	if(!TextPath.IsEmpty())
+	{
+		OutTaggedMetaData.Emplace(LOCTEXT("PropertyAccess", "Property Access"), TextPath);
+	}
 }
 
 void UK2Node_PropertyAccess::PinConnectionListChanged(UEdGraphPin* Pin)
@@ -243,9 +257,32 @@ FText UK2Node_PropertyAccess::GetMenuCategory() const
 
 bool UK2Node_PropertyAccess::IsCompatibleWithGraph(UEdGraph const* TargetGraph) const
 {
-	// Only allow placement in anim graphs, for now. If this changes then we need to address the 
-	// dependency on the anim BP compiler's subsystems
-	return TargetGraph->IsA<UAnimationGraph>();
+	IPropertyAccessBlueprintBinding::FContext BindingContext;
+	BindingContext.Blueprint = FBlueprintEditorUtils::FindBlueprintForGraph(TargetGraph);
+	BindingContext.Graph = TargetGraph;
+	BindingContext.Node = this;
+	BindingContext.Pin = FindPin(TEXT("Value"));
+	
+	// Check any property access blueprint bindings we might have registered
+	for(IPropertyAccessBlueprintBinding* Binding : IModularFeatures::Get().GetModularFeatureImplementations<IPropertyAccessBlueprintBinding>("PropertyAccessBlueprintBinding"))
+	{
+		if(Binding->CanBindToContext(BindingContext))
+		{
+			return true;
+		}
+	}
+
+	return false;
+}
+
+FText UK2Node_PropertyAccess::GetTooltipText() const
+{
+	return LOCTEXT("PropertyAccessTooltip", "Accesses properties according to property path");
+}
+
+bool UK2Node_PropertyAccess::HasResolvedPinType() const
+{
+	return ResolvedPinType != FEdGraphPinType() && ResolvedPinType.PinCategory != UEdGraphSchema_K2::PC_Wildcard;
 }
 
 #undef LOCTEXT_NAMESPACE
\ No newline at end of file
diff --git a/Engine/Plugins/Developer/PropertyAccessNode/Source/PropertyAccessNode/Private/K2Node_PropertyAccess.h b/Engine/Plugins/Developer/PropertyAccessNode/Source/PropertyAccessNode/Private/K2Node_PropertyAccess.h
--- a/Engine/Plugins/Developer/PropertyAccessNode/Source/PropertyAccessNode/Private/K2Node_PropertyAccess.h
+++ b/Engine/Plugins/Developer/PropertyAccessNode/Source/PropertyAccessNode/Private/K2Node_PropertyAccess.h
@@ -37,9 +37,24 @@ public:
 	/** Get the resolved pin type, if any. */
 	const FEdGraphPinType& GetResolvedPinType() const { return ResolvedPinType; }
 
+	/** Check whether we have a resolved pin type */
+	bool HasResolvedPinType() const;
+	
 	/** Get the output pin */
 	UEdGraphPin* GetOutputPin() const { return FindPinChecked(TEXT("Value"), EGPD_Output); }
 
+	/** Get the context Id */
+	FName GetContextId() const { return ContextId; }
+
+	/** Set the context Id */
+	void SetContextId(const FName& InContextId) { ContextId = InContextId; }
+
+	/** Get the the context that this access was assigned during compilation. This can be empty. */
+	const FText& GetCompiledContext() const { return CompiledContext; }
+	
+	/** Get the description of the context that this access was assigned during compilation. This can be empty. */
+	const FText& GetCompiledContextDesc() const { return CompiledContextDesc; }
+
 private:
 	/** Path that this access exposes */
 	UPROPERTY()
@@ -57,12 +72,28 @@ private:
 	UPROPERTY()
 	FName GeneratedPropertyName = NAME_None;
 
+	/** Property access context (set by the user) that is intended to be used */
+	UPROPERTY()
+	FName ContextId;
+	
 	/** Resolved leaf property for the path, NULL if path cant be resolved or is empty. */
 	TFieldPath<FProperty> ResolvedProperty;
 
 	/** Resolved array index, if the leaf property is an array. If this is INDEX_NONE then the property refers to the entire array */
 	int32 ResolvedArrayIndex = INDEX_NONE;
 
+	/** Whether the path was resolved as thread safe the last time the ResolvedProperty was updated */
+	bool bWasResolvedThreadSafe = false;
+
+	/** Handle used to hook into property access compilation machinery */
+	FDelegateHandle PostLibraryCompiledHandle;
+
+	/** The context that this access was assigned during compilation */
+	FText CompiledContext;
+
+	/** Description of the context that this access was assigned during compilation (used for tooltips) */
+	FText CompiledContextDesc;
+	
 	// IClassVariableCreator interface
 	virtual void CreateClassVariablesFromBlueprint(IAnimBlueprintVariableCreationContext& InCreationContext) override;
 
@@ -79,10 +110,12 @@ private:
 	virtual bool IsNodePure() const override { return true; }
 	virtual void GetMenuActions(FBlueprintActionDatabaseRegistrar& ActionRegistrar) const override;
 	virtual FText GetMenuCategory() const override;
-
+	virtual bool DrawNodeAsVariable() const override { return true; }
+	virtual FText GetTooltipText() const override;
+	
 	/** Helper function for pin allocation */
 	void AllocatePins(UEdGraphPin* InOldOutputPin = nullptr);
 
 	/** Attempt to resolve the path to a leaf property */
-	void ResolveLeafProperty();
+	void ResolvePropertyAccess();
 };
\ No newline at end of file
diff --git a/Engine/Plugins/Developer/PropertyAccessNode/Source/PropertyAccessNode/Private/PropertyAccessNodeLibrary.cpp b/Engine/Plugins/Developer/PropertyAccessNode/Source/PropertyAccessNode/Private/PropertyAccessNodeLibrary.cpp
new file mode 100644
--- /dev/null
+++ b/Engine/Plugins/Developer/PropertyAccessNode/Source/PropertyAccessNode/Private/PropertyAccessNodeLibrary.cpp
@@ -0,0 +1,3 @@
+// Copyright Epic Games, Inc. All Rights Reserved.
+
+#include "PropertyAccessNodeLibrary.h"
diff --git a/Engine/Plugins/Developer/PropertyAccessNode/Source/PropertyAccessNode/Private/PropertyAccessNodeLibrary.h b/Engine/Plugins/Developer/PropertyAccessNode/Source/PropertyAccessNode/Private/PropertyAccessNodeLibrary.h
new file mode 100644
--- /dev/null
+++ b/Engine/Plugins/Developer/PropertyAccessNode/Source/PropertyAccessNode/Private/PropertyAccessNodeLibrary.h
@@ -0,0 +1,3 @@
+// Copyright Epic Games, Inc. All Rights Reserved.
+
+#include "CoreMinimal.h"
diff --git a/Engine/Plugins/Developer/PropertyAccessNode/Source/PropertyAccessNode/Private/SPropertyAccessNode.cpp b/Engine/Plugins/Developer/PropertyAccessNode/Source/PropertyAccessNode/Private/SPropertyAccessNode.cpp
--- a/Engine/Plugins/Developer/PropertyAccessNode/Source/PropertyAccessNode/Private/SPropertyAccessNode.cpp
+++ b/Engine/Plugins/Developer/PropertyAccessNode/Source/PropertyAccessNode/Private/SPropertyAccessNode.cpp
@@ -8,9 +8,16 @@
 #include "IPropertyAccessEditor.h"
 #include "Widgets/Layout/SSpacer.h"
 #include "K2Node_PropertyAccess.h"
+#include "SGraphPin.h"
 #include "Kismet2/BlueprintEditorUtils.h"
 #include "Features/IModularFeatures.h"
 #include "Widgets/SBoxPanel.h"
+#include "Widgets/Layout/SWrapBox.h"
+#include "Features/IModularFeatures.h"
+#include "IPropertyAccessBlueprintBinding.h"
+#include "Framework/MultiBox/MultiBoxExtender.h"
+#include "Fonts/FontMeasure.h"
+#include "Framework/Application/SlateApplication.h"
 
 #define LOCTEXT_NAMESPACE "SPropertyAccessNode"
 
@@ -71,8 +78,11 @@ bool SPropertyAccessNode::CanBindProperty(FProperty* InProperty) const
 	return false;
 }
 
-void SPropertyAccessNode::CreateBelowPinControls(TSharedPtr<SVerticalBox> InMainBox)
+TSharedPtr<SGraphPin> SPropertyAccessNode::CreatePinWidget(UEdGraphPin* Pin) const
 {
+	TSharedPtr<SGraphPin> DefaultWidget = SGraphNodeK2Var::CreatePinWidget(Pin);
+	DefaultWidget->SetShowLabel(false);
+	
 	UK2Node_PropertyAccess* K2Node_PropertyAccess = CastChecked<UK2Node_PropertyAccess>(GraphNode);
 
 	FPropertyBindingWidgetArgs Args;
@@ -104,12 +114,15 @@ void SPropertyAccessNode::CreateBelowPinControls(TSharedPtr<SVerticalBox> InMain
 			TArray<FString> StringPath;
 			PropertyAccessEditor.MakeStringPath(InBindingChain, StringPath);
 			K2Node_PropertyAccess->SetPath(MoveTemp(StringPath));
+
+			FBlueprintEditorUtils::MarkBlueprintAsModified(K2Node_PropertyAccess->GetBlueprint());
 		}
 	});
 
 	Args.OnRemoveBinding = FOnRemoveBinding::CreateLambda([K2Node_PropertyAccess](FName InPropertyName)
 	{
-		return K2Node_PropertyAccess->ClearPath();
+		K2Node_PropertyAccess->ClearPath();
+		FBlueprintEditorUtils::MarkBlueprintAsModified(K2Node_PropertyAccess->GetBlueprint());
 	});
 
 	Args.OnCanRemoveBinding = FOnCanRemoveBinding::CreateLambda([K2Node_PropertyAccess](FName InPropertyName)
@@ -123,6 +136,28 @@ void SPropertyAccessNode::CreateBelowPinControls(TSharedPtr<SVerticalBox> InMain
 		return TextPath.IsEmpty() ? LOCTEXT("Bind", "Bind") : TextPath;
 	});
 
+	Args.CurrentBindingToolTipText = MakeAttributeLambda([K2Node_PropertyAccess]()
+	{
+		const FText& TextPath = K2Node_PropertyAccess->GetTextPath();
+		if(TextPath.IsEmpty())
+		{
+			return LOCTEXT("EmptyBinding", "Access is not bound to anything");
+		}
+		else
+		{
+			const FText& CompilationContext = K2Node_PropertyAccess->GetCompiledContext();
+			const FText& CompilationContextDesc = K2Node_PropertyAccess->GetCompiledContextDesc();
+			if(CompilationContext.IsEmpty() && CompilationContextDesc.IsEmpty())
+			{
+				return TextPath;
+			}
+			else
+			{
+				return FText::Format(LOCTEXT("CompiledAccessToolTipFormat", "Access property '{0}'\n{1}\n{2}"), TextPath, CompilationContext, CompilationContextDesc);
+			}
+		}
+	});
+	
 	Args.CurrentBindingImage = MakeAttributeLambda([K2Node_PropertyAccess]()
 	{
 		if(const FProperty* Property = K2Node_PropertyAccess->GetResolvedProperty())
@@ -163,9 +198,47 @@ void SPropertyAccessNode::CreateBelowPinControls(TSharedPtr<SVerticalBox> InMain
 		}
 	});
 
+	IPropertyAccessBlueprintBinding::FContext BindingContext;
+	BindingContext.Blueprint = FBlueprintEditorUtils::FindBlueprintForNode(K2Node_PropertyAccess);
+	BindingContext.Graph = K2Node_PropertyAccess->GetGraph();
+	BindingContext.Node = K2Node_PropertyAccess;
+	BindingContext.Pin = K2Node_PropertyAccess->FindPin(TEXT("Value"));
+
+	IPropertyAccessBlueprintBinding::FBindingMenuArgs MenuArgs;
+	MenuArgs.OnSetPropertyAccessContextId = FOnSetPropertyAccessContextId::CreateLambda([K2Node_PropertyAccess](const FName& InContextId)
+	{
+		K2Node_PropertyAccess->SetContextId(InContextId);
+		FBlueprintEditorUtils::MarkBlueprintAsModified(K2Node_PropertyAccess->GetBlueprint());
+	});
+	MenuArgs.OnCanSetPropertyAccessContextId = FOnCanSetPropertyAccessContextId::CreateLambda([K2Node_PropertyAccess](const FName& InContextId)
+	{
+		return K2Node_PropertyAccess->HasResolvedPinType();
+	});	
+	MenuArgs.OnGetPropertyAccessContextId = FOnGetPropertyAccessContextId::CreateLambda([K2Node_PropertyAccess]()
+	{
+		return K2Node_PropertyAccess->GetContextId();
+	});
+	
+	// Check any property access blueprint bindings we might have registered for menu extenders
+	TArray<TSharedPtr<FExtender>> Extenders;
+	for(IPropertyAccessBlueprintBinding* Binding : IModularFeatures::Get().GetModularFeatureImplementations<IPropertyAccessBlueprintBinding>("PropertyAccessBlueprintBinding"))
+	{
+		TSharedPtr<FExtender> BindingExtender = Binding->MakeBindingMenuExtender(BindingContext, MenuArgs);
+		if(BindingExtender)
+		{
+			Extenders.Add(BindingExtender);
+		}
+	}
+	
+	if(Extenders.Num() > 0)
+	{
+		Args.MenuExtender = FExtender::Combine(Extenders);
+	}
+
 	Args.bAllowArrayElementBindings = true;
 	Args.bAllowNewBindings = false;
 	Args.bAllowUObjectFunctions = true;
+	Args.bAllowStructFunctions = true;
 
 	TSharedPtr<SWidget> PropertyBindingWidget;
 
@@ -179,8 +252,11 @@ void SPropertyAccessNode::CreateBelowPinControls(TSharedPtr<SVerticalBox> InMain
 		PropertyBindingWidget = SNullWidget::NullWidget;
 	}
 
-	InMainBox->AddSlot()
-	.AutoHeight()
+	TSharedRef<SWrapBox> LabelAndValue = DefaultWidget->GetLabelAndValue();
+
+	const FTextBlockStyle& TextBlockStyle = FEditorStyle::Get().GetWidgetStyle<FTextBlockStyle>("PropertyAccess.CompiledContext.Text");
+	
+	LabelAndValue->AddSlot()
 	.Padding(5.0f)
 	[
 		SNew(SLevelOfDetailBranchNode)
@@ -191,9 +267,41 @@ void SPropertyAccessNode::CreateBelowPinControls(TSharedPtr<SVerticalBox> InMain
 		]
 		.HighDetail()
 		[
-			PropertyBindingWidget.ToSharedRef()
+			SNew(SOverlay)
+			+SOverlay::Slot()
+			[
+				PropertyBindingWidget.ToSharedRef()
+			]
+			+SOverlay::Slot()
+			.HAlign(HAlign_Right)
+			.VAlign(VAlign_Bottom)
+			[
+				SNew(SBorder)
+				.Padding(1.0f)
+				.BorderImage(FEditorStyle::GetBrush("PropertyAccess.CompiledContext.Border"))
+				.RenderTransform_Lambda([K2Node_PropertyAccess, &TextBlockStyle]()
+				{
+					const TSharedRef<FSlateFontMeasure> FontMeasureService = FSlateApplication::Get().GetRenderer()->GetFontMeasureService();
+					FVector2D TextSize = FontMeasureService->Measure(K2Node_PropertyAccess->GetCompiledContext(), TextBlockStyle.Font);
+					return FSlateRenderTransform(FVector2D(0.0f, TextSize.Y));
+				})	
+				[
+					SNew(STextBlock)
+					.TextStyle(&TextBlockStyle)
+					.Visibility_Lambda([K2Node_PropertyAccess]()
+					{
+						return K2Node_PropertyAccess->GetCompiledContext().IsEmpty() ? EVisibility::Collapsed : EVisibility::Visible;
+					})
+					.Text_Lambda([K2Node_PropertyAccess]()
+					{
+						return K2Node_PropertyAccess->GetCompiledContext();
+					})
+				]
+			]
 		]
 	];
+	
+	return DefaultWidget;
 }
 
 #undef LOCTEXT_NAMESPACE
\ No newline at end of file
diff --git a/Engine/Plugins/Developer/PropertyAccessNode/Source/PropertyAccessNode/Private/SPropertyAccessNode.h b/Engine/Plugins/Developer/PropertyAccessNode/Source/PropertyAccessNode/Private/SPropertyAccessNode.h
--- a/Engine/Plugins/Developer/PropertyAccessNode/Source/PropertyAccessNode/Private/SPropertyAccessNode.h
+++ b/Engine/Plugins/Developer/PropertyAccessNode/Source/PropertyAccessNode/Private/SPropertyAccessNode.h
@@ -3,11 +3,11 @@
 #pragma once
 
 #include "CoreMinimal.h"
-#include "KismetNodes/SGraphNodeK2Base.h"
+#include "KismetNodes/SGraphNodeK2Var.h"
 
 class UK2Node_PropertyAccess;
 
-class SPropertyAccessNode : public SGraphNodeK2Base
+class SPropertyAccessNode : public SGraphNodeK2Var
 {
 public:
 	SLATE_BEGIN_ARGS(SPropertyAccessNode) {}
@@ -17,7 +17,7 @@ public:
 	void Construct(const FArguments& InArgs, UK2Node_PropertyAccess* InNode);
 
 	// SGraphNode interface
-	virtual void CreateBelowPinControls(TSharedPtr<SVerticalBox> InMainBox) override;
+	virtual TSharedPtr<SGraphPin> CreatePinWidget(UEdGraphPin* Pin) const override;
 
 private:
 	// Helper for property/function binding
diff --git a/Engine/Plugins/Runtime/PropertyAccess/Source/PropertyAccessEditor/Private/PropertyAccessEditor.cpp b/Engine/Plugins/Runtime/PropertyAccess/Source/PropertyAccessEditor/Private/PropertyAccessEditor.cpp
--- a/Engine/Plugins/Runtime/PropertyAccess/Source/PropertyAccessEditor/Private/PropertyAccessEditor.cpp
+++ b/Engine/Plugins/Runtime/PropertyAccess/Source/PropertyAccessEditor/Private/PropertyAccessEditor.cpp
@@ -7,6 +7,7 @@
 #include "IPropertyAccessEditor.h"
 #include "EdGraph/EdGraphPin.h"
 #include "EdGraphSchema_K2.h"
+#include "Kismet2/BlueprintEditorUtils.h"
 
 #define LOCTEXT_NAMESPACE "PropertyAccessEditor"
 
@@ -22,10 +23,10 @@ struct FPropertyAccessEditorSystem
 		{}
 
 		// Starting struct
-		const UStruct* Struct;
+		const UStruct* Struct = nullptr;
 
 		// Current struct
-		const UStruct* CurrentStruct;
+		const UStruct* CurrentStruct = nullptr;
 
 		// Path as FStrings with optional array markup
 		TArrayView<FString> Path;
@@ -40,10 +41,13 @@ struct FPropertyAccessEditorSystem
 		FText ErrorMessage;
 
 		// The current segment index (or that at which the last error occurred)
-		int32 SegmentIndex;
+		int32 SegmentIndex = INDEX_NONE;
 
 		// Whether this is the final segment
-		bool bFinalSegment;
+		bool bFinalSegment = false;
+
+		// Whether this path was determined to be thread safe
+		bool bWasThreadSafe = true;
 	};
 
 	// The result of a segment resolve operation
@@ -51,15 +55,11 @@ struct FPropertyAccessEditorSystem
 	{
 		Failed,
 
-		SucceededInternal,
-
-		SucceededExternal,
+		Succeeded,
 	};
 
 	static ESegmentResolveResult ResolveSegments_CheckProperty(FPropertyAccessSegment& InSegment, FProperty* InProperty, FResolveSegmentsContext& InContext)
 	{
-		ESegmentResolveResult Result = ESegmentResolveResult::SucceededInternal;
-
 		InSegment.Property = InProperty;
 
 		// Check to see if it is an array first, as arrays get handled the same for 'leaf' and 'branch' nodes
@@ -79,7 +79,7 @@ struct FPropertyAccessEditorSystem
 				InSegment.Struct = ArrayOfObjectsProperty->PropertyClass;
 				if(!InContext.bFinalSegment)
 				{
-					Result = ESegmentResolveResult::SucceededExternal;
+					InContext.bWasThreadSafe = false;
 				}
 			}
 			else
@@ -106,7 +106,7 @@ struct FPropertyAccessEditorSystem
 			InSegment.Struct = ObjectProperty->PropertyClass;
 			if(!InContext.bFinalSegment)
 			{
-				Result = ESegmentResolveResult::SucceededExternal;
+				InContext.bWasThreadSafe = false;
 			}
 		}
 		// Check to see if this is a simple weak object property (eg. not an array of weak objects).
@@ -116,7 +116,7 @@ struct FPropertyAccessEditorSystem
 			InSegment.Struct = WeakObjectProperty->PropertyClass;
 			if(!InContext.bFinalSegment)
 			{
-				Result = ESegmentResolveResult::SucceededExternal;
+				InContext.bWasThreadSafe = false;
 			}
 		}
 		// Check to see if this is a simple soft object property (eg. not an array of soft objects).
@@ -126,7 +126,7 @@ struct FPropertyAccessEditorSystem
 			InSegment.Struct = SoftObjectProperty->PropertyClass;
 			if(!InContext.bFinalSegment)
 			{
-				Result = ESegmentResolveResult::SucceededExternal;
+				InContext.bWasThreadSafe = false;
 			}
 		}
 		else
@@ -135,7 +135,7 @@ struct FPropertyAccessEditorSystem
 			return ESegmentResolveResult::Failed;
 		}
 
-		return Result;
+		return ESegmentResolveResult::Succeeded;
 	}
 
 	static ESegmentResolveResult ResolveSegments_CheckFunction(FPropertyAccessSegment& InSegment, UFunction* InFunction, FResolveSegmentsContext& InContext)
@@ -158,17 +158,15 @@ struct FPropertyAccessEditorSystem
 		}
 
 		// Treat the function's return value as the struct/class we want to use for the next segment
-		ESegmentResolveResult Result = ResolveSegments_CheckProperty(InSegment, ReturnProperty, InContext);
+		const ESegmentResolveResult Result = ResolveSegments_CheckProperty(InSegment, ReturnProperty, InContext);
 		if(Result != ESegmentResolveResult::Failed)
 		{
-			// See if a function's thread safety means it should be considered 'external'.
+			// Check a function's thread safety.
 			// Note that this logic means that an external (ie. thread unsafe) object dereference returned from ResolveSegments_CheckProperty
 			// can be overridden here if the function that returns the value promises that it is thread safe to access that object.
 			// An example of this is would be something like accessing the main anim BP from a linked anim BP, where it is 'safe' to access
 			// the other object while running animation updated on a worker thread.
-			const UClass* FunctionClass = InFunction->GetOuterUClass();
-			const bool bThreadSafe = InFunction->HasMetaData(TEXT("BlueprintThreadSafe")) || (FunctionClass && FunctionClass->HasMetaData(TEXT("BlueprintThreadSafe")) && !InFunction->HasMetaData(TEXT("NotBlueprintThreadSafe")));
-			Result = bThreadSafe ? ESegmentResolveResult::SucceededInternal : ESegmentResolveResult::SucceededExternal;
+			InContext.bWasThreadSafe &= FBlueprintEditorUtils::HasFunctionBlueprintThreadSafeMetaData(InFunction);
 		}
 
 		return Result;
@@ -181,8 +179,6 @@ struct FPropertyAccessEditorSystem
 
 		if(InContext.Path.Num() > 0)
 		{
-			EPropertyAccessResolveResult Result = EPropertyAccessResolveResult::SucceededInternal;
-
 			for(int32 SegmentIndex = 0; SegmentIndex < InContext.Path.Num(); ++SegmentIndex)
 			{
 				const FString& SegmentString = InContext.Path[SegmentIndex];
@@ -216,24 +212,14 @@ struct FPropertyAccessEditorSystem
 
 				if(FProperty* Property = Field.Get<FProperty>())
 				{
-					ESegmentResolveResult PropertyResult = ResolveSegments_CheckProperty(Segment, Property, InContext);
-					if(PropertyResult == ESegmentResolveResult::SucceededExternal)
-					{
-						Result = EPropertyAccessResolveResult::SucceededExternal;
-					}
-					else if(PropertyResult == ESegmentResolveResult::Failed)
+					if(ResolveSegments_CheckProperty(Segment, Property, InContext) == ESegmentResolveResult::Failed)
 					{
 						return EPropertyAccessResolveResult::Failed;
 					}
 				}
 				else if(UFunction* Function = Field.Get<UFunction>())
 				{
-					ESegmentResolveResult FunctionResult = ResolveSegments_CheckFunction(Segment, Function, InContext);
-					if(FunctionResult == ESegmentResolveResult::SucceededExternal)
-					{
-						Result = EPropertyAccessResolveResult::SucceededExternal;
-					}
-					else if(FunctionResult == ESegmentResolveResult::Failed)
+					if(ResolveSegments_CheckFunction(Segment, Function, InContext) == ESegmentResolveResult::Failed)
 					{
 						return EPropertyAccessResolveResult::Failed;
 					}
@@ -242,7 +228,15 @@ struct FPropertyAccessEditorSystem
 				InContext.CurrentStruct = Segment.Struct;
 			}
 
-			return InContext.Segments.Num() > 0 ? Result : EPropertyAccessResolveResult::Failed;
+			if(InContext.Segments.Num() > 0)
+			{
+				return EPropertyAccessResolveResult::Succeeded;
+			}
+			else
+			{
+				InContext.ErrorMessage = LOCTEXT("NoSegments", "Unable to resolve any property path segments for @@");
+				return EPropertyAccessResolveResult::Failed;
+			}
 		}
 		else
 		{
@@ -251,22 +245,83 @@ struct FPropertyAccessEditorSystem
 		}
 	}
 
-	static EPropertyAccessResolveResult ResolveLeafProperty(const UStruct* InStruct, TArrayView<FString> InPath, FProperty*& OutProperty, int32& OutArrayIndex)
+	static FPropertyAccessResolveResult ResolvePropertyAccess(const UStruct* InStruct, TArrayView<FString> InPath, FProperty*& OutProperty, int32& OutArrayIndex)
 	{
 		FPropertyAccessPath AccessPath;
 		FResolveSegmentsContext Context(InStruct, InPath, AccessPath);
-		EPropertyAccessResolveResult Result = ResolveSegments(Context);
-		if(Result != EPropertyAccessResolveResult::Failed)
+		FPropertyAccessResolveResult Result;
+		Result.Result = ResolveSegments(Context);
+		Result.bIsThreadSafe = Context.bWasThreadSafe;
+		if(Result.Result != EPropertyAccessResolveResult::Failed)
 		{
 			const FPropertyAccessSegment& LeafSegment = Context.Segments.Last();
 			OutProperty = LeafSegment.Property.Get();
 			OutArrayIndex = LeafSegment.ArrayIndex;
-			return Result;
 		}
 
 		return Result;
 	}
 
+	static FPropertyAccessResolveResult ResolvePropertyAccess(const UStruct* InStruct, TArrayView<FString> InPath, const IPropertyAccessEditor::FResolvePropertyAccessArgs& InArgs)
+	{
+		FPropertyAccessPath AccessPath;
+		FResolveSegmentsContext Context(InStruct, InPath, AccessPath);
+		FPropertyAccessResolveResult Result;
+		Result.Result = ResolveSegments(Context);
+		Result.bIsThreadSafe = Context.bWasThreadSafe;
+		if(Result.Result != EPropertyAccessResolveResult::Failed)
+		{
+			for(int32 SegmentIndex = 0; SegmentIndex < Context.Segments.Num(); ++SegmentIndex)
+			{
+				const FPropertyAccessSegment& Segment = Context.Segments[SegmentIndex];
+				if(EnumHasAllFlags((EPropertyAccessSegmentFlags)Segment.Flags, EPropertyAccessSegmentFlags::Function))
+				{
+					if(InArgs.FunctionFunction != nullptr)
+					{
+						check(Segment.Function != nullptr);
+						check(Segment.Property.Get());
+						InArgs.FunctionFunction(SegmentIndex, Segment.Function, Segment.Property.Get());
+					}
+				}
+				else
+				{
+					switch((EPropertyAccessSegmentFlags)Segment.Flags & ~EPropertyAccessSegmentFlags::ModifierFlags)
+					{
+					case EPropertyAccessSegmentFlags::Struct:
+					case EPropertyAccessSegmentFlags::Leaf:
+					case EPropertyAccessSegmentFlags::Object:
+					case EPropertyAccessSegmentFlags::WeakObject:
+					case EPropertyAccessSegmentFlags::SoftObject:
+					{
+						if(InArgs.PropertyFunction != nullptr)
+						{
+							check(Segment.Property.Get());
+							InArgs.PropertyFunction(SegmentIndex, Segment.Property.Get(), Segment.ArrayIndex);
+						}
+						break;
+					}
+					case EPropertyAccessSegmentFlags::Array:
+					case EPropertyAccessSegmentFlags::ArrayOfStructs:
+					case EPropertyAccessSegmentFlags::ArrayOfObjects:
+					{
+						if(InArgs.ArrayFunction != nullptr)
+						{
+							FArrayProperty* ArrayProperty = CastFieldChecked<FArrayProperty>(Segment.Property.Get());
+							InArgs.ArrayFunction(SegmentIndex, ArrayProperty, Segment.ArrayIndex);
+						}
+						break;
+					}
+					default:
+						check(false);
+						break;
+					}
+				}
+			}
+		}
+
+		return Result;
+	}
+	
 	static EPropertyAccessCopyType GetCopyType(const FPropertyAccessSegment& InSrcSegment, const FPropertyAccessSegment& InDestSegment)
 	{
 		FProperty* SrcProperty = InSrcSegment.Property.Get();
@@ -406,7 +461,7 @@ struct FPropertyAccessEditorSystem
 		return EPropertyAccessCopyType::None;
 	}
 
-	static bool CompileCopy(const UStruct* InStruct, FPropertyAccessLibrary& InLibrary, FPropertyAccessLibraryCompiler::FQueuedCopy& OutCopy)
+	static bool CompileCopy(const UStruct* InStruct, const FOnPropertyAccessDetermineBatchId& InOnDetermineBatchId, FPropertyAccessLibrary& InLibrary, FPropertyAccessLibraryCompiler::FQueuedCopy& OutCopy)
 	{
 		FPropertyAccessPath SrcAccessPath;
 		FPropertyAccessPath DestAccessPath;
@@ -421,21 +476,20 @@ struct FPropertyAccessEditorSystem
 
 		if(OutCopy.SourceResult != EPropertyAccessResolveResult::Failed && OutCopy.DestResult != EPropertyAccessResolveResult::Failed)
 		{
-			const bool bExternal = OutCopy.SourceResult == EPropertyAccessResolveResult::SucceededExternal ||  OutCopy.DestResult == EPropertyAccessResolveResult::SucceededExternal;
-
-			// Decide on batch type
-			EPropertyAccessCopyBatch CopyBatch;
-			if(bExternal)
-			{
-				CopyBatch = OutCopy.BatchType == EPropertyAccessBatchType::Unbatched ? EPropertyAccessCopyBatch::ExternalUnbatched : EPropertyAccessCopyBatch::ExternalBatched;
-			}
-			else
-			{
-				CopyBatch = OutCopy.BatchType == EPropertyAccessBatchType::Unbatched ? EPropertyAccessCopyBatch::InternalUnbatched : EPropertyAccessCopyBatch::InternalBatched;
-			}
-
-			OutCopy.BatchIndex = InLibrary.CopyBatches[(__underlying_type(EPropertyAccessCopyBatch))CopyBatch].Copies.Num();
-			FPropertyAccessCopy& Copy = InLibrary.CopyBatches[(__underlying_type(EPropertyAccessCopyBatch))CopyBatch].Copies.AddDefaulted_GetRef();
+			FPropertyAccessCopyContext CopyContext;
+			CopyContext.Object = OutCopy.AssociatedObject;
+			CopyContext.ContextId = OutCopy.ContextId;
+			CopyContext.SourcePathAsText = PropertyAccess::MakeTextPath(OutCopy.SourcePath);
+			CopyContext.DestPathAsText = PropertyAccess::MakeTextPath(OutCopy.DestPath);
+			CopyContext.bSourceThreadSafe = DestContext.bWasThreadSafe;
+			CopyContext.bDestThreadSafe = SrcContext.bWasThreadSafe;
+			
+			OutCopy.BatchId = InOnDetermineBatchId.IsBound() ? InOnDetermineBatchId.Execute(CopyContext) : 0;
+			check(OutCopy.BatchId >= 0);
+			InLibrary.CopyBatchArray.SetNum(FMath::Max(OutCopy.BatchId + 1, InLibrary.CopyBatchArray.Num()));
+			
+			OutCopy.BatchIndex = InLibrary.CopyBatchArray[OutCopy.BatchId].Copies.Num();
+			FPropertyAccessCopy& Copy = InLibrary.CopyBatchArray[OutCopy.BatchId].Copies.AddDefaulted_GetRef();
 			Copy.AccessIndex = InLibrary.SrcPaths.Num();
 			Copy.DestAccessStartIndex = InLibrary.DestPaths.Num();
 			Copy.DestAccessEndIndex = InLibrary.DestPaths.Num() + 1;
@@ -460,9 +514,14 @@ struct FPropertyAccessEditorSystem
 
 namespace PropertyAccess
 {
-	EPropertyAccessResolveResult ResolveLeafProperty(const UStruct* InStruct, TArrayView<FString> InPath, FProperty*& OutProperty, int32& OutArrayIndex)
+	FPropertyAccessResolveResult ResolvePropertyAccess(const UStruct* InStruct, TArrayView<FString> InPath, FProperty*& OutProperty, int32& OutArrayIndex)
+	{
+		return ::FPropertyAccessEditorSystem::ResolvePropertyAccess(InStruct, InPath, OutProperty, OutArrayIndex);
+	}
+
+	FPropertyAccessResolveResult ResolvePropertyAccess(const UStruct* InStruct, TArrayView<FString> InPath, const IPropertyAccessEditor::FResolvePropertyAccessArgs& InArgs)
 	{
-		return ::FPropertyAccessEditorSystem::ResolveLeafProperty(InStruct, InPath, OutProperty, OutArrayIndex);
+		return ::FPropertyAccessEditorSystem::ResolvePropertyAccess(InStruct, InPath, InArgs);
 	}
 
 	EPropertyAccessCompatibility GetPropertyCompatibility(const FProperty* InPropertyA, const FProperty* InPropertyB)
@@ -617,11 +676,30 @@ namespace PropertyAccess
 			}
 		});
 	}
+
+	FText MakeTextPath(const TArray<FString>& InPath)
+	{
+		TStringBuilder<128> StringBuilder;
+
+		if(InPath.Num())
+		{
+			StringBuilder.Append(InPath[0]);
+
+			for(int32 SegmentIndex = 1; SegmentIndex < InPath.Num(); ++SegmentIndex)
+			{
+				StringBuilder.Append(TEXT("."));
+				StringBuilder.Append(InPath[SegmentIndex]);
+			}
+		}
+
+		return FText::FromString(StringBuilder.ToString());
+	}
 }
 
-FPropertyAccessLibraryCompiler::FPropertyAccessLibraryCompiler(FPropertyAccessLibrary* InLibrary, const UClass* InClass)
+FPropertyAccessLibraryCompiler::FPropertyAccessLibraryCompiler(FPropertyAccessLibrary* InLibrary, const UClass* InClass, const FOnPropertyAccessDetermineBatchId& InOnDetermineBatchId)
 	: Library(InLibrary)
 	, Class(InClass)
+	, OnDetermineBatchId(InOnDetermineBatchId)
 {
 }
 
@@ -633,17 +711,17 @@ void FPropertyAccessLibraryCompiler::BeginCompilation()
 	}
 }
 
-int32 FPropertyAccessLibraryCompiler::AddCopy(TArrayView<FString> InSourcePath, TArrayView<FString> InDestPath, EPropertyAccessBatchType InBatchType, UObject* InAssociatedObject)
+FPropertyAccessHandle FPropertyAccessLibraryCompiler::AddCopy(TArrayView<FString> InSourcePath, TArrayView<FString> InDestPath, const FName& InContextId, UObject* InAssociatedObject)
 {
 	FQueuedCopy QueuedCopy;
 	QueuedCopy.SourcePath = InSourcePath;
 	QueuedCopy.DestPath = InDestPath;
-	QueuedCopy.BatchType = InBatchType;
+	QueuedCopy.ContextId = InContextId;
 	QueuedCopy.AssociatedObject = InAssociatedObject;
 
 	QueuedCopies.Add(MoveTemp(QueuedCopy));
 
-	return QueuedCopies.Num() - 1;
+	return FPropertyAccessHandle(QueuedCopies.Num() - 1);
 }
 
 bool FPropertyAccessLibraryCompiler::FinishCompilation()
@@ -654,8 +732,8 @@ bool FPropertyAccessLibraryCompiler::FinishCompilation()
 		for(int32 CopyIndex = 0; CopyIndex < QueuedCopies.Num(); ++CopyIndex)
 		{
 			FQueuedCopy& Copy = QueuedCopies[CopyIndex];
-			bResult &= ::FPropertyAccessEditorSystem::CompileCopy(Class, *Library, Copy);
-			CopyMap.Add(CopyIndex, Copy.BatchIndex);
+			bResult &= ::FPropertyAccessEditorSystem::CompileCopy(Class, OnDetermineBatchId, *Library, Copy);
+			CopyMap.Add(FPropertyAccessHandle(CopyIndex), FCompiledPropertyAccessHandle(Copy.BatchIndex, Copy.BatchId));
 		}
 
 		if(bResult)
@@ -688,20 +766,14 @@ void FPropertyAccessLibraryCompiler::IterateErrors(TFunctionRef<void(const FText
 	}
 }
 
-int32 FPropertyAccessLibraryCompiler::MapCopyIndex(int32 InIndex) const
+FCompiledPropertyAccessHandle FPropertyAccessLibraryCompiler::GetCompiledHandle(FPropertyAccessHandle InHandle) const
 {
-	if(const int32* FoundIndex = CopyMap.Find(InIndex))
+	if(const FCompiledPropertyAccessHandle* FoundHandle = CopyMap.Find(InHandle))
 	{
-		return *FoundIndex;
+		return *FoundHandle;
 	}
 
-	return INDEX_NONE;
-}
-
-void FPropertyAccessLibraryCompiler::Setup(const UClass* InClass, FPropertyAccessLibrary* InLibrary)
-{
-	Class = InClass;
-	Library = InLibrary;
+	return FCompiledPropertyAccessHandle();
 }
 
 #undef LOCTEXT_NAMESPACE
\ No newline at end of file
diff --git a/Engine/Plugins/Runtime/PropertyAccess/Source/PropertyAccessEditor/Private/PropertyAccessEditor.h b/Engine/Plugins/Runtime/PropertyAccess/Source/PropertyAccessEditor/Private/PropertyAccessEditor.h
--- a/Engine/Plugins/Runtime/PropertyAccess/Source/PropertyAccessEditor/Private/PropertyAccessEditor.h
+++ b/Engine/Plugins/Runtime/PropertyAccess/Source/PropertyAccessEditor/Private/PropertyAccessEditor.h
@@ -13,7 +13,13 @@ struct FEdGraphPinType;
 namespace PropertyAccess
 {
 	/** Resolve a property path to a structure, returning the leaf property and array index if any. @return true if resolution succeeded */
-	extern EPropertyAccessResolveResult ResolveLeafProperty(const UStruct* InStruct, TArrayView<FString> InPath, FProperty*& OutProperty, int32& OutArrayIndex);
+	extern FPropertyAccessResolveResult ResolvePropertyAccess(const UStruct* InStruct, TArrayView<FString> InPath, FProperty*& OutProperty, int32& OutArrayIndex);
+
+	/**
+	 * Resolve a property path to a structure, calling back for each segment in path segment order if resolution succeed
+	 * @return true if resolution succeeded
+	 */
+	extern FPropertyAccessResolveResult ResolvePropertyAccess(const UStruct* InStruct, TArrayView<FString> InPath, const IPropertyAccessEditor::FResolvePropertyAccessArgs& InArgs);
 
 	// Get the compatibility of the two supplied properties. Ordering matters for promotion (A->B).
 	extern EPropertyAccessCompatibility GetPropertyCompatibility(const FProperty* InPropertyA, const FProperty* InPropertyB);
@@ -23,36 +29,56 @@ namespace PropertyAccess
 
 	// Makes a string path from a binding chain
 	extern void MakeStringPath(const TArray<FBindingChainElement>& InBindingChain, TArray<FString>& OutStringPath);
+
+	// Makes a text path from an array of path segments 
+	extern FText MakeTextPath(const TArray<FString>& InPath);
 }
 
 // A helper structure used to compile a property access library
 class FPropertyAccessLibraryCompiler : public IPropertyAccessLibraryCompiler
 {
 public:
-	FPropertyAccessLibraryCompiler(FPropertyAccessLibrary* InLibrary, const UClass* InClass);
+	FPropertyAccessLibraryCompiler(FPropertyAccessLibrary* InLibrary, const UClass* InClass, const FOnPropertyAccessDetermineBatchId& InOnDetermineBatchId);
 
 	// IPropertyAccessLibraryCompiler interface
 	virtual void BeginCompilation() override;
-	virtual int32 AddCopy(TArrayView<FString> InSourcePath, TArrayView<FString> InDestPath, EPropertyAccessBatchType InBatchType, UObject* InAssociatedObject = nullptr) override;
+	virtual FPropertyAccessHandle AddCopy(TArrayView<FString> InSourcePath, TArrayView<FString> InDestPath, const FName& InContextId, UObject* InAssociatedObject = nullptr) override;
 	virtual bool FinishCompilation() override;
 	virtual void IterateErrors(TFunctionRef<void(const FText&, UObject*)> InFunction) const override;
-	virtual int32 MapCopyIndex(int32 InIndex) const override;
-
-	// Set up the compiler ready for compilation
-	void Setup(const UClass* InClass, FPropertyAccessLibrary* InLibrary);
+	virtual FCompiledPropertyAccessHandle GetCompiledHandle(FPropertyAccessHandle InHandle) const override;
 
 public:
 	// Stored copy info for processing in FinishCompilation()
 	struct FQueuedCopy
 	{
+		// The path to copy from
 		TArray<FString> SourcePath;
+
+		// The path to copy to
 		TArray<FString> DestPath;
-		EPropertyAccessBatchType BatchType;
+
+		// The user-defined context ID
+		FName ContextId = NAME_None;
+
+		// Error text associated with the source path
 		FText SourceErrorText;
+
+		// Error text associated with the dest path
 		FText DestErrorText;
+
+		// Result of the source path resolution
 		EPropertyAccessResolveResult SourceResult = EPropertyAccessResolveResult::Failed;
+
+		// Result of the dest path resolution
 		EPropertyAccessResolveResult DestResult = EPropertyAccessResolveResult::Failed;
+
+		// Associated object, used to provide context for path resolution
 		UObject* AssociatedObject = nullptr;
+
+		// Batch Id
+		int32 BatchId = INDEX_NONE;
+		
+		// Index within the batch
 		int32 BatchIndex = INDEX_NONE;
 	};
 
@@ -65,9 +91,12 @@ protected:
 	// The class we are compiling the library for
 	const UClass* Class;
 
+	// Delegate used to determine batch ID (index) for a particular copy context
+	FOnPropertyAccessDetermineBatchId OnDetermineBatchId;
+	
 	// All copies to process in FinishCompilation()
 	TArray<FQueuedCopy> QueuedCopies;
 
 	// Copy map
-	TMap<int32, int32> CopyMap;
+	TMap<FPropertyAccessHandle, FCompiledPropertyAccessHandle> CopyMap;
 };
diff --git a/Engine/Plugins/Runtime/PropertyAccess/Source/PropertyAccessEditor/Private/PropertyAccessEditorModule.cpp b/Engine/Plugins/Runtime/PropertyAccess/Source/PropertyAccessEditor/Private/PropertyAccessEditorModule.cpp
--- a/Engine/Plugins/Runtime/PropertyAccess/Source/PropertyAccessEditor/Private/PropertyAccessEditorModule.cpp
+++ b/Engine/Plugins/Runtime/PropertyAccess/Source/PropertyAccessEditor/Private/PropertyAccessEditorModule.cpp
@@ -6,6 +6,7 @@
 #include "SPropertyBinding.h"
 #include "EdGraphUtilities.h"
 #include "PropertyAccessEditor.h"
+#include "Algo/Accumulate.h"
 #include "Modules/ModuleInterface.h"
 #include "Features/IModularFeatures.h"
 
@@ -37,11 +38,16 @@ public:
 			.Args(InArgs);
 	}
 
-	virtual EPropertyAccessResolveResult ResolveLeafProperty(const UStruct* InStruct, TArrayView<FString> InPath, FProperty*& OutProperty, int32& OutArrayIndex) const override
+	virtual FPropertyAccessResolveResult ResolvePropertyAccess(const UStruct* InStruct, TArrayView<FString> InPath, FProperty*& OutProperty, int32& OutArrayIndex) const override
 	{
-		return PropertyAccess::ResolveLeafProperty(InStruct, InPath, OutProperty, OutArrayIndex);
+		return PropertyAccess::ResolvePropertyAccess(InStruct, InPath, OutProperty, OutArrayIndex);
 	}
 
+	virtual FPropertyAccessResolveResult ResolvePropertyAccess(const UStruct* InStruct, TArrayView<FString> InPath, const FResolvePropertyAccessArgs& InArgs) const override
+	{
+		return PropertyAccess::ResolvePropertyAccess(InStruct, InPath, InArgs);
+	}
+	
 	virtual EPropertyAccessCompatibility GetPropertyCompatibility(const FProperty* InPropertyA, const FProperty* InPropertyB) const override
 	{
 		return PropertyAccess::GetPropertyCompatibility(InPropertyA, InPropertyB);
@@ -59,8 +65,13 @@ public:
 
 	virtual TUniquePtr<IPropertyAccessLibraryCompiler> MakePropertyAccessCompiler(const FPropertyAccessLibraryCompilerArgs& InArgs) const override
 	{
-		return MakeUnique<FPropertyAccessLibraryCompiler>(&InArgs.Library, InArgs.ClassContext);
+		return MakeUnique<FPropertyAccessLibraryCompiler>(&InArgs.Library, InArgs.ClassContext, InArgs.OnDetermineBatchId);
 	}
+
+	virtual FText MakeTextPath(const TArray<FString>& InPath) const override
+	{
+		return PropertyAccess::MakeTextPath(InPath);
+	};
 };
 
 IMPLEMENT_MODULE(FPropertyAccessEditorModule, PropertyAccessEditor)
\ No newline at end of file
diff --git a/Engine/Plugins/Runtime/PropertyAccess/Source/PropertyAccessEditor/Private/SPropertyBinding.cpp b/Engine/Plugins/Runtime/PropertyAccess/Source/PropertyAccessEditor/Private/SPropertyBinding.cpp
--- a/Engine/Plugins/Runtime/PropertyAccess/Source/PropertyAccessEditor/Private/SPropertyBinding.cpp
+++ b/Engine/Plugins/Runtime/PropertyAccess/Source/PropertyAccessEditor/Private/SPropertyBinding.cpp
@@ -46,39 +46,35 @@ void SPropertyBinding::Construct(const FArguments& InArgs, UBlueprint* InBluepri
 		+ SHorizontalBox::Slot()
 		.FillWidth(1.0f)
 		[
-			SNew(SBox)
-			.MaxDesiredWidth(200.0f)
+			SNew(SComboButton)
+			.ToolTipText(this, &SPropertyBinding::GetCurrentBindingToolTipText)
+			.OnGetMenuContent(this, &SPropertyBinding::OnGenerateDelegateMenu)
+			.ContentPadding(1)
+			.ButtonContent()
 			[
-				SNew(SComboButton)
-				.ToolTipText(this, &SPropertyBinding::GetCurrentBindingText)
-				.OnGetMenuContent(this, &SPropertyBinding::OnGenerateDelegateMenu)
-				.ContentPadding(1)
-				.ButtonContent()
+				SNew(SHorizontalBox)
+				.Clipping(EWidgetClipping::ClipToBounds)
+				+ SHorizontalBox::Slot()
+				.AutoWidth()
+				.VAlign(VAlign_Center)
 				[
-					SNew(SHorizontalBox)
-					.Clipping(EWidgetClipping::ClipToBounds)
-					+ SHorizontalBox::Slot()
-					.AutoWidth()
-					.VAlign(VAlign_Center)
+					SNew(SBox)
+					.HeightOverride(16.0f)
 					[
-						SNew(SBox)
-						.HeightOverride(16.0f)
-						[
-							SNew(SImage)
-							.Image(this, &SPropertyBinding::GetCurrentBindingImage)
-							.ColorAndOpacity(this, &SPropertyBinding::GetCurrentBindingColor)
-						]
-					]
-					+ SHorizontalBox::Slot()
-					.AutoWidth()
-					.VAlign(VAlign_Center)
-					.Padding(4, 0, 0, 0)
-					[
-						SNew(STextBlock)
-						.Text(this, &SPropertyBinding::GetCurrentBindingText)
-						.Font(IDetailLayoutBuilder::GetDetailFont())
+						SNew(SImage)
+						.Image(this, &SPropertyBinding::GetCurrentBindingImage)
+						.ColorAndOpacity(this, &SPropertyBinding::GetCurrentBindingColor)
 					]
 				]
+				+ SHorizontalBox::Slot()
+				.AutoWidth()
+				.VAlign(VAlign_Center)
+				.Padding(4, 0, 0, 0)
+				[
+					SNew(STextBlock)
+					.Text(this, &SPropertyBinding::GetCurrentBindingText)
+					.Font(IDetailLayoutBuilder::GetDetailFont())
+				]
 			]
 		]
 
@@ -152,7 +148,17 @@ void SPropertyBinding::ForEachBindableFunction(UClass* FromClass, Predicate Pred
 				}
 			}
 
-			if(bValidObjectFunction || Args.OnCanBindFunction.Execute(Function))
+			bool bValidStructFunction = false;
+			if(Args.bAllowStructFunctions)
+			{
+				FStructProperty* StructProperty = CastField<FStructProperty>(Function->GetReturnProperty());
+				if(StructProperty != nullptr && Function->NumParms == 1)
+				{
+					bValidStructFunction = true;
+				}
+			}
+
+			if(bValidObjectFunction || bValidStructFunction || Args.OnCanBindFunction.Execute(Function))
 			{
 				Pred(FFunctionInfo(Function));
 			}
@@ -472,7 +478,7 @@ void SPropertyBinding::FillPropertyMenu(FMenuBuilder& MenuBuilder, UStruct* InOw
 						MakeFunctionWidget(Info));
 				}
 
-				// Only show bindable subobjects and variables if we're generating pure bindings.
+				// Only show bindable subobjects, structs and variables if we're generating pure bindings.
 				if(Args.bGeneratePureBindings)
 				{
 					if(FObjectPropertyBase* ObjectPropertyBase = CastField<FObjectPropertyBase>(ReturnProperty))
@@ -481,6 +487,12 @@ void SPropertyBinding::FillPropertyMenu(FMenuBuilder& MenuBuilder, UStruct* InOw
 							MakeFunctionWidget(Info),
 							FNewMenuDelegate::CreateSP(this, &SPropertyBinding::FillPropertyMenu, static_cast<UStruct*>(ObjectPropertyBase->PropertyClass), NewBindingChain));
 					}
+					else if(FStructProperty* StructProperty = CastField<FStructProperty>(ReturnProperty))
+					{
+						MenuBuilder.AddSubMenu(
+							MakeFunctionWidget(Info),
+							FNewMenuDelegate::CreateSP(this, &SPropertyBinding::FillPropertyMenu, static_cast<UStruct*>(StructProperty->Struct), NewBindingChain));
+					}
 				}
 			});
 		}
@@ -629,6 +641,16 @@ FText SPropertyBinding::GetCurrentBindingText() const
 	return LOCTEXT("Bind", "Bind");
 }
 
+FText SPropertyBinding::GetCurrentBindingToolTipText() const
+{
+	if(Args.CurrentBindingToolTipText.IsSet())
+	{
+		return Args.CurrentBindingToolTipText.Get();
+	}
+
+	return GetCurrentBindingText();
+}
+
 FSlateColor SPropertyBinding::GetCurrentBindingColor() const
 {
 	if(Args.CurrentBindingColor.IsSet())
diff --git a/Engine/Plugins/Runtime/PropertyAccess/Source/PropertyAccessEditor/Private/SPropertyBinding.h b/Engine/Plugins/Runtime/PropertyAccess/Source/PropertyAccessEditor/Private/SPropertyBinding.h
--- a/Engine/Plugins/Runtime/PropertyAccess/Source/PropertyAccessEditor/Private/SPropertyBinding.h
+++ b/Engine/Plugins/Runtime/PropertyAccess/Source/PropertyAccessEditor/Private/SPropertyBinding.h
@@ -55,6 +55,7 @@ protected:
 
 	const FSlateBrush* GetCurrentBindingImage() const;
 	FText GetCurrentBindingText() const;
+	FText GetCurrentBindingToolTipText() const;
 	FSlateColor GetCurrentBindingColor() const;
 
 	bool CanRemoveBinding();
diff --git a/Engine/Source/Editor/AnimGraph/Private/AnimBlueprintCompiler.cpp b/Engine/Source/Editor/AnimGraph/Private/AnimBlueprintCompiler.cpp
--- a/Engine/Source/Editor/AnimGraph/Private/AnimBlueprintCompiler.cpp
+++ b/Engine/Source/Editor/AnimGraph/Private/AnimBlueprintCompiler.cpp
@@ -389,7 +389,7 @@ void FAnimBlueprintCompilerContext::GatherFoldRecordsForAnimationNode(const UScr
 					UEdGraphPin* Pin = InVisualAnimNode->FindPin(ArrayElementPinName);
 					const bool bExposedOnPin = Pin != nullptr;
 					const bool bPinConnected = InVisualAnimNode->IsPinExposedAndLinked(ArrayElementPinName, EGPD_Input);
-					const bool bPinBound = InVisualAnimNode->IsPinExposedAndBound(SubProperty->GetName(), EGPD_Input);
+					const bool bPinBound = InVisualAnimNode->IsPinExposedAndBound(ArrayElementPinName, EGPD_Input);
 					
 					if(bExposedOnPin)
 					{
@@ -997,70 +997,84 @@ void FAnimBlueprintCompilerContext::CopyTermDefaultsToDefaultObject(UObject* Def
 				if (UbergraphFunctionContext)
 				{
 					// run through the per-node compiled statements looking for struct-sets used by anim nodes
-					for (auto& StatementPair : UbergraphFunctionContext->StatementsPerNode)
+					for (const TPair<UEdGraphNode*, TArray<FBlueprintCompiledStatement*>>& StatementPair : UbergraphFunctionContext->StatementsPerNode)
 					{
 						if (UK2Node_StructMemberSet* StructMemberSetNode = Cast<UK2Node_StructMemberSet>(StatementPair.Key))
 						{
-							UObject* SourceNode = MessageLog.FindSourceObject(StructMemberSetNode);
+							UK2Node* SourceNode = CastChecked<UK2Node>(MessageLog.FindSourceObject(StructMemberSetNode));
+							UEdGraph* SourceGraph = SourceNode->GetGraph();
 
-							if (SourceNode && StructMemberSetNode->StructType->IsChildOf(FAnimNode_Base::StaticStruct()))
+							if (SourceNode && (StructMemberSetNode->StructType->IsChildOf(FAnimNode_Base::StaticStruct()) || StructMemberSetNode->StructType->IsChildOf(FAnimBlueprintMutableData::StaticStruct())))
 							{
-								for (FBlueprintCompiledStatement* Statement : StatementPair.Value)
+								const bool bEmitErrors = false;
+								bool bIsThreadSafe = true;
+
+								static const FBoolConfigValueHelper UseLegacyAnimBlueprintThreadSafetyChecks(TEXT("Kismet"), TEXT("bUseLegacyAnimBlueprintThreadSafetyChecks"), GEngineIni);
+								if(UseLegacyAnimBlueprintThreadSafetyChecks)
 								{
-									if (Statement->Type == KCST_CallFunction && Statement->FunctionToCall)
+									for (FBlueprintCompiledStatement* Statement : StatementPair.Value)
 									{
-										// pure function?
-										const bool bPureFunctionCall = Statement->FunctionToCall->HasAnyFunctionFlags(FUNC_BlueprintPure);
-
-										// function called on something other than function library or anim instance?
-										UClass* FunctionClass = CastChecked<UClass>(Statement->FunctionToCall->GetOuter());
-										const bool bFunctionLibraryCall = FunctionClass->IsChildOf<UBlueprintFunctionLibrary>();
-										const bool bAnimInstanceCall = FunctionClass->IsChildOf<UAnimInstance>();
+										if (Statement->Type == KCST_CallFunction && Statement->FunctionToCall)
+										{
+											// pure function?
+											const bool bPureFunctionCall = Statement->FunctionToCall->HasAnyFunctionFlags(FUNC_BlueprintPure);
 
-										// Whitelisted/blacklisted? Some functions are not really 'pure', so we give people the opportunity to mark them up.
-										// Mark up the class if it is generally thread safe, then unsafe functions can be marked up individually. We assume
-										// that classes are unsafe by default, as well as if they are marked up NotBlueprintThreadSafe.
-										const bool bClassThreadSafe = FunctionClass->HasMetaData(TEXT("BlueprintThreadSafe"));
-										const bool bClassNotThreadSafe = FunctionClass->HasMetaData(TEXT("NotBlueprintThreadSafe")) || !FunctionClass->HasMetaData(TEXT("BlueprintThreadSafe"));
-										const bool bFunctionThreadSafe = Statement->FunctionToCall->HasMetaData(TEXT("BlueprintThreadSafe"));
-										const bool bFunctionNotThreadSafe = Statement->FunctionToCall->HasMetaData(TEXT("NotBlueprintThreadSafe"));
+											// function called on something other than function library or anim instance?
+											UClass* FunctionClass = CastChecked<UClass>(Statement->FunctionToCall->GetOuter());
+											const bool bFunctionLibraryCall = FunctionClass->IsChildOf<UBlueprintFunctionLibrary>();
+											const bool bAnimInstanceCall = FunctionClass->IsChildOf<UAnimInstance>();
 
-										const bool bThreadSafe = (bClassThreadSafe && !bFunctionNotThreadSafe) || (bClassNotThreadSafe && bFunctionThreadSafe);
+											// Whitelisted/blacklisted? Some functions are not really 'pure', so we give people the opportunity to mark them up.
+											// Mark up the class if it is generally thread safe, then unsafe functions can be marked up individually. We assume
+											// that classes are unsafe by default, as well as if they are marked up NotBlueprintThreadSafe.
+											const bool bClassThreadSafe = FunctionClass->HasMetaData(TEXT("BlueprintThreadSafe"));
+											const bool bClassNotThreadSafe = FunctionClass->HasMetaData(TEXT("NotBlueprintThreadSafe")) || !FunctionClass->HasMetaData(TEXT("BlueprintThreadSafe"));
+											const bool bFunctionThreadSafe = Statement->FunctionToCall->HasMetaData(TEXT("BlueprintThreadSafe"));
+											const bool bFunctionNotThreadSafe = Statement->FunctionToCall->HasMetaData(TEXT("NotBlueprintThreadSafe"));
 
-										const bool bValidForUsage = bPureFunctionCall && bThreadSafe && (bFunctionLibraryCall || bAnimInstanceCall);
+											const bool bThreadSafe = (bClassThreadSafe && !bFunctionNotThreadSafe) || (bClassNotThreadSafe && bFunctionThreadSafe);
 
-										if (!bValidForUsage)
-										{
-											UEdGraphNode* FunctionNode = nullptr;
-											if (Statement->FunctionContext && Statement->FunctionContext->SourcePin)
-											{
-												FunctionNode = Statement->FunctionContext->SourcePin->GetOwningNode();
-											}
-											else if (Statement->LHS && Statement->LHS->SourcePin)
-											{
-												FunctionNode = Statement->LHS->SourcePin->GetOwningNode();
-											}
+											const bool bValidForUsage = bPureFunctionCall && bThreadSafe && (bFunctionLibraryCall || bAnimInstanceCall);
 
-											if (FunctionNode)
+											if (!bValidForUsage)
 											{
-												MessageLog.Warning(*LOCTEXT("NotThreadSafeWarningNodeContext", "Node @@ uses potentially thread-unsafe call @@. Disable threaded update or use a thread-safe call. Function may need BlueprintThreadSafe metadata adding.").ToString(), SourceNode, FunctionNode)
-													->AddToken(FDocumentationToken::Create(TEXT("Engine/Animation/AnimBlueprints/AnimGraph")));
+												UEdGraphNode* FunctionNode = nullptr;
+												if (Statement->FunctionContext && Statement->FunctionContext->SourcePin)
+												{
+													FunctionNode = Statement->FunctionContext->SourcePin->GetOwningNode();
+												}
+												else if (Statement->LHS && Statement->LHS->SourcePin)
+												{
+													FunctionNode = Statement->LHS->SourcePin->GetOwningNode();
+												}
+
+												if (FunctionNode)
+												{
+													MessageLog.Warning(*LOCTEXT("NotThreadSafeWarningNodeContext", "Node @@ uses potentially thread-unsafe call @@. Disable threaded update or use a thread-safe call. Function may need BlueprintThreadSafe metadata adding.").ToString(), SourceNode, FunctionNode)
+														->AddToken(FDocumentationToken::Create(TEXT("Engine/Animation/AnimBlueprints/AnimGraph")));
+												}
+												else if (Statement->FunctionToCall)
+												{
+													MessageLog.Warning(*FText::Format(LOCTEXT("NotThreadSafeWarningFunctionContext", "Node @@ uses potentially thread-unsafe call {0}. Disable threaded update or use a thread-safe call. Function may need BlueprintThreadSafe metadata adding."), Statement->FunctionToCall->GetDisplayNameText()).ToString(), SourceNode)
+														->AddToken(FDocumentationToken::Create(TEXT("Engine/Animation/AnimBlueprints/AnimGraph")));
+												}
+												else
+												{
+													MessageLog.Warning(*LOCTEXT("NotThreadSafeWarningUnknownContext", "Node @@ uses potentially thread-unsafe call. Disable threaded update or use a thread-safe call.").ToString(), SourceNode)
+														->AddToken(FDocumentationToken::Create(TEXT("Engine/Animation/AnimBlueprints/AnimGraph")));
+												}
+
+												bIsThreadSafe = false;
 											}
-											else if (Statement->FunctionToCall)
-											{
-												MessageLog.Warning(*FText::Format(LOCTEXT("NotThreadSafeWarningFunctionContext", "Node @@ uses potentially thread-unsafe call {0}. Disable threaded update or use a thread-safe call. Function may need BlueprintThreadSafe metadata adding."), Statement->FunctionToCall->GetDisplayNameText()).ToString(), SourceNode)
-													->AddToken(FDocumentationToken::Create(TEXT("Engine/Animation/AnimBlueprints/AnimGraph")));
-											}
-											else
-											{
-												MessageLog.Warning(*LOCTEXT("NotThreadSafeWarningUnknownContext", "Node @@ uses potentially thread-unsafe call. Disable threaded update or use a thread-safe call.").ToString(), SourceNode)
-													->AddToken(FDocumentationToken::Create(TEXT("Engine/Animation/AnimBlueprints/AnimGraph")));
-											}
-
-											DefaultAnimInstance->bUseMultiThreadedAnimationUpdate = false;
 										}
 									}
 								}
+								else
+								{
+									bIsThreadSafe = FKismetCompilerUtilities::CheckFunctionCompiledStatementsThreadSafety(SourceNode, SourceGraph, StatementPair.Value, MessageLog, bEmitErrors);
+								}
+								
+								DefaultAnimInstance->bUseMultiThreadedAnimationUpdate = bIsThreadSafe;
 							}
 						}
 					}
@@ -1248,7 +1262,7 @@ void FAnimBlueprintCompilerContext::RecreateMutables()
 {
 	UAnimBlueprintGeneratedClass* NewAnimBlueprintClass = GetNewAnimBlueprintClass();
 
-	UScriptStruct* OldMutablesStruct = NewAnimBlueprintClass->MutableNodeDataProperty ? NewAnimBlueprintClass->MutableNodeDataProperty->Struct : nullptr;
+	UScriptStruct* OldMutablesStruct = FindObject<UScriptStruct>(NewAnimBlueprintClass, *UAnimBlueprintGeneratedClass::GetMutablesStructName().ToString());
 	
 	// Set up our mutables struct
 	NewAnimBlueprintMutables = NewObject<UScriptStruct>(NewAnimBlueprintClass, UAnimBlueprintGeneratedClass::GetMutablesStructName());
@@ -1712,7 +1726,9 @@ TArrayView<const FName> FAnimBlueprintCompilerContext::GetAttributesFromNode(UAn
 FProperty* FAnimBlueprintCompilerContext::CreateUniqueVariable(UObject* InForObject, const FEdGraphPinType& Type)
 {
 	const FString VariableName = ClassScopeNetNameMap.MakeValidName(InForObject);
-	return CreateVariable(*VariableName, Type);
+	FProperty* Variable = CreateVariable(*VariableName, Type);
+	Variable->SetMetaData(FBlueprintMetadata::MD_Private, TEXT("true"));
+	return Variable;
 }
 
 FProperty* FAnimBlueprintCompilerContext::CreateStructVariable(UScriptStruct* InStruct, const FName VarName, const FEdGraphPinType& VarType)
diff --git a/Engine/Source/Editor/AnimGraph/Private/AnimBlueprintExtension_Base.cpp b/Engine/Source/Editor/AnimGraph/Private/AnimBlueprintExtension_Base.cpp
--- a/Engine/Source/Editor/AnimGraph/Private/AnimBlueprintExtension_Base.cpp
+++ b/Engine/Source/Editor/AnimGraph/Private/AnimBlueprintExtension_Base.cpp
@@ -47,14 +47,15 @@ void UAnimBlueprintExtension_Base::HandleCopyTermDefaultsToDefaultObject(UObject
 	{
 		for(const FEvaluationHandlerRecord& EvaluationHandler : ValidEvaluationHandlerList)
 		{
-			if(EvaluationHandler.EvaluationHandlerIdx != INDEX_NONE && EvaluationHandler.ServicedProperties.Num() > 0)
+			if(EvaluationHandler.AnimGraphNode)
 			{
-				const FAnimNodeSinglePropertyHandler& Handler = EvaluationHandler.ServicedProperties.CreateConstIterator()->Value;
-				check(Handler.CopyRecords.Num() > 0);
-				if(Handler.CopyRecords[0].DestPin != nullptr)
+				UAnimGraphNode_Base* Node = CastChecked<UAnimGraphNode_Base>(EvaluationHandler.AnimGraphNode);
+				UAnimGraphNode_Base* TrueNode = InCompilationContext.GetMessageLog().FindSourceObjectTypeChecked<UAnimGraphNode_Base>(Node);	
+				
+				if(EvaluationHandler.EvaluationHandlerIdx != INDEX_NONE && EvaluationHandler.ServicedProperties.Num() > 0)
 				{
-					UAnimGraphNode_Base* Node = CastChecked<UAnimGraphNode_Base>(Handler.CopyRecords[0].DestPin->GetOwningNode());
-					UAnimGraphNode_Base* TrueNode = InCompilationContext.GetMessageLog().FindSourceObjectTypeChecked<UAnimGraphNode_Base>(Node);	
+					const FAnimNodeSinglePropertyHandler& Handler = EvaluationHandler.ServicedProperties.CreateConstIterator()->Value;
+					check(Handler.CopyRecords.Num() > 0);
 
 					const FExposedValueHandler& ValueHandler = Subsystem.ExposedValueHandlers[ EvaluationHandler.EvaluationHandlerIdx ];
 					TrueNode->BlueprintUsage = ValueHandler.BoundFunction != NAME_None ? EBlueprintUsage::UsesBlueprint : EBlueprintUsage::DoesNotUseBlueprint; 
@@ -173,14 +174,7 @@ void UAnimBlueprintExtension_Base::HandleStartCompilingClass(const UClass* InCla
 						{
 							if(Record.IsFastPath())
 							{
-								// Check if the resolved copy
-								FProperty* LeafProperty = nullptr;
-								int32 ArrayIndex = INDEX_NONE;
-								EPropertyAccessResolveResult Result = PropertyAccessEditor.ResolveLeafProperty(InClass, Record.SourcePropertyPath, LeafProperty, ArrayIndex);
-
-								// Batch all external accesses, we cant call them safely from a worker thread.
-								Record.LibraryBatchType = Result == EPropertyAccessResolveResult::SucceededExternal ? EPropertyAccessBatchType::Batched : EPropertyAccessBatchType::Unbatched;
-								Record.LibraryCopyIndex = PropertyAccessExtension->AddCopy(Record.SourcePropertyPath, Record.DestPropertyPath, Record.LibraryBatchType, HandlerRecord.AnimGraphNode);
+								Record.LibraryHandle = PropertyAccessExtension->AddCopy(Record.SourcePropertyPath, Record.DestPropertyPath, Record.BindingContextId, HandlerRecord.AnimGraphNode);
 							}
 						}
 					}
@@ -194,6 +188,8 @@ void UAnimBlueprintExtension_Base::HandleStartCompilingClass(const UClass* InCla
 		{
 			for(FEvaluationHandlerRecord& HandlerRecord : ValidEvaluationHandlerList)
 			{
+				UAnimGraphNode_Base* OriginalNode = Cast<UAnimGraphNode_Base>(InCompilationContext.GetMessageLog().FindSourceObject(HandlerRecord.AnimGraphNode));
+				
 				// Map global copy index to batched indices
 				for(TPair<FName, FAnimNodeSinglePropertyHandler>& PropertyHandler : HandlerRecord.ServicedProperties)
 				{
@@ -201,7 +197,22 @@ void UAnimBlueprintExtension_Base::HandleStartCompilingClass(const UClass* InCla
 					{
 						if(CopyRecord.IsFastPath())
 						{
-							CopyRecord.LibraryCopyIndex = PropertyAccessExtension->MapCopyIndex(CopyRecord.LibraryCopyIndex);
+							CopyRecord.LibraryCompiledHandle = PropertyAccessExtension->GetCompiledHandle(CopyRecord.LibraryHandle);
+
+							// Push compiled desc back to origial node for feedback
+							if(FAnimGraphNodePropertyBinding* Binding = OriginalNode->PropertyBindings.Find(CopyRecord.DestProperty->GetFName()))
+							{
+								if(CopyRecord.LibraryCompiledHandle.IsValid())
+								{
+									Binding->CompiledContext = UAnimBlueprintExtension_PropertyAccess::GetCompiledHandleContext(CopyRecord.LibraryCompiledHandle);
+									Binding->CompiledContextDesc = UAnimBlueprintExtension_PropertyAccess::GetCompiledHandleContextDesc(CopyRecord.LibraryCompiledHandle);
+								}
+								else
+								{
+									Binding->CompiledContext = FText::GetEmpty();
+									Binding->CompiledContextDesc = FText::GetEmpty();
+								}
+							}
 						}
 					}
 				}
@@ -224,12 +235,10 @@ void UAnimBlueprintExtension_Base::HandleFinishCompilingClass(const UClass* InCl
 	}
 }
 
-void UAnimBlueprintExtension_Base::AddStructEvalHandlers(UAnimGraphNode_Base* InNode, IAnimBlueprintCompilationContext& InCompilationContext, IAnimBlueprintGeneratedClassCompiledData& OutCompiledData)
+void UAnimBlueprintExtension_Base::ProcessNodePins(UAnimGraphNode_Base* InNode, IAnimBlueprintCompilationContext& InCompilationContext, IAnimBlueprintGeneratedClassCompiledData& OutCompiledData)
 {
 	const UAnimationGraphSchema* AnimGraphDefaultSchema = GetDefault<UAnimationGraphSchema>();
 
-	FEvaluationHandlerRecord& EvalHandler = PerNodeStructEvalHandlers.Add(InNode);
-
 	FStructProperty* NodeProperty = CastFieldChecked<FStructProperty>(InCompilationContext.GetAllocatedPropertiesByNode().FindChecked(InNode));
 
 	for (auto SourcePinIt = InNode->Pins.CreateIterator(); SourcePinIt; ++SourcePinIt)
@@ -248,8 +257,14 @@ void UAnimBlueprintExtension_Base::AddStructEvalHandlers(UAnimGraphNode_Base* In
 				bConsumed = true;
 			}
 		}
+		else if(!InNode->ShouldCreateStructEvalHandlers())
+		{
+			bConsumed = true;
+		}
 		else
 		{
+			FEvaluationHandlerRecord& EvalHandler = PerNodeStructEvalHandlers.FindOrAdd(InNode);
+			
 			// The property source for our data, either the struct property for an anim node, or the
 			// owning anim instance if using a linked instance node.
 			FProperty* SourcePinProperty = nullptr;
@@ -329,17 +344,22 @@ void UAnimBlueprintExtension_Base::AddStructEvalHandlers(UAnimGraphNode_Base* In
 	{
 		if(PropertyBinding.Value.bIsBound)
 		{
+			FEvaluationHandlerRecord& EvalHandler = PerNodeStructEvalHandlers.FindOrAdd(InNode);
 			EvalHandler.AnimGraphNode = InNode;
 			EvalHandler.NodeVariableProperty = NodeProperty;
 			EvalHandler.bServicesNodeProperties = true;
 
-			if (FProperty* Property = FindFProperty<FProperty>(NodeProperty->Struct, PropertyBinding.Key))
+			// for array properties we need to account for the extra FName number 
+			FName ComparisonName = PropertyBinding.Key;
+			ComparisonName.SetNumber(0);
+
+			if (FProperty* Property = FindFProperty<FProperty>(NodeProperty->Struct, ComparisonName))
 			{
 				EvalHandler.RegisterPropertyBinding(Property, PropertyBinding.Value);
 			}
 			else
 			{
-				InCompilationContext.GetMessageLog().Warning(*FString::Printf(TEXT("ICE: @@ Failed to find a property '%s'"), *PropertyBinding.Key.ToString()), InNode);
+				InCompilationContext.GetMessageLog().Warning(*FString::Printf(TEXT("ICE: @@ Failed to find a property '%s'"), *ComparisonName.ToString()), InNode);
 			}
 		}
 	}
@@ -537,11 +557,7 @@ void UAnimBlueprintExtension_Base::CreateEvaluationHandler(IAnimBlueprintCompila
 		{
 			BP_SCOPED_COMPILER_EVENT_STAT(EAnimBlueprintCompilerStats_CreateEvaluationHandler_CreateInstanceAssignmentNode);
 			
-			// Spawn 'manually' as we dont want the original node to be included in the backtrack map for this member set - the struct
-			// is referenced by multiple nodes
-			InstanceAssignmentNode = InNode->GetGraph()->CreateIntermediateNode<UK2Node_StructMemberSet>();
-			InstanceAssignmentNode->CreateNewGuid();
-			
+			InstanceAssignmentNode = InCompilationContext.SpawnIntermediateNode<UK2Node_StructMemberSet>(InNode, InCompilationContext.GetConsolidatedEventGraph());
 			InstanceAssignmentNode->VariableReference.SetSelfMember(MutableDataProperty->GetFName());
 			InstanceAssignmentNode->StructType = MutableDataProperty->Struct;
 
@@ -742,9 +758,9 @@ void UAnimBlueprintExtension_Base::FEvaluationHandlerRecord::PatchFunctionNameAn
 			{
 				// Only unbatched copies can be processed on a per-node basis
 				// Skip invalid copy indices as these are usually the result of BP errors/warnings
-				if(PropertyCopyRecord.LibraryCopyIndex != INDEX_NONE && PropertyCopyRecord.LibraryBatchType == EPropertyAccessBatchType::Unbatched)
+				if(PropertyCopyRecord.LibraryCompiledHandle.IsValid() && PropertyCopyRecord.LibraryCompiledHandle.GetBatchId() == (int32)EAnimPropertyAccessCallSite::WorkerThread_Unbatched)
 				{
-					Handler.CopyRecords.Emplace(PropertyCopyRecord.LibraryCopyIndex, PropertyCopyRecord.Operation);
+					Handler.CopyRecords.Emplace(PropertyCopyRecord.LibraryCompiledHandle.GetId(), PropertyCopyRecord.Operation);
 				}
 			}
 		}
@@ -858,9 +874,19 @@ void UAnimBlueprintExtension_Base::FEvaluationHandlerRecord::RegisterPropertyBin
 		DestPropertyPath.Add(NodeVariableProperty->GetName());
 	}
 
-	DestPropertyPath.Add(InProperty->GetName());
-
-	Handler.CopyRecords.Emplace(InBinding.PropertyPath, DestPropertyPath);
+	if(InBinding.ArrayIndex != INDEX_NONE)
+	{
+		DestPropertyPath.Add(FString::Printf(TEXT("%s[%d]"), *InProperty->GetName(), InBinding.ArrayIndex));
+	}
+	else
+	{
+		DestPropertyPath.Add(InProperty->GetName());
+	}
+	
+	FPropertyCopyRecord& CopyRecord = Handler.CopyRecords.Emplace_GetRef(InBinding.PropertyPath, DestPropertyPath);
+	CopyRecord.DestProperty = InProperty;
+	CopyRecord.DestArrayIndex = InBinding.ArrayIndex;
+	CopyRecord.BindingContextId = InBinding.ContextId;
 }
 
 void UAnimBlueprintExtension_Base::FEvaluationHandlerRecord::BuildFastPathCopyRecords(IAnimBlueprintPostExpansionStepContext& InCompilationContext)
diff --git a/Engine/Source/Editor/AnimGraph/Private/AnimBlueprintExtension_Base.h b/Engine/Source/Editor/AnimGraph/Private/AnimBlueprintExtension_Base.h
--- a/Engine/Source/Editor/AnimGraph/Private/AnimBlueprintExtension_Base.h
+++ b/Engine/Source/Editor/AnimGraph/Private/AnimBlueprintExtension_Base.h
@@ -26,8 +26,11 @@ class UAnimBlueprintExtension_Base : public UAnimBlueprintExtension
 	GENERATED_BODY()
 
 public:
-	// Adds a map of struct eval handlers for the specified node
-	void AddStructEvalHandlers(UAnimGraphNode_Base* InNode, IAnimBlueprintCompilationContext& InCompilationContext, IAnimBlueprintGeneratedClassCompiledData& OutCompiledData);
+	// Processes a node's pins:
+	// - Processes pose pins, building a map of pose links for later processing
+	// - Adds a map of struct eval handlers for the specified node
+	// - Builds property binding data
+	void ProcessNodePins(UAnimGraphNode_Base* InNode, IAnimBlueprintCompilationContext& InCompilationContext, IAnimBlueprintGeneratedClassCompiledData& OutCompiledData);
 
 	// Create an 'expanded' evaluation handler for the specified node, called in the compiler's node expansion step
 	void CreateEvaluationHandlerForNode(IAnimBlueprintCompilationContext& InCompilationContext, UAnimGraphNode_Base* InNode);
@@ -51,8 +54,7 @@ private:
 			, DestProperty(InDestProperty)
 			, DestArrayIndex(InDestArrayIndex)
 			, DestPropertyPath(MoveTemp(InDestPropertyPath))
-			, LibraryCopyIndex(INDEX_NONE)
-			, LibraryBatchType(EPropertyAccessBatchType::Unbatched)
+			, BindingContextId(NAME_None)
 			, Operation(EPostCopyOperation::None)
 			, bIsFastPath(true)
 		{}
@@ -63,8 +65,7 @@ private:
 			, DestArrayIndex(INDEX_NONE)
 			, SourcePropertyPath(InSourcePropertyPath)
 			, DestPropertyPath(InDestPropertyPath)
-			, LibraryCopyIndex(INDEX_NONE)
-			, LibraryBatchType(EPropertyAccessBatchType::Unbatched)
+			, BindingContextId(NAME_None)
 			, Operation(EPostCopyOperation::None)
 			, bIsFastPath(true)
 		{}
@@ -94,12 +95,15 @@ private:
 		/** The property path relative to the class */
 		TArray<FString> DestPropertyPath;
 
-		/** The index of the copy in the property access library */
-		int32 LibraryCopyIndex;
+		/** The handle of the copy in the property access compiler */
+		FPropertyAccessHandle LibraryHandle;
 
-		/** the batch type within the property access library */
-		EPropertyAccessBatchType LibraryBatchType;
+		/** The handle of the copy in the property access library */
+		FCompiledPropertyAccessHandle LibraryCompiledHandle;
 
+		// Context in which a property binding occurs
+		FName BindingContextId;
+		
 		/** Any operation we want to perform post-copy on the destination data */
 		EPostCopyOperation Operation;
 
diff --git a/Engine/Source/Editor/AnimGraph/Private/AnimBlueprintExtension_CallFunction.cpp b/Engine/Source/Editor/AnimGraph/Private/AnimBlueprintExtension_CallFunction.cpp
new file mode 100644
--- /dev/null
+++ b/Engine/Source/Editor/AnimGraph/Private/AnimBlueprintExtension_CallFunction.cpp
@@ -0,0 +1,24 @@
+// Copyright Epic Games, Inc. All Rights Reserved.
+
+#include "AnimBlueprintExtension_CallFunction.h"
+#include "AnimGraphNode_Base.h"
+
+void UAnimBlueprintExtension_CallFunction::HandleStartCompilingClass(const UClass* InClass, IAnimBlueprintCompilationBracketContext& InCompilationContext, IAnimBlueprintGeneratedClassCompiledData& OutCompiledData)
+{
+	Counter = 0;
+	CustomEventNames.Empty();
+}
+
+FName UAnimBlueprintExtension_CallFunction::AddCustomEventName(UAnimGraphNode_Base* InNode)
+{
+	static FName BaseName(TEXT("AnimNode_CallFunction_EventStub"));
+	FName Name(BaseName, Counter++);
+	CustomEventNames.Add(InNode, Name);
+
+	return Name;
+}
+
+FName UAnimBlueprintExtension_CallFunction::GetCustomEventName(UAnimGraphNode_Base* InNode) const
+{
+	return CustomEventNames.FindChecked(InNode);
+}
diff --git a/Engine/Source/Editor/AnimGraph/Private/AnimBlueprintExtension_CallFunction.h b/Engine/Source/Editor/AnimGraph/Private/AnimBlueprintExtension_CallFunction.h
new file mode 100644
--- /dev/null
+++ b/Engine/Source/Editor/AnimGraph/Private/AnimBlueprintExtension_CallFunction.h
@@ -0,0 +1,31 @@
+// Copyright Epic Games, Inc. All Rights Reserved.
+
+#pragma once
+
+#include "CoreMinimal.h"
+#include "AnimBlueprintExtension.h"
+#include "AnimBlueprintExtension_CallFunction.generated.h"
+
+class UAnimGraphNode_Base;
+
+UCLASS()
+class UAnimBlueprintExtension_CallFunction : public UAnimBlueprintExtension
+{
+	GENERATED_BODY()
+
+public:
+	virtual void HandleStartCompilingClass(const UClass* InClass, IAnimBlueprintCompilationBracketContext& InCompilationContext, IAnimBlueprintGeneratedClassCompiledData& OutCompiledData) override;
+	
+	// Add a custom event name for the specified node
+	FName AddCustomEventName(UAnimGraphNode_Base* InNode);
+
+	// Get the custom event name for the specified node
+	FName GetCustomEventName(UAnimGraphNode_Base* InNode) const;
+
+private:
+	// Counter to allow us to create unique stub function names
+	int32 Counter = 0;
+	
+	// Set of used custom event names per-node
+	TMap<UAnimGraphNode_Base*, FName> CustomEventNames;
+};
\ No newline at end of file
diff --git a/Engine/Source/Editor/AnimGraph/Private/AnimBlueprintExtension_PropertyAccess.cpp b/Engine/Source/Editor/AnimGraph/Private/AnimBlueprintExtension_PropertyAccess.cpp
--- a/Engine/Source/Editor/AnimGraph/Private/AnimBlueprintExtension_PropertyAccess.cpp
+++ b/Engine/Source/Editor/AnimGraph/Private/AnimBlueprintExtension_PropertyAccess.cpp
@@ -1,26 +1,132 @@
 // Copyright Epic Games, Inc. All Rights Reserved.
 
 #include "AnimBlueprintExtension_PropertyAccess.h"
+
+#include "EdGraphSchema_K2.h"
 #include "IPropertyAccessEditor.h"
 #include "Features/IModularFeatures.h"
 #include "IPropertyAccessCompiler.h"
 #include "IAnimBlueprintCompilationBracketContext.h"
+#include "K2Node_BreakStruct.h"
+#include "K2Node_GetArrayItem.h"
+#include "K2Node_VariableGet.h"
+#include "K2Node_CallFunction.h"
+#include "KismetCompiler.h"
 #include "Kismet2/CompilerResultsLog.h"
 
-int32 UAnimBlueprintExtension_PropertyAccess::AddCopy(TArrayView<FString> InSourcePath, TArrayView<FString> InDestPath, EPropertyAccessBatchType InBatchType, UObject* InObject)
+#define LOCTEXT_NAMESPACE "UAnimBlueprintExtension_PropertyAccess"
+
+const FName UAnimBlueprintExtension_PropertyAccess::ContextId_Automatic(NAME_None);
+const FName UAnimBlueprintExtension_PropertyAccess::ContextId_UnBatched_ThreadSafe(TEXT("UnBatched_ThreadSafe"));
+const FName UAnimBlueprintExtension_PropertyAccess::ContextId_Batched_WorkerThreadPreEventGraph(TEXT("Batched_WorkerThreadPreEventGraph"));
+const FName UAnimBlueprintExtension_PropertyAccess::ContextId_Batched_WorkerThreadPostEventGraph(TEXT("Batched_WorkerThreadPostEventGraph"));
+const FName UAnimBlueprintExtension_PropertyAccess::ContextId_Batched_GameThreadPreEventGraph(TEXT("Batched_GameThreadPreEventGraph"));
+const FName UAnimBlueprintExtension_PropertyAccess::ContextId_Batched_GameThreadPostEventGraph(TEXT("Batched_GameThreadPostEventGraph"));
+
+FPropertyAccessHandle UAnimBlueprintExtension_PropertyAccess::AddCopy(TArrayView<FString> InSourcePath, TArrayView<FString> InDestPath, const FName& InContextId, UObject* InObject)
 {
 	if(PropertyAccessLibraryCompiler.IsValid())
 	{
-		return PropertyAccessLibraryCompiler->AddCopy(InSourcePath, InDestPath, InBatchType, InObject);
+		return PropertyAccessLibraryCompiler->AddCopy(InSourcePath, InDestPath, InContextId, InObject);
 	}
-	return INDEX_NONE;
+	return FPropertyAccessHandle();
 }
 
 void UAnimBlueprintExtension_PropertyAccess::HandleStartCompilingClass(const UClass* InClass, IAnimBlueprintCompilationBracketContext& InCompilationContext, IAnimBlueprintGeneratedClassCompiledData& OutCompiledData)
 {
 	IPropertyAccessEditor& PropertyAccessEditor = IModularFeatures::Get().GetModularFeature<IPropertyAccessEditor>("PropertyAccessEditor");
+
+	FCompilerResultsLog& MessageLog = InCompilationContext.GetMessageLog();
+
+	FPropertyAccessLibraryCompilerArgs Args(Subsystem.Library, InClass);
+	Args.OnDetermineBatchId = FOnPropertyAccessDetermineBatchId::CreateLambda([&MessageLog](const FPropertyAccessCopyContext& InContext) -> int32
+	{
+		if(InContext.ContextId == ContextId_Automatic)
+		{
+			if(InContext.bSourceThreadSafe && InContext.bDestThreadSafe)
+			{
+				// Can only be in the worker thread batch if both endpoints are thread-safe
+				return (int32)EAnimPropertyAccessCallSite::WorkerThread_Unbatched;
+			}
+			else
+			{
+				return (int32)EAnimPropertyAccessCallSite::GameThread_Batched_PreEventGraph;
+			}
+		}
+		else if(InContext.ContextId == ContextId_UnBatched_ThreadSafe)
+		{
+			if(InContext.bSourceThreadSafe && InContext.bDestThreadSafe)
+			{
+				// Can only be in the worker thread batch if both endpoints are thread-safe
+				return (int32)EAnimPropertyAccessCallSite::WorkerThread_Unbatched;
+			}
+			else
+			{
+				if(!InContext.bSourceThreadSafe)
+				{
+					MessageLog.Warning(*FText::Format(LOCTEXT("ThreadSafetyIncompatible", "@@ '{0}' is not thread-safe, access will be performed on the game thread (pre-event graph) and cached"), InContext.SourcePathAsText).ToString(), InContext.Object);
+				}
+				if(!InContext.bDestThreadSafe)
+				{
+					MessageLog.Warning(*FText::Format(LOCTEXT("ThreadSafetyIncompatible", "@@ '{0}' is not thread-safe, access will be performed on the game thread (pre-event graph) and cached"), InContext.DestPathAsText).ToString(), InContext.Object);
+				}
+				return (int32)EAnimPropertyAccessCallSite::GameThread_Batched_PreEventGraph;
+			}
+		}
+		else if(InContext.ContextId == ContextId_Batched_GameThreadPreEventGraph)
+		{
+			return (int32)EAnimPropertyAccessCallSite::GameThread_Batched_PreEventGraph;
+		}
+		else if(InContext.ContextId == ContextId_Batched_GameThreadPostEventGraph)
+		{
+			return (int32)EAnimPropertyAccessCallSite::GameThread_Batched_PostEventGraph;
+		}
+		else if(InContext.ContextId == ContextId_Batched_WorkerThreadPreEventGraph)
+		{
+			if(InContext.bSourceThreadSafe && InContext.bDestThreadSafe)
+			{
+				return (int32)EAnimPropertyAccessCallSite::WorkerThread_Batched_PreEventGraph;
+			}
+			else
+			{
+				if(!InContext.bSourceThreadSafe)
+				{
+					MessageLog.Warning(*FText::Format(LOCTEXT("ThreadSafetyIncompatible", "@@ '{0}' is not thread-safe, access will be performed on the game thread (pre-event graph) and cached"), InContext.SourcePathAsText).ToString(), InContext.Object);
+				}
+				if(!InContext.bDestThreadSafe)
+				{
+					MessageLog.Warning(*FText::Format(LOCTEXT("ThreadSafetyIncompatible", "@@ '{0}' is not thread-safe, access will be performed on the game thread (pre-event graph) and cached"), InContext.DestPathAsText).ToString(), InContext.Object);
+				}
+				return (int32)EAnimPropertyAccessCallSite::GameThread_Batched_PreEventGraph;
+			}
+		}
+		else if(InContext.ContextId == ContextId_Batched_WorkerThreadPostEventGraph)
+		{
+			if(InContext.bSourceThreadSafe && InContext.bDestThreadSafe)
+			{
+				return (int32)EAnimPropertyAccessCallSite::WorkerThread_Batched_PostEventGraph;
+			}
+			else
+			{
+				if(!InContext.bSourceThreadSafe)
+				{
+					MessageLog.Warning(*FText::Format(LOCTEXT("ThreadSafetyIncompatible", "@@ '{0}' is not thread-safe, access will be performed on the game thread (pre-event graph) and cached"), InContext.SourcePathAsText).ToString(), InContext.Object);
+				}
+				if(!InContext.bDestThreadSafe)
+				{
+					MessageLog.Warning(*FText::Format(LOCTEXT("ThreadSafetyIncompatible", "@@ '{0}' is not thread-safe, access will be performed on the game thread (pre-event graph) and cached"), InContext.DestPathAsText).ToString(), InContext.Object);
+				}
+				return (int32)EAnimPropertyAccessCallSite::GameThread_Batched_PreEventGraph;
+			}
+		}	
+		else
+		{
+			MessageLog.Warning(*FText::Format(LOCTEXT("UnknownContext", "@@ has unknown context '{0}', access will be performed on the game thread (pre-event graph) and cached"), FText::FromName(InContext.ContextId)).ToString(), InContext.Object);
+			return (int32)EAnimPropertyAccessCallSite::GameThread_Batched_PreEventGraph;
+		}
+	});
 	
-	PropertyAccessLibraryCompiler = PropertyAccessEditor.MakePropertyAccessCompiler(FPropertyAccessLibraryCompilerArgs(Subsystem.Library, InClass));
+	PropertyAccessLibraryCompiler = PropertyAccessEditor.MakePropertyAccessCompiler(Args);
 	PropertyAccessLibraryCompiler->BeginCompilation();
 }
 
@@ -49,11 +155,201 @@ void UAnimBlueprintExtension_PropertyAccess::HandleFinishCompilingClass(const UC
 	PropertyAccessLibraryCompiler.Reset();
 }
 
-int32 UAnimBlueprintExtension_PropertyAccess::MapCopyIndex(int32 InIndex) const
+FCompiledPropertyAccessHandle UAnimBlueprintExtension_PropertyAccess::GetCompiledHandle(FPropertyAccessHandle InHandle) const
 {
 	if(PropertyAccessLibraryCompiler.IsValid())
 	{
-		return PropertyAccessLibraryCompiler->MapCopyIndex(InIndex);
+		return PropertyAccessLibraryCompiler->GetCompiledHandle(InHandle);
 	}
-	return INDEX_NONE;
-}
\ No newline at end of file
+	return FCompiledPropertyAccessHandle();
+}
+
+FText UAnimBlueprintExtension_PropertyAccess::GetCompiledHandleContext(FCompiledPropertyAccessHandle InHandle)
+{
+	UEnum* EnumClass = StaticEnum<EAnimPropertyAccessCallSite>();
+	check(EnumClass != nullptr);
+	return EnumClass->GetDisplayNameTextByValue(InHandle.GetBatchId());
+}
+
+FText UAnimBlueprintExtension_PropertyAccess::GetCompiledHandleContextDesc(FCompiledPropertyAccessHandle InHandle)
+{
+	UEnum* EnumClass = StaticEnum<EAnimPropertyAccessCallSite>();
+	check(EnumClass != nullptr);
+	return EnumClass->GetToolTipTextByIndex(InHandle.GetBatchId());
+}
+
+bool UAnimBlueprintExtension_PropertyAccess::ContextRequiresCachedVariable(FName InName)
+{
+	if(InName == ContextId_Automatic)
+	{
+		// Indeterminate, so assume false - calling code will need to opt-in
+		return false;
+	}
+	else if(InName == ContextId_UnBatched_ThreadSafe)
+	{
+		return false;
+	}
+	else if(InName == ContextId_Batched_WorkerThreadPreEventGraph)
+	{
+		return true;
+	}
+	else if(InName == ContextId_Batched_WorkerThreadPostEventGraph)
+	{
+		return true;
+	}	
+	else if(InName == ContextId_Batched_GameThreadPreEventGraph)
+	{
+		return true;
+	}
+	else if(InName == ContextId_Batched_GameThreadPostEventGraph)
+	{
+		return true;
+	}
+
+	return false;
+}
+
+void UAnimBlueprintExtension_PropertyAccess::ExpandPropertyAccess(FKismetCompilerContext& InCompilerContext, TArrayView<FString> InSourcePath, UEdGraph* InParentGraph, UEdGraphPin* InTargetPin) const
+{
+	check(InParentGraph);
+	check(InTargetPin);
+	
+	IPropertyAccessEditor& PropertyAccessEditor = IModularFeatures::Get().GetModularFeature<IPropertyAccessEditor>("PropertyAccessEditor");
+
+	UEdGraphNode* SourceNode = InTargetPin->GetOwningNode();
+	check(SourceNode);
+	
+	// Track the current pin in the pure chain
+	UEdGraphPin* CurrentPin = nullptr;
+
+	auto SpawnVariableGetNode = [&CurrentPin, &InCompilerContext, &SourceNode, &InParentGraph](FName InPropertyName)
+	{
+		const UEdGraphSchema* GraphSchema = InParentGraph->GetSchema();
+
+		if(CurrentPin == nullptr)
+		{
+			// No current pin means we must be at the start of a chain with a 'self' member
+			UK2Node_VariableGet* VariableGetNode = InCompilerContext.SpawnIntermediateNode<UK2Node_VariableGet>(SourceNode, InParentGraph);
+			VariableGetNode->VariableReference.SetSelfMember(InPropertyName);
+			VariableGetNode->AllocateDefaultPins();
+
+			// Find pin that we just created - variable out pin is now CurrentPin
+			CurrentPin = VariableGetNode->FindPinChecked(VariableGetNode->GetVarName());
+		}
+		else
+		{
+			// Current pin means we must be a class or a struct context
+			if(UClass* Class = Cast<UClass>(CurrentPin->PinType.PinSubCategoryObject.Get()))
+			{
+				UK2Node_VariableGet* VariableGetNode = InCompilerContext.SpawnIntermediateNode<UK2Node_VariableGet>(SourceNode, InParentGraph);
+				VariableGetNode->VariableReference.SetExternalMember(InPropertyName, Class);
+				VariableGetNode->AllocateDefaultPins();
+
+				// Link current to target, connection must succeed
+				UEdGraphPin* TargetPin = VariableGetNode->FindPinChecked(UEdGraphSchema_K2::PN_Self);
+				bool bSucceeded = GraphSchema->TryCreateConnection(CurrentPin, TargetPin);
+				if(!bSucceeded)
+				{
+					InCompilerContext.MessageLog.Error(*LOCTEXT("VariableConnectionFailed", "@@ ICE: could not connect variable when expanding node").ToString(), SourceNode);
+				}
+
+				// Find pin that we just created - variable out pin is now CurrentPin
+				CurrentPin = VariableGetNode->FindPinChecked(VariableGetNode->GetVarName());
+			}
+			else if(UScriptStruct* ScriptStruct = Cast<UScriptStruct>(CurrentPin->PinType.PinSubCategoryObject.Get()))
+			{
+				// Create a break struct/split pin node
+				const UEdGraphSchema_K2* K2Schema = CastChecked<UEdGraphSchema_K2>(InParentGraph->GetSchema());
+				UK2Node* SplitPinNode = K2Schema->CreateSplitPinNode(CurrentPin, UEdGraphSchema_K2::FCreateSplitPinNodeParams(&InCompilerContext, InParentGraph));
+
+				// Current pin is the property of the struct
+				CurrentPin = SplitPinNode->FindPinChecked(InPropertyName);
+			}
+		}
+	};
+	
+	IPropertyAccessEditor::FResolvePropertyAccessArgs Args;
+	Args.PropertyFunction = [&CurrentPin, &InCompilerContext, &SourceNode, &InParentGraph, &SpawnVariableGetNode](int32 InSegmentIndex, FProperty* InProperty, int32 InStaticArrayIndex)
+	{
+		SpawnVariableGetNode(InProperty->GetFName());
+	};
+
+	Args.ArrayFunction = [&CurrentPin, &InCompilerContext, &SourceNode, &InParentGraph, &SpawnVariableGetNode](int32 InSegmentIndex, FArrayProperty* InProperty, int32 InArrayIndex)
+	{
+		// First spawn a variable get to get the array into CurrentPin
+		SpawnVariableGetNode(InProperty->GetFName());
+
+		// Spawn an array item getter
+		UK2Node_GetArrayItem* GetArrayItemNode = InCompilerContext.SpawnIntermediateNode<UK2Node_GetArrayItem>(SourceNode, InParentGraph);
+		GetArrayItemNode->AllocateDefaultPins();
+
+		UEdGraphPin* ArrayPin = GetArrayItemNode->GetTargetArrayPin();
+		UEdGraphPin* ResultPin = GetArrayItemNode->GetResultPin();
+		UEdGraphPin* IndexPin = GetArrayItemNode->GetIndexPin();
+		check(ArrayPin && ResultPin && IndexPin);
+
+		// Fill in the array index pin default value
+		IndexPin->DefaultValue = FString::FromInt(InArrayIndex); 
+
+		// Connect up the array variable output (CurrentPin) to the array pin
+		const UEdGraphSchema* GraphSchema = InParentGraph->GetSchema();
+		check(CurrentPin);
+		bool bSucceeded = GraphSchema->TryCreateConnection(CurrentPin, ArrayPin);
+		if(!bSucceeded)
+		{
+			InCompilerContext.MessageLog.Error(*LOCTEXT("ArrayConnectionFailed", "@@ ICE: could not connect array when expanding node").ToString(), SourceNode);
+		}
+
+		// Array element is now the current pin
+		CurrentPin = ResultPin;
+	};
+
+	Args.FunctionFunction = [&CurrentPin, &InCompilerContext, &SourceNode, &InParentGraph](int32 InSegmentIndex, UFunction* InFunction, FProperty* InReturnProperty)
+	{
+		// Spawn a function call
+		UK2Node_CallFunction* CallFunctionNode = InCompilerContext.SpawnIntermediateNode<UK2Node_CallFunction>(SourceNode, InParentGraph);
+		CallFunctionNode->SetFromFunction(InFunction);
+		CallFunctionNode->AllocateDefaultPins();
+
+		if(CurrentPin)
+		{
+			// If we have a current pin, hook it up to self
+			UEdGraphPin* SelfPin = CallFunctionNode->FindPinChecked(UEdGraphSchema_K2::PN_Self);
+
+			const UEdGraphSchema* GraphSchema = InParentGraph->GetSchema();
+			bool bSucceeded = GraphSchema->TryCreateConnection(CurrentPin, SelfPin);
+			if(!bSucceeded)
+			{
+				InCompilerContext.MessageLog.Error(*LOCTEXT("FunctionConnectionFailed", "@@ ICE: could not connect function when expanding node").ToString(), SourceNode);
+			}
+		}
+
+		// The new current pin is the return value
+		UEdGraphPin* ReturnValuePin = CallFunctionNode->GetReturnValuePin();
+		check(ReturnValuePin);
+		CurrentPin = ReturnValuePin;
+	};
+
+	FPropertyAccessResolveResult Result = PropertyAccessEditor.ResolvePropertyAccess(InCompilerContext.Blueprint->SkeletonGeneratedClass, InSourcePath, Args);
+	if(Result.Result == EPropertyAccessResolveResult::Succeeded)
+	{
+		const UEdGraphSchema* GraphSchema = InParentGraph->GetSchema();
+
+		// Link current pin to target
+		check(CurrentPin);
+		if(InTargetPin->Direction == CurrentPin->Direction)
+		{
+			InCompilerContext.MovePinLinksToIntermediate(*InTargetPin, *CurrentPin);
+		}
+		else
+		{
+			bool bSucceeded = GraphSchema->TryCreateConnection(CurrentPin, InTargetPin);
+			if(!bSucceeded)
+			{
+				InCompilerContext.MessageLog.Error(*LOCTEXT("TargetConnectionFailed", "@@ ICE: could not connect target when expanding node").ToString(), SourceNode);
+			}
+		}
+	}
+}
+
+#undef LOCTEXT_NAMESPACE
\ No newline at end of file
diff --git a/Engine/Source/Editor/AnimGraph/Private/AnimGraphModule.cpp b/Engine/Source/Editor/AnimGraph/Private/AnimGraphModule.cpp
--- a/Engine/Source/Editor/AnimGraph/Private/AnimGraphModule.cpp
+++ b/Engine/Source/Editor/AnimGraph/Private/AnimGraphModule.cpp
@@ -1,6 +1,6 @@
 // Copyright Epic Games, Inc. All Rights Reserved.
 
-#include "AnimGraphModule.h"
+#include "CoreMinimal.h"
 #include "Textures/SlateIcon.h"
 #include "AnimGraphCommands.h"
 #include "Modules/ModuleManager.h"
@@ -29,6 +29,18 @@
 #include "Animation/AnimSync.h"
 #include "Animation/AnimNode_Inertialization.h"
 #include "Animation/AnimRootMotionProvider.h"
+#include "Features/IModularFeatures.h"
+#include "PropertyAccessAnimBlueprintBinding.h"
+
+class FAnimGraphModule : public IModuleInterface
+{
+public:
+	/** IModuleInterface interface */
+	virtual void StartupModule() override;
+	virtual void ShutdownModule() override;
+
+	FPropertyAccessAnimBlueprintBinding PropertyAccessAnimBlueprintBinding;
+};
 
 IMPLEMENT_MODULE(FAnimGraphModule, AnimGraph);
 
@@ -43,6 +55,8 @@ void FAnimGraphModule::StartupModule()
 		return MakeShared<FAnimBlueprintCompilerContext>(CastChecked<UAnimBlueprint>(InBlueprint), InMessageLog, InCompileOptions);
 	});
 
+	IModularFeatures::Get().RegisterModularFeature("PropertyAccessBlueprintBinding", &PropertyAccessAnimBlueprintBinding);
+
 	// Register the editor modes
 	FEditorModeRegistry::Get().RegisterMode<FAnimNodeEditMode>(AnimNodeEditModes::AnimNode, LOCTEXT("AnimNodeEditMode", "Anim Node"), FSlateIcon(), false);
 	FEditorModeRegistry::Get().RegisterMode<FTwoBoneIKEditMode>(AnimNodeEditModes::TwoBoneIK, LOCTEXT("TwoBoneIKEditMode", "2-Bone IK"), FSlateIcon(), false);
@@ -60,76 +74,79 @@ void FAnimGraphModule::StartupModule()
 
 	PropertyModule.RegisterCustomPropertyTypeLayout("AnimBlueprintFunctionPinInfo", FOnGetPropertyTypeCustomizationInstance::CreateStatic(&FAnimBlueprintFunctionPinInfoDetails::MakeInstance));
 
-	// Register built-in graph attributes
-	UAnimGraphAttributes* AnimGraphAttributes = GetMutableDefault<UAnimGraphAttributes>();
-	AnimGraphAttributes->LoadConfig();
-
-	AnimGraphAttributes->Register(
-		FAnimGraphAttributeDesc(
-			UE::Anim::FAttributes::Pose,
-			EAnimGraphAttributeBlend::Blendable,
-			FSlateBrush(),
-			LOCTEXT("Pose", "Pose"),
-			LOCTEXT("PoseToolTip", "Pose: Transforms for all bones in the current LOD, blendable"),
-			FEditorStyle::GetSlateColor("AnimGraph.Attribute.Pose.Color"),
-			EAnimGraphAttributesDisplayMode::HideOnPins)
-	);
-
-	AnimGraphAttributes->Register(
-		FAnimGraphAttributeDesc(
-			UE::Anim::FAttributes::Curves,
-			EAnimGraphAttributeBlend::Blendable,
-			*FEditorStyle::GetBrush("AnimGraph.Attribute.Curves.Icon"),
-			LOCTEXT("Curves", "Curves"),
-			LOCTEXT("CurvesToolTip", "Curves: Floating point curve attributes, blendable"),
-			FEditorStyle::GetSlateColor("AnimGraph.Attribute.Curves.Color"),
-			EAnimGraphAttributesDisplayMode::Automatic)
-	);
-
-	AnimGraphAttributes->Register(
-		FAnimGraphAttributeDesc(
-			UE::Anim::FAttributes::Attributes,
-			EAnimGraphAttributeBlend::Blendable,
-			*FEditorStyle::GetBrush("AnimGraph.Attribute.Attributes.Icon"),
-			LOCTEXT("Attributes", "Attributes"),
-			LOCTEXT("AttributesToolTip", "Attributes: Misc. custom attributes, blendable"),
-			FEditorStyle::GetSlateColor("AnimGraph.Attribute.Attributes.Color"),
-			EAnimGraphAttributesDisplayMode::Automatic)
-	);
-
-	AnimGraphAttributes->Register(
-		FAnimGraphAttributeDesc(
-			UE::Anim::FAnimSync::Attribute,
-			EAnimGraphAttributeBlend::NonBlendable,
-			*FEditorStyle::Get().GetBrush("AnimGraph.Attribute.Sync.Icon"),
-			LOCTEXT("Sync", "Sync"),
-			LOCTEXT("SyncToolTip", "Sync: Synchronization between the playback of different assets present in the graph, not blendable"),
-			FEditorStyle::Get().GetSlateColor("AnimGraph.Attribute.Sync.Color"),
-			EAnimGraphAttributesDisplayMode::Automatic)
-	);
-
-	AnimGraphAttributes->Register(
-		FAnimGraphAttributeDesc(
-			UE::Anim::IAnimRootMotionProvider::RootMotionDeltaAttributeName,
-			EAnimGraphAttributeBlend::Blendable,
-			*FEditorStyle::Get().GetBrush("AnimGraph.Attribute.RootMotionDelta.Icon"),
-			LOCTEXT("RootMotionDelta", "Root Motion Delta"),
-			LOCTEXT("RootMotionDeltaToolTip", "Root Motion Delta: Per-frame root motion delta transform, blendable"),
-			FEditorStyle::Get().GetSlateColor("AnimGraph.Attribute.RootMotionDelta.Color"),
-			EAnimGraphAttributesDisplayMode::Automatic)
-	);
-
-	AnimGraphAttributes->Register(
-		FAnimGraphAttributeDesc(
-			UE::Anim::IInertializationRequester::Attribute,
-			EAnimGraphAttributeBlend::NonBlendable,
-			*FEditorStyle::Get().GetBrush("AnimGraph.Attribute.InertialBlending.Icon"),
-			LOCTEXT("InertialBlending", "Inertialize"),
-			LOCTEXT("InertialBlendingToolTip", "Inertial Blending: Used to blend between poses preserving bone velocities, not blendable"),
-			FEditorStyle::Get().GetSlateColor("AnimGraph.Attribute.InertialBlending.Color"),
-			EAnimGraphAttributesDisplayMode::Automatic)
-	);
-
+	if(GIsEditor)
+	{
+		// Register built-in graph attributes
+		UAnimGraphAttributes* AnimGraphAttributes = GetMutableDefault<UAnimGraphAttributes>();
+		AnimGraphAttributes->LoadConfig();
+
+		AnimGraphAttributes->Register(
+			FAnimGraphAttributeDesc(
+				UE::Anim::FAttributes::Pose,
+				EAnimGraphAttributeBlend::Blendable,
+				FSlateBrush(),
+				LOCTEXT("Pose", "Pose"),
+				LOCTEXT("PoseToolTip", "Pose: Transforms for all bones in the current LOD, blendable"),
+				FEditorStyle::GetSlateColor("AnimGraph.Attribute.Pose.Color"),
+				EAnimGraphAttributesDisplayMode::HideOnPins)
+		);
+
+		AnimGraphAttributes->Register(
+			FAnimGraphAttributeDesc(
+				UE::Anim::FAttributes::Curves,
+				EAnimGraphAttributeBlend::Blendable,
+				*FEditorStyle::GetBrush("AnimGraph.Attribute.Curves.Icon"),
+				LOCTEXT("Curves", "Curves"),
+				LOCTEXT("CurvesToolTip", "Curves: Floating point curve attributes, blendable"),
+				FEditorStyle::GetSlateColor("AnimGraph.Attribute.Curves.Color"),
+				EAnimGraphAttributesDisplayMode::Automatic)
+		);
+
+		AnimGraphAttributes->Register(
+			FAnimGraphAttributeDesc(
+				UE::Anim::FAttributes::Attributes,
+				EAnimGraphAttributeBlend::Blendable,
+				*FEditorStyle::GetBrush("AnimGraph.Attribute.Attributes.Icon"),
+				LOCTEXT("Attributes", "Attributes"),
+				LOCTEXT("AttributesToolTip", "Attributes: Misc. custom attributes, blendable"),
+				FEditorStyle::GetSlateColor("AnimGraph.Attribute.Attributes.Color"),
+				EAnimGraphAttributesDisplayMode::Automatic)
+		);
+
+		AnimGraphAttributes->Register(
+			FAnimGraphAttributeDesc(
+				UE::Anim::FAnimSync::Attribute,
+				EAnimGraphAttributeBlend::NonBlendable,
+				*FEditorStyle::Get().GetBrush("AnimGraph.Attribute.Sync.Icon"),
+				LOCTEXT("Sync", "Sync"),
+				LOCTEXT("SyncToolTip", "Sync: Synchronization between the playback of different assets present in the graph, not blendable"),
+				FEditorStyle::Get().GetSlateColor("AnimGraph.Attribute.Sync.Color"),
+				EAnimGraphAttributesDisplayMode::Automatic)
+		);
+
+		AnimGraphAttributes->Register(
+			FAnimGraphAttributeDesc(
+				UE::Anim::IAnimRootMotionProvider::RootMotionDeltaAttributeName,
+				EAnimGraphAttributeBlend::Blendable,
+				*FEditorStyle::Get().GetBrush("AnimGraph.Attribute.RootMotionDelta.Icon"),
+				LOCTEXT("RootMotionDelta", "Root Motion Delta"),
+				LOCTEXT("RootMotionDeltaToolTip", "Root Motion Delta: Per-frame root motion delta transform, blendable"),
+				FEditorStyle::Get().GetSlateColor("AnimGraph.Attribute.RootMotionDelta.Color"),
+				EAnimGraphAttributesDisplayMode::Automatic)
+		);
+
+		AnimGraphAttributes->Register(
+			FAnimGraphAttributeDesc(
+				UE::Anim::IInertializationRequester::Attribute,
+				EAnimGraphAttributeBlend::NonBlendable,
+				*FEditorStyle::Get().GetBrush("AnimGraph.Attribute.InertialBlending.Icon"),
+				LOCTEXT("InertialBlending", "Inertialize"),
+				LOCTEXT("InertialBlendingToolTip", "Inertial Blending: Used to blend between poses preserving bone velocities, not blendable"),
+				FEditorStyle::Get().GetSlateColor("AnimGraph.Attribute.InertialBlending.Color"),
+				EAnimGraphAttributesDisplayMode::Automatic)
+		);
+	}
+	
 	// Register BP-editor function customization once the Kismet module is loaded.
 	if (FModuleManager::Get().IsModuleLoaded("Kismet"))
 	{
@@ -161,6 +178,8 @@ void FAnimGraphModule::ShutdownModule()
 	FEditorModeRegistry::Get().UnregisterMode(AnimNodeEditModes::TwoBoneIK);
 	FEditorModeRegistry::Get().UnregisterMode(AnimNodeEditModes::AnimNode);
 
+	IModularFeatures::Get().UnregisterModularFeature("PropertyAccessBlueprintBinding", &PropertyAccessAnimBlueprintBinding);
+	
 	// Unregister details customization
 	if (UObjectInitialized() && FModuleManager::Get().IsModuleLoaded(TEXT("PropertyEditor")))
 	{
diff --git a/Engine/Source/Editor/AnimGraph/Private/AnimGraphNode_Base.cpp b/Engine/Source/Editor/AnimGraph/Private/AnimGraphNode_Base.cpp
--- a/Engine/Source/Editor/AnimGraph/Private/AnimGraphNode_Base.cpp
+++ b/Engine/Source/Editor/AnimGraph/Private/AnimGraphNode_Base.cpp
@@ -25,7 +25,18 @@
 #include "AnimBlueprintCompilationContext.h"
 #include "AnimBlueprintExtension.h"
 #include "FindInBlueprintManager.h"
+#include "IPropertyAccessBlueprintBinding.h"
+#include "IPropertyAccessEditor.h"
+#include "ScopedTransaction.h"
+#include "Algo/Accumulate.h"
+#include "Fonts/FontMeasure.h"
 #include "UObject/ReleaseObjectVersion.h"
+#include "Framework/MultiBox/MultiBoxBuilder.h"
+#include "Widgets/Layout/SSpacer.h"
+#include "Widgets/Layout/SBox.h"
+#include "Widgets/SOverlay.h"
+#include "Widgets/Text/STextBlock.h"
+#include "Fonts/FontMeasure.h"
 
 #define LOCTEXT_NAMESPACE "UAnimGraphNode_Base"
 
@@ -80,6 +91,19 @@ void UAnimGraphNode_Base::PostEditChangeProperty(FPropertyChangedEvent& Property
 	PropertyChangeEvent.Broadcast(PropertyChangedEvent);
 }
 
+void UAnimGraphNode_Base::SetPinVisibility(bool bInVisible, int32 InOptionalPinIndex)
+{
+	if(ShowPinForProperties[InOptionalPinIndex].bShowPin != bInVisible)
+	{
+		FOptionalPinManager::CacheShownPins(ShowPinForProperties, OldShownPins);
+
+		ShowPinForProperties[InOptionalPinIndex].bShowPin = bInVisible;
+		
+		FOptionalPinManager::EvaluateOldShownPins(ShowPinForProperties, OldShownPins, this);
+		GetSchema()->ReconstructNode(*this);
+	}
+}
+
 void UAnimGraphNode_Base::Serialize(FArchive& Ar)
 {
 	Super::Serialize(Ar);
@@ -431,7 +455,7 @@ void UAnimGraphNode_Base::ProcessDuringCompilation(IAnimBlueprintCompilationCont
 	UAnimBlueprintExtension_Base* Extension = UAnimBlueprintExtension::GetExtension<UAnimBlueprintExtension_Base>(GetAnimBlueprint());
 
 	// Record pose pins for later patchup and gather pins that have an associated evaluation handler
-	Extension->AddStructEvalHandlers(this, InCompilationContext, OutCompiledData);
+	Extension->ProcessNodePins(this, InCompilationContext, OutCompiledData);
 
 	// Call the override point
 	OnProcessDuringCompilation(InCompilationContext, OutCompiledData);
@@ -612,4 +636,606 @@ void UAnimGraphNode_Base::AutowireNewNode(UEdGraphPin* FromPin)
 	}
 }
 
+TSharedRef<SWidget> UAnimGraphNode_Base::MakePropertyBindingWidget(const FAnimPropertyBindingWidgetArgs& InArgs)
+{
+	UAnimGraphNode_Base* FirstAnimGraphNode = InArgs.Nodes[0];
+	UBlueprint* Blueprint = FirstAnimGraphNode->GetAnimBlueprint();
+	const bool bMultiSelect = InArgs.Nodes.Num() > 1;
+	
+	if(IModularFeatures::Get().IsModularFeatureAvailable("PropertyAccessEditor"))
+	{
+		int32 PinArrayIndex = InArgs.PinName.GetNumber() - 1;
+		const bool bIsArrayOrArrayElement = InArgs.PinProperty->IsA<FArrayProperty>();
+		const bool bIsArrayElement = bIsArrayOrArrayElement && PinArrayIndex != INDEX_NONE;
+		const bool bIsArray = bIsArrayOrArrayElement && PinArrayIndex == INDEX_NONE;
+		
+		FProperty* BindingProperty = bIsArrayElement ? CastField<FArrayProperty>(InArgs.PinProperty)->Inner : InArgs.PinProperty;
+
+		auto OnCanBindProperty = [BindingProperty](FProperty* InProperty)
+		{
+			// Note: We support type promotion here
+			IPropertyAccessEditor& PropertyAccessEditor = IModularFeatures::Get().GetModularFeature<IPropertyAccessEditor>("PropertyAccessEditor");
+			return PropertyAccessEditor.GetPropertyCompatibility(InProperty, BindingProperty) != EPropertyAccessCompatibility::Incompatible;
+		};
+
+		auto OnCanBindFunction = [BindingProperty](UFunction* InFunction)
+		{
+			IPropertyAccessEditor& PropertyAccessEditor = IModularFeatures::Get().GetModularFeature<IPropertyAccessEditor>("PropertyAccessEditor");
+
+			// Note: We support type promotion here
+			return InFunction->NumParms == 1 
+				&& PropertyAccessEditor.GetPropertyCompatibility(InFunction->GetReturnProperty(), BindingProperty) != EPropertyAccessCompatibility::Incompatible
+				&& InFunction->HasAnyFunctionFlags(FUNC_BlueprintPure);
+		};
+
+		auto OnAddBinding = [InArgs, Blueprint, BindingProperty](FName InPropertyName, const TArray<FBindingChainElement>& InBindingChain)
+		{
+			const UEdGraphSchema_K2* Schema = GetDefault<UEdGraphSchema_K2>();
+			IPropertyAccessEditor& PropertyAccessEditor = IModularFeatures::Get().GetModularFeature<IPropertyAccessEditor>("PropertyAccessEditor");
+
+			for(UAnimGraphNode_Base* AnimGraphNode : InArgs.Nodes)
+			{
+				AnimGraphNode->Modify();
+
+				// Pins are exposed if we have a binding or not - and after running this we do.
+				AnimGraphNode->SetPinVisibility(true, InArgs.OptionalPinIndex);
+
+				// Need to break all pin links now we have a binding
+				if(UEdGraphPin* Pin = AnimGraphNode->FindPin(InArgs.PinName))
+				{
+					Pin->BreakAllPinLinks();
+				}
+
+				const FFieldVariant& LeafField = InBindingChain.Last().Field;
+
+				FAnimGraphNodePropertyBinding Binding;
+				Binding.PropertyName = InArgs.PinName;
+				if(InArgs.PinProperty->IsA<FArrayProperty>())
+				{
+					// Pull array index from the pin's FName if this is an array property
+					Binding.ArrayIndex = InArgs.PinName.GetNumber() - 1;
+				}
+				PropertyAccessEditor.MakeStringPath(InBindingChain, Binding.PropertyPath);
+				Binding.PathAsText = PropertyAccessEditor.MakeTextPath(Binding.PropertyPath);
+				Binding.Type = LeafField.IsA<UFunction>() ? EAnimGraphNodePropertyBindingType::Function : EAnimGraphNodePropertyBindingType::Property;
+				Binding.bIsBound = true;
+				if(LeafField.IsA<FProperty>())
+				{
+					const FProperty* LeafProperty = LeafField.Get<FProperty>();
+					if(LeafProperty)
+					{
+						if(PropertyAccessEditor.GetPropertyCompatibility(LeafProperty, BindingProperty) == EPropertyAccessCompatibility::Promotable)
+						{
+							Binding.bIsPromotion = true;
+							Schema->ConvertPropertyToPinType(LeafProperty, Binding.PromotedPinType);
+						}
+
+						Schema->ConvertPropertyToPinType(LeafProperty, Binding.PinType);
+					}
+				}
+				else if(LeafField.IsA<UFunction>())
+				{
+					const UFunction* LeafFunction = LeafField.Get<UFunction>();
+					if(LeafFunction)
+					{
+						if(FProperty* ReturnProperty = LeafFunction->GetReturnProperty())
+						{
+							if(PropertyAccessEditor.GetPropertyCompatibility(ReturnProperty, BindingProperty) == EPropertyAccessCompatibility::Promotable)
+							{
+								Binding.bIsPromotion = true;
+								Schema->ConvertPropertyToPinType(ReturnProperty, Binding.PromotedPinType);
+							}
+
+							Schema->ConvertPropertyToPinType(ReturnProperty, Binding.PinType);
+						}
+					}
+				}
+				AnimGraphNode->PropertyBindings.Add(InArgs.PinName, Binding);
+			}
+
+			FBlueprintEditorUtils::MarkBlueprintAsModified(Blueprint);
+		};
+
+		auto OnRemoveBinding = [InArgs, Blueprint](FName InPropertyName)
+		{
+			for(UAnimGraphNode_Base* AnimGraphNode : InArgs.Nodes)
+			{
+				AnimGraphNode->Modify();
+				AnimGraphNode->PropertyBindings.Remove(InArgs.PinName);
+			}
+
+			FBlueprintEditorUtils::MarkBlueprintAsModified(Blueprint);
+		};
+
+		auto CanRemoveBinding = [InArgs](FName InPropertyName)
+		{
+			for(UAnimGraphNode_Base* AnimGraphNode : InArgs.Nodes)
+			{
+				if(AnimGraphNode->PropertyBindings.Contains(InArgs.PinName))
+				{
+					return true;
+				}
+			}
+
+			return false;
+		}; 
+
+		enum class ECurrentValueType : int32
+		{
+			None,
+			Pin,
+			Binding,
+			MultipleValues,
+		};
+
+		auto CurrentBindingText = [InArgs]()
+		{
+			ECurrentValueType CurrentValueType = ECurrentValueType::None;
+
+			const FText MultipleValues = LOCTEXT("MultipleValuesLabel", "Multiple Values");
+			const FText Bind = LOCTEXT("BindLabel", "Bind");
+			const FText ExposedAsPin = LOCTEXT("ExposedAsPinLabel", "Pin");
+			FText CurrentValue = Bind;
+
+			auto SetAssignValue = [&CurrentValueType, &CurrentValue, &MultipleValues](const FText& InValue, ECurrentValueType InType)
+			{
+				if(CurrentValueType != ECurrentValueType::MultipleValues)
+				{
+					if(CurrentValueType == ECurrentValueType::None)
+					{
+						CurrentValueType = InType;
+						CurrentValue = InValue;
+					}
+					else if(CurrentValueType == InType)
+					{
+						if(!CurrentValue.EqualTo(InValue))
+						{
+							CurrentValueType = ECurrentValueType::MultipleValues;
+							CurrentValue = MultipleValues;
+						}
+					}
+					else
+					{
+						CurrentValueType = ECurrentValueType::MultipleValues;
+						CurrentValue = MultipleValues;
+					}
+				}
+			};
+
+			for(UAnimGraphNode_Base* AnimGraphNode : InArgs.Nodes)
+			{
+				if(FAnimGraphNodePropertyBinding* BindingPtr = AnimGraphNode->PropertyBindings.Find(InArgs.PinName))
+				{
+					SetAssignValue(BindingPtr->PathAsText, ECurrentValueType::Binding);
+				}
+				else
+				{
+					if(AnimGraphNode->ShowPinForProperties[InArgs.OptionalPinIndex].bShowPin)
+					{
+						SetAssignValue(InArgs.bOnGraphNode ? Bind : ExposedAsPin, ECurrentValueType::Pin);
+					}
+					else
+					{
+						SetAssignValue(Bind, ECurrentValueType::None);
+					}
+				}
+			}
+
+			return CurrentValue;
+		};
+
+		auto CurrentBindingToolTipText = [InArgs]()
+		{
+			ECurrentValueType CurrentValueType = ECurrentValueType::None;
+
+			const FText MultipleValues = LOCTEXT("MultipleValuesToolTip", "Bindings Have Multiple Values");
+			const FText ExposedAsPin = LOCTEXT("ExposedAsPinToolTip", "Exposed As a Pin on the Node");
+			const FText BindValue = LOCTEXT("BindValueToolTip", "Bind This Value");
+			FText CurrentValue;
+			
+			auto SetAssignValue = [&CurrentValueType, &CurrentValue, &MultipleValues](const FText& InValue, ECurrentValueType InType)
+			{
+				if(CurrentValueType != ECurrentValueType::MultipleValues)
+				{
+					if(CurrentValueType == ECurrentValueType::None)
+					{
+						CurrentValueType = InType;
+						CurrentValue = InValue;
+					}
+					else if(CurrentValueType == InType)
+					{
+						if(!CurrentValue.EqualTo(InValue))
+						{
+							CurrentValueType = ECurrentValueType::MultipleValues;
+							CurrentValue = MultipleValues;
+						}
+					}
+					else
+					{
+						CurrentValueType = ECurrentValueType::MultipleValues;
+						CurrentValue = MultipleValues;
+					}
+				}
+			};
+
+			for(UAnimGraphNode_Base* AnimGraphNode : InArgs.Nodes)
+			{
+				if(FAnimGraphNodePropertyBinding* BindingPtr = AnimGraphNode->PropertyBindings.Find(InArgs.PinName))
+				{
+					if(BindingPtr->PathAsText.IsEmpty())
+					{
+						SetAssignValue(BindValue, ECurrentValueType::Binding);
+					}
+					else
+					{
+						const FText& CompilationContext = BindingPtr->CompiledContext;
+						const FText& CompilationContextDesc = BindingPtr->CompiledContextDesc;
+						if(CompilationContext.IsEmpty() && CompilationContextDesc.IsEmpty())
+						{
+							SetAssignValue(FText::Format(LOCTEXT("BindingToolTipFormat", "Pin is bound to property '{0}'"), BindingPtr->PathAsText), ECurrentValueType::Binding);
+						}
+						else
+						{
+							SetAssignValue(FText::Format(LOCTEXT("BindingToolTipFormatWithDesc", "Pin is bound to property '{0}'\n{1}\n{2}"), BindingPtr->PathAsText, CompilationContext, CompilationContextDesc), ECurrentValueType::Binding);
+						}
+					}
+				}
+				else
+				{
+					if(AnimGraphNode->ShowPinForProperties[InArgs.OptionalPinIndex].bShowPin)
+					{
+						SetAssignValue(InArgs.bOnGraphNode ? BindValue : ExposedAsPin, ECurrentValueType::Pin);
+					}
+					else
+					{
+						SetAssignValue(BindValue, ECurrentValueType::None);
+					}
+				}
+			}
+
+			return CurrentValue;
+		};
+
+		auto CurrentBindingImage = [InArgs, bIsArrayElement]() -> const FSlateBrush*
+		{
+			static FName PropertyIcon(TEXT("Kismet.VariableList.TypeIcon"));
+			static FName FunctionIcon(TEXT("GraphEditor.Function_16x"));
+
+			EAnimGraphNodePropertyBindingType BindingType = EAnimGraphNodePropertyBindingType::None;
+			for(UObject* OuterObject : InArgs.Nodes)
+			{
+				if(UAnimGraphNode_Base* AnimGraphNode = Cast<UAnimGraphNode_Base>(OuterObject))
+				{
+					if(AnimGraphNode->ShowPinForProperties[InArgs.OptionalPinIndex].bShowPin)
+					{
+						BindingType = EAnimGraphNodePropertyBindingType::None;
+						break;
+					}
+					else if(FAnimGraphNodePropertyBinding* BindingPtr = AnimGraphNode->PropertyBindings.Find(InArgs.PinName))
+					{
+						if(BindingType == EAnimGraphNodePropertyBindingType::None)
+						{
+							BindingType = BindingPtr->Type;
+						}
+						else if(BindingType != BindingPtr->Type)
+						{
+							BindingType = EAnimGraphNodePropertyBindingType::None;
+							break;
+						}
+					}
+					else if(BindingType != EAnimGraphNodePropertyBindingType::None)
+					{
+						BindingType = EAnimGraphNodePropertyBindingType::None;
+						break;
+					}
+				}
+			}
+
+			if (BindingType == EAnimGraphNodePropertyBindingType::Function)
+			{
+				return FEditorStyle::GetBrush(FunctionIcon);
+			}
+			else
+			{
+				const UAnimationGraphSchema* AnimationGraphSchema = GetDefault<UAnimationGraphSchema>();
+				FEdGraphPinType PinType;
+				FProperty* PropertyToUse = bIsArrayElement ? CastFieldChecked<FArrayProperty>(InArgs.PinProperty)->Inner : InArgs.PinProperty;
+				if(AnimationGraphSchema->ConvertPropertyToPinType(PropertyToUse, PinType))
+				{
+					return FBlueprintEditorUtils::GetIconFromPin(PinType, false);
+				}
+				else
+				{
+					return FEditorStyle::GetBrush(PropertyIcon);
+				}
+			}
+		};
+
+		auto CurrentBindingColor = [InArgs, BindingProperty]() -> FLinearColor
+		{
+			const UEdGraphSchema_K2* Schema = GetDefault<UEdGraphSchema_K2>();
+
+			FEdGraphPinType PinType;
+			Schema->ConvertPropertyToPinType(BindingProperty, PinType);
+			FLinearColor BindingColor = Schema->GetPinTypeColor(PinType);
+
+			enum class EPromotionState
+			{
+				NotChecked,
+				NotPromoted,
+				Promoted,
+			} Promotion = EPromotionState::NotChecked;
+
+			for(UAnimGraphNode_Base* AnimGraphNode : InArgs.Nodes)
+			{
+				if(AnimGraphNode->ShowPinForProperties[InArgs.OptionalPinIndex].bShowPin)
+				{
+					if(Promotion == EPromotionState::NotChecked)
+					{
+						Promotion = EPromotionState::NotPromoted;
+					}
+					else if(Promotion == EPromotionState::Promoted)
+					{
+						BindingColor = FLinearColor::Gray;
+						break;
+					}
+				}
+				else if(FAnimGraphNodePropertyBinding* BindingPtr = AnimGraphNode->PropertyBindings.Find(InArgs.PinName))
+				{
+					if(Promotion == EPromotionState::NotChecked)
+					{
+						if(BindingPtr->bIsPromotion)
+						{
+							Promotion = EPromotionState::Promoted;
+							BindingColor = Schema->GetPinTypeColor(BindingPtr->PromotedPinType);
+						}
+						else
+						{
+							Promotion = EPromotionState::NotPromoted;
+						}
+					}
+					else
+					{
+						EPromotionState NewPromotion = BindingPtr->bIsPromotion ? EPromotionState::Promoted : EPromotionState::NotPromoted;
+						if(Promotion != NewPromotion)
+						{
+							BindingColor = FLinearColor::Gray;
+							break;
+						}
+					}
+				}
+			}
+
+			return BindingColor;
+		};
+
+		auto AddMenuExtension = [InArgs, Blueprint](FMenuBuilder& InMenuBuilder)
+		{
+			InMenuBuilder.BeginSection("Pins", LOCTEXT("Pin", "Pin"));
+			{
+				InMenuBuilder.AddMenuEntry(
+					LOCTEXT("ExposeAsPin", "Expose As Pin"),
+					LOCTEXT("ExposeAsPinTooltip", "Show/hide this property as a pin on the node"),
+					FSlateIcon("EditorStyle", "GraphEditor.PinIcon"),
+					FUIAction(
+						FExecuteAction::CreateLambda([InArgs, Blueprint]()
+						{
+							bool bHasBinding = false;
+
+							// Comparison without name index to deal with arrays
+							const FName ComparisonName = FName(InArgs.PinName, 0);
+							
+							for(UAnimGraphNode_Base* AnimGraphNode : InArgs.Nodes)
+							{
+								for(const TPair<FName, FAnimGraphNodePropertyBinding>& BindingPair : AnimGraphNode->PropertyBindings)
+								{
+									if(ComparisonName == FName(BindingPair.Key, 0))
+									{
+										bHasBinding = true;
+										break;
+									}
+								}
+							}
+
+							{
+								FScopedTransaction Transaction(LOCTEXT("PinExposure", "Pin Exposure"));
+
+								// Switching from non-pin to pin, remove any bindings
+								for(UAnimGraphNode_Base* AnimGraphNode : InArgs.Nodes)
+								{
+									AnimGraphNode->Modify();
+
+									bool bVisible = AnimGraphNode->ShowPinForProperties[InArgs.OptionalPinIndex].bShowPin;
+									AnimGraphNode->SetPinVisibility(!bVisible || bHasBinding, InArgs.OptionalPinIndex);
+
+									// Remove all bindings that match the property, array or array elements
+									for(auto It = AnimGraphNode->PropertyBindings.CreateIterator(); It; ++It)
+									{
+										if(ComparisonName == FName(It.Key(), 0))
+										{
+											It.RemoveCurrent();
+										}
+									}
+								}
+
+								FBlueprintEditorUtils::MarkBlueprintAsModified(Blueprint);
+							}
+						}),
+						FCanExecuteAction(),
+						FGetActionCheckState::CreateLambda([InArgs]()
+						{
+							bool bPinShown = false;
+							bool bHasBinding = false;
+
+							// Comparison without name index to deal with arrays
+							const FName ComparisonName = FName(InArgs.PinName, 0);
+							
+							for(UAnimGraphNode_Base* AnimGraphNode : InArgs.Nodes)
+							{
+								for(const TPair<FName, FAnimGraphNodePropertyBinding>& BindingPair : AnimGraphNode->PropertyBindings)
+								{
+									if(ComparisonName == FName(BindingPair.Key, 0))
+									{
+										bHasBinding = true;
+										break;
+									}
+								}
+								
+								bPinShown |= AnimGraphNode->ShowPinForProperties[InArgs.OptionalPinIndex].bShowPin;
+							}
+
+							// Pins are exposed if we have a binding or not, so treat as unchecked only if we have
+							// no binding
+							return bPinShown && !bHasBinding ? ECheckBoxState::Checked : ECheckBoxState::Unchecked;
+						})
+					),
+					NAME_None,
+					EUserInterfaceActionType::Check
+				);
+			}
+			InMenuBuilder.EndSection();
+		};
+
+		FPropertyBindingWidgetArgs Args;
+		Args.Property = BindingProperty;
+		Args.OnCanBindProperty = FOnCanBindProperty::CreateLambda(OnCanBindProperty);
+		Args.OnCanBindFunction = FOnCanBindFunction::CreateLambda(OnCanBindFunction);
+		Args.OnCanBindToClass = FOnCanBindToClass::CreateLambda([](UClass* InClass){ return true; });
+		Args.OnAddBinding = FOnAddBinding::CreateLambda(OnAddBinding);
+		Args.OnRemoveBinding = FOnRemoveBinding::CreateLambda(OnRemoveBinding);
+		Args.OnCanRemoveBinding = FOnCanRemoveBinding::CreateLambda(CanRemoveBinding);
+		Args.CurrentBindingText = MakeAttributeLambda(CurrentBindingText);
+		Args.CurrentBindingToolTipText = MakeAttributeLambda(CurrentBindingToolTipText);
+		Args.CurrentBindingImage = MakeAttributeLambda(CurrentBindingImage);
+		Args.CurrentBindingColor = MakeAttributeLambda(CurrentBindingColor);
+		Args.MenuExtender = MakeShared<FExtender>();
+		Args.MenuExtender->AddMenuExtension("BindingActions", EExtensionHook::Before, nullptr, FMenuExtensionDelegate::CreateLambda(AddMenuExtension));
+
+		IPropertyAccessBlueprintBinding::FContext BindingContext;
+		BindingContext.Blueprint = FBlueprintEditorUtils::FindBlueprintForNode(FirstAnimGraphNode);
+		BindingContext.Graph = FirstAnimGraphNode->GetGraph();
+		BindingContext.Node = FirstAnimGraphNode;
+		BindingContext.Pin = FirstAnimGraphNode->FindPin(InArgs.PinName);
+
+		auto OnSetPropertyAccessContextId = [InArgs, Blueprint](const FName& InContextId)
+		{
+			for(UAnimGraphNode_Base* AnimGraphNode : InArgs.Nodes)
+			{
+				if(FAnimGraphNodePropertyBinding* Binding = AnimGraphNode->PropertyBindings.Find(InArgs.PinName))
+				{
+					AnimGraphNode->Modify();
+					Binding->ContextId = InContextId;
+				}
+			}
+
+			FBlueprintEditorUtils::MarkBlueprintAsModified(Blueprint);
+		};
+
+		auto OnCanSetPropertyAccessContextId = [InArgs, FirstAnimGraphNode](const FName& InContextId)
+		{
+			return FirstAnimGraphNode->PropertyBindings.Find(InArgs.PinName) != nullptr;
+		};
+		
+		auto OnGetPropertyAccessContextId = [InArgs]() -> FName
+		{
+			FName CurrentContext = NAME_None;
+			for(UAnimGraphNode_Base* AnimGraphNode : InArgs.Nodes)
+			{
+				if(FAnimGraphNodePropertyBinding* Binding = AnimGraphNode->PropertyBindings.Find(InArgs.PinName))
+				{
+					if(CurrentContext != NAME_None && CurrentContext != Binding->ContextId)
+					{
+						return NAME_None;
+					}
+					else
+					{
+						CurrentContext = Binding->ContextId;
+					}
+				}
+			}
+		
+			return CurrentContext;
+		};
+		
+		IPropertyAccessBlueprintBinding::FBindingMenuArgs MenuArgs;
+		MenuArgs.OnSetPropertyAccessContextId = FOnSetPropertyAccessContextId::CreateLambda(OnSetPropertyAccessContextId);
+		MenuArgs.OnCanSetPropertyAccessContextId = FOnCanSetPropertyAccessContextId::CreateLambda(OnCanSetPropertyAccessContextId);
+		MenuArgs.OnGetPropertyAccessContextId = FOnGetPropertyAccessContextId::CreateLambda(OnGetPropertyAccessContextId);
+		
+		// Add the binding menu extenders
+		TArray<TSharedPtr<FExtender>> Extenders;
+		for(IPropertyAccessBlueprintBinding* Binding : IModularFeatures::Get().GetModularFeatureImplementations<IPropertyAccessBlueprintBinding>("PropertyAccessBlueprintBinding"))
+		{
+			TSharedPtr<FExtender> BindingExtender = Binding->MakeBindingMenuExtender(BindingContext, MenuArgs);
+			if(BindingExtender)
+			{
+				Extenders.Add(BindingExtender);
+			}
+		}
+		
+		if(Extenders.Num() > 0)
+		{
+			Args.MenuExtender = FExtender::Combine(Extenders);
+		}
+
+		Args.bAllowNewBindings = false;
+		Args.bAllowArrayElementBindings = !bIsArray;
+		Args.bAllowUObjectFunctions = !bIsArray;
+		Args.bAllowStructFunctions = !bIsArray;
+
+		IPropertyAccessEditor& PropertyAccessEditor = IModularFeatures::Get().GetModularFeature<IPropertyAccessEditor>("PropertyAccessEditor");
+
+		const FTextBlockStyle& TextBlockStyle = FEditorStyle::Get().GetWidgetStyle<FTextBlockStyle>("PropertyAccess.CompiledContext.Text");
+		
+		return
+			SNew(SBox)
+			.MaxDesiredWidth(200.0f)
+			[
+				SNew(SOverlay)
+				+SOverlay::Slot()
+				.HAlign(HAlign_Right)
+				.VAlign(VAlign_Bottom)
+				[
+					SNew(SBorder)
+					.Padding(FMargin(1.0f, 3.0f, 1.0f, 1.0f))
+					.Visibility(InArgs.bOnGraphNode ? EVisibility::Visible : EVisibility::Collapsed)
+					.BorderImage(FEditorStyle::GetBrush("PropertyAccess.CompiledContext.Border"))
+					.RenderTransform_Lambda([InArgs, FirstAnimGraphNode, &TextBlockStyle]()
+					{
+						const FAnimGraphNodePropertyBinding* BindingPtr = FirstAnimGraphNode->PropertyBindings.Find(InArgs.PinName);
+						FVector2D TextSize(0.0f, 0.0f);
+						if(BindingPtr)
+						{
+							const TSharedRef<FSlateFontMeasure> FontMeasureService = FSlateApplication::Get().GetRenderer()->GetFontMeasureService();
+							TextSize = FontMeasureService->Measure(BindingPtr->CompiledContext, TextBlockStyle.Font);
+						}
+						return FSlateRenderTransform(FVector2D(0.0f, TextSize.Y - 1.0f));
+					})	
+					[
+						SNew(STextBlock)
+						.TextStyle(&TextBlockStyle)
+						.Visibility_Lambda([InArgs, FirstAnimGraphNode, bMultiSelect]()
+						{
+							const FAnimGraphNodePropertyBinding* BindingPtr = FirstAnimGraphNode->PropertyBindings.Find(InArgs.PinName);
+							return bMultiSelect || BindingPtr == nullptr || BindingPtr->CompiledContext.IsEmpty() ? EVisibility::Collapsed : EVisibility::Visible;
+						})
+						.Text_Lambda([InArgs, FirstAnimGraphNode]()
+						{
+							const FAnimGraphNodePropertyBinding* BindingPtr = FirstAnimGraphNode->PropertyBindings.Find(InArgs.PinName);
+							return BindingPtr != nullptr ? BindingPtr->CompiledContext : FText::GetEmpty();
+						})
+					]
+				]
+				+SOverlay::Slot()
+				[
+					PropertyAccessEditor.MakePropertyBindingWidget(Blueprint, Args)
+				]	
+			];
+	}
+	else
+	{
+		return SNew(SSpacer);
+	}
+}
+
 #undef LOCTEXT_NAMESPACE
diff --git a/Engine/Source/Editor/AnimGraph/Private/AnimGraphNode_CallFunction.cpp b/Engine/Source/Editor/AnimGraph/Private/AnimGraphNode_CallFunction.cpp
new file mode 100644
--- /dev/null
+++ b/Engine/Source/Editor/AnimGraph/Private/AnimGraphNode_CallFunction.cpp
@@ -0,0 +1,353 @@
+﻿// Copyright Epic Games, Inc. All Rights Reserved.
+
+#include "AnimGraphNode_CallFunction.h"
+
+#include "BlueprintCompilationManager.h"
+#include "Subsystems/AssetEditorSubsystem.h"
+#include "IAnimationBlueprintEditor.h"
+#include "K2Node_CustomEvent.h"
+#include "KismetCompiler.h"
+#include "Kismet2/BlueprintEditorUtils.h"
+#include "K2Node_CallFunction.h"
+#include "BlueprintNodeSpawner.h"
+#include "BlueprintActionDatabaseRegistrar.h"
+#include "IAnimBlueprintCompilationContext.h"
+#include "Classes/EditorStyleSettings.h"
+#include "AnimBlueprintExtension_CallFunction.h"
+
+#define LOCTEXT_NAMESPACE "AnimGraphNode_CallFunction"
+
+void UAnimGraphNode_CallFunction::PostLoad()
+{
+	Super::PostLoad();
+
+	BindDelegates();
+}
+
+FLinearColor UAnimGraphNode_CallFunction::GetNodeTitleColor() const
+{
+	return CallFunctionPrototype ? CallFunctionPrototype->GetNodeTitleColor() : Super::GetNodeTitleColor();
+}
+
+FText UAnimGraphNode_CallFunction::GetNodeTitle(ENodeTitleType::Type TitleType) const
+{
+	FText FunctionName;
+	UFunction* Function = CallFunctionPrototype ? CallFunctionPrototype->GetTargetFunction() : nullptr;
+	if (Function)
+	{
+		FunctionName = UK2Node_CallFunction::GetUserFacingFunctionName(Function);
+	}
+	else if(CallFunctionPrototype)
+	{
+		FunctionName = FText::FromName(CallFunctionPrototype->FunctionReference.GetMemberName());
+		if ((GEditor != nullptr) && (GetDefault<UEditorStyleSettings>()->bShowFriendlyNames))
+		{
+			FunctionName = FText::FromString(FName::NameToDisplayString(FunctionName.ToString(), false));
+		}
+	}
+	else
+	{
+		FunctionName = LOCTEXT("Function", "Function");
+	}
+	
+	if(TitleType == ENodeTitleType::FullTitle)
+	{
+		FTextBuilder TextBuilder;
+		TextBuilder.AppendLine(FunctionName);
+		TextBuilder.AppendLine(UEnum::GetDisplayValueAsText(Node.CallSite));
+
+		return TextBuilder.ToText();
+	}
+	else
+	{
+		return FunctionName;
+	}
+}
+
+FText UAnimGraphNode_CallFunction::GetTooltipText() const
+{
+	return LOCTEXT("NodeTooltip", "A node that calls user-defined functions during animation graph execution");
+}
+
+void UAnimGraphNode_CallFunction::AllocateFunctionPins()
+{
+	if(CallFunctionPrototype)
+	{
+		const UEdGraphSchema_K2* K2Schema = GetDefault<UEdGraphSchema_K2>();
+		
+		for(UEdGraphPin* Pin : CallFunctionPrototype->Pins)
+		{
+			if(!K2Schema->IsExecPin(*Pin) && Pin->PinName != UEdGraphSchema_K2::PN_Self && Pin->Direction == EGPD_Input)
+			{
+				// Create and copy pin data from prototype K2 node
+				UEdGraphPin* NewPin = CreatePin(EGPD_Input, Pin->PinType, Pin->PinName);
+				NewPin->DefaultObject = Pin->DefaultObject;
+                NewPin->DefaultValue = Pin->DefaultValue;
+                NewPin->DefaultTextValue = Pin->DefaultTextValue;
+				NewPin->AutogeneratedDefaultValue = Pin->AutogeneratedDefaultValue;
+			}
+		}
+	}
+}
+
+void UAnimGraphNode_CallFunction::ReallocatePinsDuringReconstruction(TArray<UEdGraphPin*>& InOldPins)
+{
+	Super::ReallocatePinsDuringReconstruction(InOldPins);
+	
+	AllocateFunctionPins();
+}
+
+void UAnimGraphNode_CallFunction::AllocateDefaultPins()
+{
+	Super::AllocateDefaultPins();
+	
+	AllocateFunctionPins();
+}
+
+FText UAnimGraphNode_CallFunction::GetMenuCategory() const
+{
+	if(CallFunctionPrototype)
+	{
+		if(UFunction* Function = CallFunctionPrototype->GetTargetFunction())
+		{
+			return UK2Node_CallFunction::GetDefaultCategoryForFunction(CallFunctionPrototype->GetTargetFunction(), LOCTEXT("BaseCategory", "Call Function"));
+		}
+	}
+	return FText::GetEmpty();
+}
+
+void UAnimGraphNode_CallFunction::ExpandNode(FKismetCompilerContext& CompilerContext, UEdGraph* SourceGraph)
+{
+	// Note we dont call super here as we dont have an 'evaluation handler'
+	
+	if(CallFunctionPrototype)
+	{
+		const UEdGraphSchema_K2* K2Schema = GetDefault<UEdGraphSchema_K2>();
+
+		UAnimBlueprintExtension_CallFunction* Extension = UAnimBlueprintExtension::GetExtension<UAnimBlueprintExtension_CallFunction>(GetAnimBlueprint());
+		
+		const FName EventName = Extension->AddCustomEventName(this);
+
+		// @TODO: move this name copy to the CDO baking step
+		Node.FunctionName = EventName;
+
+		UK2Node_CustomEvent* CustomEventNode = CompilerContext.SpawnIntermediateEventNode<UK2Node_CustomEvent>(this, nullptr, CompilerContext.ConsolidatedEventGraph);
+		CustomEventNode->bInternalEvent = true;
+		CustomEventNode->CustomFunctionName = EventName;
+		CustomEventNode->AllocateDefaultPins();
+
+		UEdGraphPin* ExecChain = K2Schema->FindExecutionPin(*CustomEventNode, EGPD_Output);
+
+		// Add call function node
+		UK2Node_CallFunction* NewCallFunctionNode = CompilerContext.SpawnIntermediateEventNode<UK2Node_CallFunction>(this, nullptr, CompilerContext.ConsolidatedEventGraph);
+		NewCallFunctionNode->FunctionReference = CallFunctionPrototype->FunctionReference;
+		NewCallFunctionNode->AllocateDefaultPins();
+		
+		// link up pins
+		for(UEdGraphPin* Pin : CallFunctionPrototype->Pins)
+		{
+			if(!K2Schema->IsExecPin(*Pin) && Pin->PinName != UEdGraphSchema_K2::PN_Self && Pin->Direction == EGPD_Input)
+			{
+				UEdGraphPin* AnimGraphPin = FindPinChecked(Pin->PinName);
+				UEdGraphPin* NewPin = NewCallFunctionNode->FindPinChecked(Pin->PinName);
+
+				NewPin->CopyPersistentDataFromOldPin(*AnimGraphPin);
+			}
+		}
+
+		// Link function call into exec chain
+		UEdGraphPin* ExecFunctionCall = K2Schema->FindExecutionPin(*NewCallFunctionNode, EGPD_Input);
+		ExecChain->MakeLinkTo(ExecFunctionCall);
+	}
+}
+
+void UAnimGraphNode_CallFunction::SetupFromFunction(UFunction* InFunction)
+{
+	// Create graph and inner node
+	InnerGraph = FBlueprintEditorUtils::CreateNewGraph(this, NAME_None, UEdGraph::StaticClass(), UEdGraphSchema_K2::StaticClass());
+
+	FGraphNodeCreator<UK2Node_CallFunction> NodeCreator(*InnerGraph);
+	CallFunctionPrototype = NodeCreator.CreateNode();
+	CallFunctionPrototype->FunctionReference.SetFromField<UFunction>(InFunction, true);
+	NodeCreator.Finalize();
+
+	BindDelegates();
+}
+
+void UAnimGraphNode_CallFunction::BindDelegates()
+{
+	if(!GraphChangedHandle.IsValid())
+	{
+		GraphChangedHandle = InnerGraph->AddOnGraphChangedHandler(FOnGraphChanged::FDelegate::CreateLambda([this](const FEdGraphEditAction& InAction)
+        {
+            // Reconstruct node when the inner graph changes (this catches changes to the function signature)
+            ReconstructNode();
+        }));
+	}
+
+	if(!PinRenamedHandle.IsValid())
+	{
+		PinRenamedHandle = CallFunctionPrototype->OnUserDefinedPinRenamed().AddLambda([this](UK2Node* InNode, FName InOldName, FName InNewName)
+        {
+            RenameUserDefinedPin(InOldName, InNewName);
+        });
+	}
+}
+
+bool UAnimGraphNode_CallFunction::IsFunctionBlacklisted(const UFunction* InFunction) const
+{
+	return InFunction->GetFName() == GET_FUNCTION_NAME_CHECKED(UAnimInstance, BlueprintThreadSafeUpdateAnimation);
+}
+
+bool UAnimGraphNode_CallFunction::AreFunctionParamsValid(const UFunction* InFunction) const
+{
+	for(TFieldIterator<FProperty> It(InFunction); It; ++It)
+	{
+		const FProperty* Property = *It;
+
+		// As we cant process return params, we don't allow functions with them to be called in the anim graph
+		if(Property->HasAnyPropertyFlags(CPF_ReturnParm))
+		{
+			return false;
+		}
+	}
+
+	return true;
+}
+
+bool UAnimGraphNode_CallFunction::ValidateFunction(const UFunction* InFunction, FCompilerResultsLog* InMessageLog) const
+{
+	bool bValid = true;
+
+	auto InvalidateMessage = [this, &bValid, InMessageLog](const FText& InMessage)
+	{
+		bValid = false;
+		if(InMessageLog)
+		{
+			InMessageLog->Error(*InMessage.ToString(), this);
+		}
+	};
+
+	if(InFunction->HasAnyFunctionFlags(FUNC_BlueprintPure))
+	{
+		InvalidateMessage(LOCTEXT("PureFunctionError", "@@ cannot call a pure function"));
+	}
+
+	if(!FBlueprintEditorUtils::HasFunctionBlueprintThreadSafeMetaData(InFunction))
+	{
+		InvalidateMessage(LOCTEXT("ThreadSafetyError", "@@ call is not thread safe"));
+	}
+
+	if(!AreFunctionParamsValid(InFunction))
+	{
+		InvalidateMessage(LOCTEXT("FunctionParamsInvalidError", "@@ has invalid parameters. Return parameters are not allowed"));
+	}
+
+	if(InFunction->HasMetaData(FBlueprintMetadata::MD_BlueprintInternalUseOnly))
+	{
+		InvalidateMessage(LOCTEXT("FunctionInternalError", "@@ uses an internal-only function"));
+	}
+
+	if(IsFunctionBlacklisted(InFunction))
+	{
+		InvalidateMessage(LOCTEXT("FunctionBlacklistedError", "@@ uses a blacklisted function"));
+	}
+	
+	return bValid;
+}
+
+void UAnimGraphNode_CallFunction::GetMenuActions(FBlueprintActionDatabaseRegistrar& ActionRegistrar) const
+{
+	const UAnimBlueprint* AnimBlueprint = Cast<UAnimBlueprint>(ActionRegistrar.GetActionKeyFilter());
+	if(AnimBlueprint && ActionRegistrar.IsOpenForRegistration(AnimBlueprint))
+	{
+		auto MakeFunctionActionsForClass = [this, &ActionRegistrar, AnimBlueprint](UClass* InClass)
+		{
+			auto MakeFunctionAction = [this, &ActionRegistrar, AnimBlueprint](UFunction* InFunction)
+			{
+				if(UEdGraphSchema_K2::CanUserKismetCallFunction(InFunction) && ValidateFunction(InFunction))
+				{
+					auto CustomizeNode = [InFunction](UEdGraphNode* InNode, bool bIsTemplate)
+					{
+						UAnimGraphNode_CallFunction* CallFunctionNode = CastChecked<UAnimGraphNode_CallFunction>(InNode);
+						CallFunctionNode->SetupFromFunction(InFunction);
+					};
+
+					UBlueprintNodeSpawner* Spawner = UBlueprintNodeSpawner::Create(UAnimGraphNode_CallFunction::StaticClass(), nullptr, UBlueprintNodeSpawner::FCustomizeNodeDelegate::CreateLambda(CustomizeNode));
+					FBlueprintActionUiSpec& MenuSignature = Spawner->DefaultMenuSignature;
+					
+					MenuSignature.MenuName = FText::Format(LOCTEXT("MenuNameFormat", "{0} (From Anim Graph)"), UK2Node_CallFunction::GetUserFacingFunctionName(InFunction));
+					MenuSignature.Category = UK2Node_CallFunction::GetDefaultCategoryForFunction(InFunction, LOCTEXT("BaseCategory", "Call Function From Anim Graph"));
+					MenuSignature.Tooltip = FText::FromString(UK2Node_CallFunction::GetDefaultTooltipForFunction(InFunction));
+					MenuSignature.Keywords = UK2Node_CallFunction::GetKeywordsForFunction(InFunction);
+
+					// add at least one character, so that PrimeDefaultUiSpec() doesn't attempt to query the template node
+					if (MenuSignature.Keywords.IsEmpty())
+                	{
+                		MenuSignature.Keywords = FText::FromString(TEXT(" "));
+                	}
+					
+					ActionRegistrar.AddBlueprintAction(AnimBlueprint, Spawner);
+				}
+			};
+
+			for(TFieldIterator<UFunction> It(InClass); It; ++It)
+			{
+				MakeFunctionAction(*It);
+			}
+		};
+
+		// Add this class
+		MakeFunctionActionsForClass(AnimBlueprint->GetAnimBlueprintGeneratedClass());
+
+		// Add function libraries too
+		for(TObjectIterator<UBlueprintFunctionLibrary> It(RF_NoFlags); It; ++It)
+		{
+			MakeFunctionActionsForClass(It->GetClass());
+		}
+	}
+}
+
+void UAnimGraphNode_CallFunction::GetRequiredExtensions(TArray<TSubclassOf<UAnimBlueprintExtension>>& OutExtensions) const
+{
+	OutExtensions.Add(UAnimBlueprintExtension_CallFunction::StaticClass());
+}
+
+UObject* UAnimGraphNode_CallFunction::GetJumpTargetForDoubleClick() const
+{
+	return CallFunctionPrototype ? CallFunctionPrototype->GetTargetFunction() : nullptr;
+}
+
+void UAnimGraphNode_CallFunction::JumpToDefinition() const
+{
+	if(CallFunctionPrototype)
+	{
+		UFunction* Function = CallFunctionPrototype->GetTargetFunction();
+		TSharedPtr<IBlueprintEditor> BlueprintEditor = FKismetEditorUtilities::GetIBlueprintEditorForObject(Function, true);
+		if(BlueprintEditor.IsValid())
+		{
+			BlueprintEditor->JumpToHyperlink(Function, false);
+		}
+	}
+}
+
+void UAnimGraphNode_CallFunction::ValidateNodeDuringCompilation(FCompilerResultsLog& MessageLog) const
+{
+	if(CallFunctionPrototype == nullptr || CallFunctionPrototype->GetTargetFunction() == nullptr)
+	{
+		MessageLog.Error(*LOCTEXT("MissingFunctionPrototypeError", "Missing function, node @@ is invalid").ToString(), this);
+	}
+	
+	if(CallFunctionPrototype != nullptr)
+	{
+		CallFunctionPrototype->ValidateNodeDuringCompilation(MessageLog);
+		
+		UFunction* Function = CallFunctionPrototype->GetTargetFunction();
+		if(Function)
+		{
+			ValidateFunction(Function, &MessageLog);
+		}
+	}
+}
+
+#undef LOCTEXT_NAMESPACE
\ No newline at end of file
diff --git a/Engine/Source/Editor/AnimGraph/Private/AnimGraphNode_CallFunction.h b/Engine/Source/Editor/AnimGraph/Private/AnimGraphNode_CallFunction.h
new file mode 100644
--- /dev/null
+++ b/Engine/Source/Editor/AnimGraph/Private/AnimGraphNode_CallFunction.h
@@ -0,0 +1,73 @@
+﻿// Copyright Epic Games, Inc. All Rights Reserved.
+
+#pragma once
+
+#include "CoreMinimal.h"
+#include "AnimGraphNode_Base.h"
+#include "AnimNodes/AnimNode_CallFunction.h"
+#include "Engine/MemberReference.h"
+
+#include "AnimGraphNode_CallFunction.generated.h"
+
+class UEdGraph;
+class UK2Node_CallFunction;
+
+UCLASS(MinimalAPI, meta=(Keywords = "Event"))
+class UAnimGraphNode_CallFunction : public UAnimGraphNode_Base
+{
+	GENERATED_BODY()
+
+private:
+	// Inner graph we maintain to hook into the function call machinery
+	UPROPERTY()
+	UEdGraph* InnerGraph;
+
+	// Inner node we maintain to hook into the function call machinery
+	UPROPERTY()
+	UK2Node_CallFunction* CallFunctionPrototype;
+	
+    UPROPERTY(EditAnywhere, Category=Settings)
+	FAnimNode_CallFunction Node;
+
+	// Delegate handles used to hook into inner graph & function nodes
+	FDelegateHandle GraphChangedHandle;
+	FDelegateHandle PinRenamedHandle;
+
+	// UObject interface
+	virtual void PostLoad() override;
+	
+	// UEdGraphNode interface
+	virtual FLinearColor GetNodeTitleColor() const override;
+	virtual FText GetNodeTitle(ENodeTitleType::Type TitleType) const override;
+	virtual FText GetTooltipText() const override;
+	virtual void AllocateDefaultPins() override;
+	virtual void ReallocatePinsDuringReconstruction(TArray<UEdGraphPin*>& InOldPins) override;
+	virtual UObject* GetJumpTargetForDoubleClick() const override;
+	virtual void JumpToDefinition() const override;
+	virtual void ValidateNodeDuringCompilation(FCompilerResultsLog& MessageLog) const override;
+	
+	// UK2Node interface
+	virtual FText GetMenuCategory() const override;
+	virtual void ExpandNode(FKismetCompilerContext& CompilerContext, UEdGraph* SourceGraph) override;
+	virtual void GetMenuActions(FBlueprintActionDatabaseRegistrar& ActionRegistrar) const override;
+
+	// UAnimGraphNode_Base interface
+	virtual void GetRequiredExtensions(TArray<TSubclassOf<UAnimBlueprintExtension>>& OutExtensions) const override;
+	virtual bool ShouldCreateStructEvalHandlers() const override { return false; }
+
+private:
+	void AllocateFunctionPins();
+
+	/** Sets up node internals */
+	void SetupFromFunction(UFunction* InFunction);
+	void BindDelegates();
+
+	/** Checks to see if the passed-in function is blacklisted (i.e. cannot be called) in anim graphs. */
+	bool IsFunctionBlacklisted(const UFunction* InFunction) const;
+
+	/** Validate the function's parameters to check if it can be called in anim graphs */
+	bool AreFunctionParamsValid(const UFunction* InFunction) const;
+
+	/** Validates a function. Used for populating menu actions and verifying already-referenced functions */
+	bool ValidateFunction(const UFunction* InFunction, FCompilerResultsLog* InMessageLog = nullptr) const;
+};
diff --git a/Engine/Source/Editor/AnimGraph/Private/AnimPropertyAccessUtils.cpp b/Engine/Source/Editor/AnimGraph/Private/AnimPropertyAccessUtils.cpp
new file mode 100644
--- /dev/null
+++ b/Engine/Source/Editor/AnimGraph/Private/AnimPropertyAccessUtils.cpp
@@ -0,0 +1,3 @@
+// Copyright Epic Games, Inc. All Rights Reserved.
+
+#include "AnimPropertyAccessUtils.h"
diff --git a/Engine/Source/Editor/AnimGraph/Private/PropertyAccessAnimBlueprintBinding.cpp b/Engine/Source/Editor/AnimGraph/Private/PropertyAccessAnimBlueprintBinding.cpp
new file mode 100644
--- /dev/null
+++ b/Engine/Source/Editor/AnimGraph/Private/PropertyAccessAnimBlueprintBinding.cpp
@@ -0,0 +1,99 @@
+// Copyright Epic Games, Inc. All Rights Reserved.
+
+#include "PropertyAccessAnimBlueprintBinding.h"
+
+#include "AnimBlueprintExtension_PropertyAccess.h"
+#include "AnimPropertyAccessUtils.h"
+#include "Animation/AnimBlueprint.h"
+#include "Framework/MultiBox/MultiBoxBuilder.h"
+#include "Framework/MultiBox/MultiBoxExtender.h"
+
+#define LOCTEXT_NAMESPACE "FPropertyAccessAnimBlueprintBinding"
+
+bool FPropertyAccessAnimBlueprintBinding::CanBindToContext(const FContext& InContext) const
+{
+	if(InContext.Blueprint)
+	{
+		return InContext.Blueprint->IsA<UAnimBlueprint>();
+	}
+
+	return false;
+}
+
+TSharedPtr<FExtender> FPropertyAccessAnimBlueprintBinding::MakeBindingMenuExtender(const FContext& InContext, const FBindingMenuArgs& InArgs) const
+{
+	TSharedPtr<FExtender> Extender = MakeShared<FExtender>();
+	Extender->AddMenuExtension("BindingActions", EExtensionHook::After, nullptr, FMenuExtensionDelegate::CreateLambda([InArgs](FMenuBuilder& InMenuBuilder)
+	{
+		InMenuBuilder.BeginSection("CallSite", LOCTEXT("CallSite", "Call Site"));
+		{
+			InMenuBuilder.AddSubMenu(
+				LOCTEXT("CallSite", "Call Site"),
+				LOCTEXT("CallSiteTooltip", "Choose when this property access value will be called & cached for later use when this graph is executed."),
+				FNewMenuDelegate::CreateLambda([InArgs](FMenuBuilder& InSubMenuBuilder)
+				{
+					InSubMenuBuilder.BeginSection("CallSiteSubMenu", LOCTEXT("CallSite", "Call Site"));
+					{
+						auto AddContextEntry = [InArgs, &InSubMenuBuilder](const FName& InContextId, const FText& InLabel, const FText& InToolTip)
+						{
+							InSubMenuBuilder.AddMenuEntry(
+								InLabel,
+								InToolTip,
+								FSlateIcon(),
+								FUIAction(
+									FExecuteAction::CreateLambda([InArgs, InContextId]()
+									{
+										InArgs.OnSetPropertyAccessContextId.ExecuteIfBound(InContextId);
+									}),
+									FCanExecuteAction::CreateLambda([InArgs, InContextId]()
+									{
+										if(InArgs.OnCanSetPropertyAccessContextId.IsBound())
+										{
+											return InArgs.OnCanSetPropertyAccessContextId.Execute(InContextId);
+										}
+										return false;
+									}),
+									FIsActionChecked::CreateLambda([InArgs, InContextId]()
+									{
+										const FName CurrentId = InArgs.OnGetPropertyAccessContextId.IsBound() ? InArgs.OnGetPropertyAccessContextId.Execute() : NAME_None;
+										return CurrentId == InContextId;
+									})),
+								NAME_None,
+								EUserInterfaceActionType::RadioButton
+							);
+						};
+						
+						AddContextEntry(UAnimBlueprintExtension_PropertyAccess::ContextId_Automatic,
+							LOCTEXT("CallSiteAutomatic", "Automatic"),
+							LOCTEXT("CallSiteAutomatic_ToolTip", "Automatically determine the call site for this property access based on context and thread safety"));
+
+						AddContextEntry(UAnimBlueprintExtension_PropertyAccess::ContextId_UnBatched_ThreadSafe,
+							LOCTEXT("CallSiteThreadSafe", "Thread Safe"),
+							LOCTEXT("CallSiteThreadSafe_ToolTip", "Can safely be executed on worker threads"));
+
+						AddContextEntry(UAnimBlueprintExtension_PropertyAccess::ContextId_Batched_GameThreadPreEventGraph,
+							LOCTEXT("CallSiteGameThreadPreEventGraph", "Game Thread (Pre-Event Graph)"),
+							LOCTEXT("CallSiteGameThreadPreEventGraph_ToolTip", "Executed on the game thread before the event graph is run"));
+
+						AddContextEntry(UAnimBlueprintExtension_PropertyAccess::ContextId_Batched_GameThreadPostEventGraph,
+							LOCTEXT("CallSiteGameThreadPostEventGraph", "Game Thread (Post-Event Graph)"),
+							LOCTEXT("CallSiteGameThreadPostEventGraph_ToolTip", "Executed on the game thread after the event graph is run"));
+
+						AddContextEntry(UAnimBlueprintExtension_PropertyAccess::ContextId_Batched_WorkerThreadPreEventGraph,
+							LOCTEXT("CallSiteWorkerThreadPreEventGraph", "Worker Thread (Pre-Event Graph)"),
+							LOCTEXT("CallSiteWorkerThreadPreEventGraph_ToolTip", "Executed on a worker thread before the thread safe event graph is run"));
+
+						AddContextEntry(UAnimBlueprintExtension_PropertyAccess::ContextId_Batched_WorkerThreadPostEventGraph,
+							LOCTEXT("CallSiteWorkerThreadPostEventGraph", "Worker Thread (Post-Event Graph)"),
+							LOCTEXT("CallSiteWorkerThreadPostEventGraph_ToolTip", "Executed on a worker thread after the thread safe event graph is run"));
+					}
+					InSubMenuBuilder.EndSection();
+				})
+			);
+		}
+		InMenuBuilder.EndSection();
+	}));
+	return Extender;
+}
+
+#undef LOCTEXT_NAMESPACE
diff --git a/Engine/Source/Editor/AnimGraph/Private/PropertyAccessAnimBlueprintBinding.h b/Engine/Source/Editor/AnimGraph/Private/PropertyAccessAnimBlueprintBinding.h
new file mode 100644
--- /dev/null
+++ b/Engine/Source/Editor/AnimGraph/Private/PropertyAccessAnimBlueprintBinding.h
@@ -0,0 +1,14 @@
+// Copyright Epic Games, Inc. All Rights Reserved.
+
+#pragma once
+
+#include "CoreMinimal.h"
+#include "IPropertyAccessBlueprintBinding.h"
+
+class FPropertyAccessAnimBlueprintBinding : public IPropertyAccessBlueprintBinding
+{
+public:
+	// IPropertyAccessBlueprintBinding interface
+	virtual bool CanBindToContext(const FContext& InContext) const override;
+	virtual TSharedPtr<FExtender> MakeBindingMenuExtender(const FContext& InContext, const FBindingMenuArgs& InArgs) const override;
+};
diff --git a/Engine/Source/Editor/AnimGraph/Public/AnimBlueprintExtension_PropertyAccess.h b/Engine/Source/Editor/AnimGraph/Public/AnimBlueprintExtension_PropertyAccess.h
--- a/Engine/Source/Editor/AnimGraph/Public/AnimBlueprintExtension_PropertyAccess.h
+++ b/Engine/Source/Editor/AnimGraph/Public/AnimBlueprintExtension_PropertyAccess.h
@@ -13,6 +13,9 @@ class IAnimBlueprintCompilationBracketContext;
 class IAnimBlueprintGeneratedClassCompiledData;
 class UAnimBlueprintExtension_PropertyAccess;
 enum class EPropertyAccessBatchType : uint8;
+class FKismetCompilerContext;
+class UEdGraph;
+class UEdGraphPin;
 
 // Delegate called when the library is compiled (whether successfully or not)
 DECLARE_MULTICAST_DELEGATE_TwoParams(FOnPostLibraryCompiled, IAnimBlueprintCompilationBracketContext& /*InCompilationContext*/, IAnimBlueprintGeneratedClassCompiledData& /*OutCompiledData*/)
@@ -24,10 +27,36 @@ class ANIMGRAPH_API UAnimBlueprintExtension_PropertyAccess : public UAnimBluepri
 
 	friend class UAnimBlueprintExtension_Base;
 
+public:	
+	// Try to determine the context in which this property access can safely be performed.
+	static const FName ContextId_Automatic;
+
+	// Can safely be executed on worker threads
+	static const FName ContextId_UnBatched_ThreadSafe;
+
+	// Executed batched on the game thread before the event graph is run
+	static const FName ContextId_Batched_WorkerThreadPreEventGraph;
+
+	// Executed batched on the game thread after the event graph is run
+	static const FName ContextId_Batched_WorkerThreadPostEventGraph;
+
+	// Executed batched on the game thread before the event graph is run
+	static const FName ContextId_Batched_GameThreadPreEventGraph;
+
+	// Executed batched on the game thread after the event graph is run
+	static const FName ContextId_Batched_GameThreadPostEventGraph;
+	
 public:
 	// Add a copy to the property access library we are compiling
 	// @return an integer handle to the pending copy. This can be resolved to a true copy index by calling MapCopyIndex
-	int32 AddCopy(TArrayView<FString> InSourcePath, TArrayView<FString> InDestPath, EPropertyAccessBatchType InBatchType, UObject* InObject = nullptr);
+	FPropertyAccessHandle AddCopy(TArrayView<FString> InSourcePath, TArrayView<FString> InDestPath, const FName& InContextId, UObject* InObject = nullptr);
+	
+	UE_DEPRECATED(5.0, "Please use AddCopy with a context ID")
+	int32 AddCopy(TArrayView<FString> InSourcePath, TArrayView<FString> InDestPath, EPropertyAccessBatchType InBatchType, UObject* InObject = nullptr)
+	{
+		FPropertyAccessHandle Handle = AddCopy(InSourcePath, InDestPath, FName(NAME_None), InObject);
+		return Handle.GetId();
+	}
 
 	// Delegate called when the library is compiled (whether successfully or not)
 	FSimpleMulticastDelegate& OnPreLibraryCompiled() { return OnPreLibraryCompiledDelegate; }
@@ -35,9 +64,28 @@ public:
 	// Delegate called when the library is compiled (whether successfully or not)
 	FOnPostLibraryCompiled& OnPostLibraryCompiled() { return OnPostLibraryCompiledDelegate; }
 
-	// Maps the initial integer copy handle to a true handle, post compilation
-	int32 MapCopyIndex(int32 InIndex) const;
+	// Maps the initial copy handle to a true handle, post compilation
+	FCompiledPropertyAccessHandle GetCompiledHandle(FPropertyAccessHandle InHandle) const;
+
+	UE_DEPRECATED(5.0, "Please use GetCompiledHandle")
+	int32 MapCopyIndex(int32 InIndex) const
+	{
+		FCompiledPropertyAccessHandle CompiledHandle = GetCompiledHandle(FPropertyAccessHandle(InIndex));
+		return CompiledHandle.GetId();
+	}
+
+	// Expands a property access path to a pure chain of BP nodes
+	void ExpandPropertyAccess(FKismetCompilerContext& InCompilerContext, TArrayView<FString> InSourcePath, UEdGraph* InParentGraph, UEdGraphPin* InTargetPin) const;
+
+	// Maps a compiled handle back to a human-readable context 
+	static FText GetCompiledHandleContext(FCompiledPropertyAccessHandle InHandle);
+	
+	// Maps a compiled handle back to a human-readable context description
+	static FText GetCompiledHandleContextDesc(FCompiledPropertyAccessHandle InHandle);
 
+	// Whether a context name requires an auto-generated variable to be cached for an access
+	static bool ContextRequiresCachedVariable(FName InName);
+	
 private:
 	// UAnimBlueprintExtension interface
 	virtual void HandleStartCompilingClass(const UClass* InClass, IAnimBlueprintCompilationBracketContext& InCompilationContext, IAnimBlueprintGeneratedClassCompiledData& OutCompiledData) override;
diff --git a/Engine/Source/Editor/AnimGraph/Public/AnimGraphModule.h b/Engine/Source/Editor/AnimGraph/Public/AnimGraphModule.h
deleted file mode 100644
diff --git a/Engine/Source/Editor/AnimGraph/Public/AnimGraphNode_Base.h b/Engine/Source/Editor/AnimGraph/Public/AnimGraphNode_Base.h
--- a/Engine/Source/Editor/AnimGraph/Public/AnimGraphNode_Base.h
+++ b/Engine/Source/Editor/AnimGraph/Public/AnimGraphNode_Base.h
@@ -10,6 +10,7 @@
 #include "Animation/AnimBlueprintGeneratedClass.h"
 #include "Animation/AnimNodeBase.h"
 #include "Editor.h"
+#include "IPropertyAccessEditor.h"
 #include "K2Node.h"
 #include "AnimGraphNode_Base.generated.h"
 
@@ -22,7 +23,6 @@ class IDetailLayoutBuilder;
 class UAnimGraphNode_Base;
 class UEdGraphSchema;
 class USkeletalMeshComponent;
-class IAnimBlueprintCompilerHandlerCollection;
 class IAnimBlueprintGeneratedClassCompiledData;
 class IAnimBlueprintCompilationContext;
 class IAnimBlueprintCopyTermDefaultsContext;
@@ -148,6 +148,9 @@ struct FAnimGraphNodePropertyBinding
 	UPROPERTY()
 	FName PropertyName;
 
+	UPROPERTY()
+	int32 ArrayIndex = INDEX_NONE;
+	
 	/** The property path as text */
 	UPROPERTY()
 	FText PathAsText;
@@ -156,6 +159,16 @@ struct FAnimGraphNodePropertyBinding
 	UPROPERTY()
 	TArray<FString> PropertyPath;
 
+	/** The context of the binding */
+	UPROPERTY()
+	FName ContextId;
+
+	UPROPERTY()
+	FText CompiledContext;
+
+	UPROPERTY()
+	FText CompiledContextDesc;
+	
 	/** Whether the binding is a function or not */
 	UPROPERTY()
 	EAnimGraphNodePropertyBindingType Type = EAnimGraphNodePropertyBindingType::Property;
@@ -386,11 +399,42 @@ class ANIMGRAPH_API UAnimGraphNode_Base : public UK2Node
 	// @return the proxy for attribute display
 	virtual const UAnimGraphNode_Base* GetProxyNodeForAttributes() const { return this; }
 
+	// Sets the visibility of the specified pin, reconstructs the node if it changes
+	void SetPinVisibility(bool bInVisible, int32 InOptionalPinIndex);
+
+	// Arguments used to construct a property binding widget
+	struct FAnimPropertyBindingWidgetArgs
+	{
+		FAnimPropertyBindingWidgetArgs(const TArray<UAnimGraphNode_Base*>& InNodes, FProperty* InPinProperty, FName InPinName, int32 InOptionalPinIndex)
+			: Nodes(InNodes)
+			, PinProperty(InPinProperty)
+			, PinName(InPinName)
+			, OptionalPinIndex(InOptionalPinIndex)
+			, bOnGraphNode(true)
+		{
+		}
+
+		// The nodes to display the binding for
+		TArray<UAnimGraphNode_Base*> Nodes;
+		// The pin property for this binding
+		FProperty* PinProperty;
+		// The name of the pin
+		FName PinName;
+		// The optional pin index that refers to this pin
+		int32 OptionalPinIndex;
+		// Whether this is for display on a graph node
+		bool bOnGraphNode;
+	};
+	
+	// Make a property binding widget to edit the bindings of the passed-in nodes
+	static TSharedRef<SWidget> MakePropertyBindingWidget(const FAnimPropertyBindingWidgetArgs& InArgs);
+	
 protected:
 	friend class FAnimBlueprintCompilerContext;
 	friend class FAnimGraphNodeDetails;
 	friend class UAnimBlueprintExtension;
 	friend class UAnimBlueprintExtension_Base;
+	friend class SAnimationGraphNode;
 
 	// Gets the animation FNode type represented by this ed graph node
 	UScriptStruct* GetFNodeType() const;
@@ -430,6 +474,9 @@ protected:
 
 	// Override point for OverrideAssets
 	virtual void OnOverrideAssets(IAnimBlueprintNodeOverrideAssetsContext& InContext) const {}
+
+	// Whether this node should create BP evaluation handlers as part of compilation
+	virtual bool ShouldCreateStructEvalHandlers() const { return true; }
 	
 	// Allocates or reallocates pins
 	void InternalPinCreation(TArray<UEdGraphPin*>* OldPins);
diff --git a/Engine/Source/Editor/AnimGraph/Public/AnimPropertyAccessUtils.h b/Engine/Source/Editor/AnimGraph/Public/AnimPropertyAccessUtils.h
new file mode 100644
--- /dev/null
+++ b/Engine/Source/Editor/AnimGraph/Public/AnimPropertyAccessUtils.h
@@ -0,0 +1,5 @@
+// Copyright Epic Games, Inc. All Rights Reserved.
+
+#pragma once
+
+#include "CoreMinimal.h"
diff --git a/Engine/Source/Editor/AnimationBlueprintEditor/Private/AnimationNodes/SAnimationGraphNode.cpp b/Engine/Source/Editor/AnimationBlueprintEditor/Private/AnimationNodes/SAnimationGraphNode.cpp
--- a/Engine/Source/Editor/AnimationBlueprintEditor/Private/AnimationNodes/SAnimationGraphNode.cpp
+++ b/Engine/Source/Editor/AnimationBlueprintEditor/Private/AnimationNodes/SAnimationGraphNode.cpp
@@ -282,6 +282,11 @@ void SAnimationGraphNode::ReconfigurePinWidgetsForPropertyBindings(UAnimGraphNod
 
 			if(PinWidget.IsValid())
 			{
+				// Tweak padding a little to improve extended appearance
+				PinWidget->GetLabelAndValue()->SetInnerSlotPadding(FVector2D(2.0f, 0.0f));
+				
+				const FName PinName = Pin->GetFName();
+
 				// Compare FName without number to make sure we catch array properties that are split into multiple pins
 				FName ComparisonName = Pin->GetFName();
 				ComparisonName.SetNumber(0);
@@ -291,7 +296,7 @@ void SAnimationGraphNode::ReconfigurePinWidgetsForPropertyBindings(UAnimGraphNod
 				{
 					TWeakPtr<SGraphPin> WeakPinWidget = PinWidget;
 
-					PinWidget->GetValueWidget()->SetVisibility(MakeAttributeLambda([ComparisonName, InAnimGraphNode, WeakPinWidget]()
+					PinWidget->GetValueWidget()->SetVisibility(MakeAttributeLambda([PinName, InAnimGraphNode, WeakPinWidget]()
 					{
 						EVisibility Visibility = EVisibility::Collapsed;
 
@@ -299,7 +304,7 @@ void SAnimationGraphNode::ReconfigurePinWidgetsForPropertyBindings(UAnimGraphNod
 						{
 							Visibility = WeakPinWidget.Pin()->GetDefaultValueVisibility();
 
-							if (FAnimGraphNodePropertyBinding* BindingPtr = InAnimGraphNode->PropertyBindings.Find(ComparisonName))
+							if (FAnimGraphNodePropertyBinding* BindingPtr = InAnimGraphNode->PropertyBindings.Find(PinName))
 							{
 								Visibility = EVisibility::Collapsed;
 							}
@@ -309,69 +314,22 @@ void SAnimationGraphNode::ReconfigurePinWidgetsForPropertyBindings(UAnimGraphNod
 					}));
 				}
 
-				// Add an image & label for a binding
-				PinWidget->GetLabelAndValue()->AddSlot()
-				[
-					SNew(SHorizontalBox)
-					.ToolTipText_Lambda([ComparisonName, InAnimGraphNode]()
-					{
-						if (FAnimGraphNodePropertyBinding* BindingPtr = InAnimGraphNode->PropertyBindings.Find(ComparisonName))
-						{
-							return FText::Format(LOCTEXT("BindingTooltipFormat", "Pin is bound to property '{0}'"), BindingPtr->PathAsText);
-						}
-
-						return FText::GetEmpty();
-					})
-					.Visibility_Lambda([ComparisonName, InAnimGraphNode]()
+				FProperty* PinProperty = InAnimGraphNode->GetFNodeType()->FindPropertyByName(ComparisonName);
+				if(PinProperty)
+				{
+					const int32 OptionalPinIndex = InAnimGraphNode->ShowPinForProperties.IndexOfByPredicate([PinProperty](const FOptionalPinFromProperty& InOptionalPin)
 					{
-						return InAnimGraphNode->PropertyBindings.Contains(ComparisonName) ? EVisibility::Visible : EVisibility::Collapsed;
-					})
-					+SHorizontalBox::Slot()
-					.AutoWidth()
-					.VAlign(VAlign_Center)
-					.Padding(3.0f, 2.0f)
-					[
-						SNew(SImage)
-						.Image_Lambda([ComparisonName, InAnimGraphNode, PinType]() -> const FSlateBrush*
-						{
-							if (FAnimGraphNodePropertyBinding* BindingPtr = InAnimGraphNode->PropertyBindings.Find(ComparisonName))
-							{
-								static FName FunctionIcon(TEXT("GraphEditor.Function_16x"));
+						return PinProperty->GetFName() == InOptionalPin.PropertyName;
+					});
 
-								return BindingPtr->Type == EAnimGraphNodePropertyBindingType::Property ? FBlueprintEditorUtils::GetIconFromPin(PinType, true) : FEditorStyle::GetBrush(FunctionIcon);
-							}
+					UAnimGraphNode_Base::FAnimPropertyBindingWidgetArgs Args( { InAnimGraphNode }, PinProperty, Pin->GetFName(), OptionalPinIndex);
 
-							return nullptr;
-						})
-						.ColorAndOpacity_Lambda([InAnimGraphNode, ComparisonName]()
-						{
-							if(const UEdGraphSchema* Schema = InAnimGraphNode->GetSchema())
-							{
-								if (FAnimGraphNodePropertyBinding* BindingPtr = InAnimGraphNode->PropertyBindings.Find(ComparisonName))
-								{
-									return Schema->GetPinTypeColor(BindingPtr->bIsPromotion ? BindingPtr->PromotedPinType : BindingPtr->PinType);
-								}
-							}
-							return FLinearColor::White;
-						})
-					]
-					+SHorizontalBox::Slot()
-					.AutoWidth()
-					.VAlign(VAlign_Center)
-					.Padding(3.0f, 2.0f)
+					// Add binding widget
+					PinWidget->GetLabelAndValue()->AddSlot()
 					[
-						SNew(STextBlock)
-						.Text_Lambda([ComparisonName, InAnimGraphNode]()
-						{
-							if (const FAnimGraphNodePropertyBinding* BindingPtr = InAnimGraphNode->PropertyBindings.Find(ComparisonName))
-							{
-								return BindingPtr->PathAsText;
-							}
-
-							return FText::GetEmpty();
-						})
-					]
-				];
+						UAnimGraphNode_Base::MakePropertyBindingWidget(Args)
+					];
+				}
 			}
 		}
 	}
diff --git a/Engine/Source/Editor/EditorStyle/Private/StarshipStyle.cpp b/Engine/Source/Editor/EditorStyle/Private/StarshipStyle.cpp
--- a/Engine/Source/Editor/EditorStyle/Private/StarshipStyle.cpp
+++ b/Engine/Source/Editor/EditorStyle/Private/StarshipStyle.cpp
@@ -5708,7 +5708,17 @@ void FStarshipEditorStyle::FStyle::SetupPersonaStyle()
 		Set("AnimGraph.Attribute.RootMotionDelta.Color", FStyleColors::AccentGreen);
 		Set("AnimGraph.Attribute.InertialBlending.Icon", new IMAGE_BRUSH("Icons/icon_tab_TimecodeProvider_16x", Icon10x10));
 		Set("AnimGraph.Attribute.InertialBlending.Color", FStyleColors::AccentOrange);
-	}	
+	}
+	
+	// Property Access 
+	{
+		Set("PropertyAccess.CompiledContext.Text", FTextBlockStyle(NormalText)
+			.SetFont(DEFAULT_FONT("Italic", 8))
+			.SetColorAndOpacity(FLinearColor(218.0f/255.0f,218.0f/255.0f,96.0f/255.0f, 0.5f))
+		);
+
+		Set("PropertyAccess.CompiledContext.Border", new FSlateRoundedBoxBrush(FStyleColors::DropdownOutline, 2.0f));
+	}
 #endif // WITH_EDITOR
 }
 
diff --git a/Engine/Source/Editor/GraphEditor/Private/KismetNodes/SGraphNodeK2Var.h b/Engine/Source/Editor/GraphEditor/Private/KismetNodes/SGraphNodeK2Var.h
deleted file mode 100644
diff --git a/Engine/Source/Editor/GraphEditor/Public/KismetNodes/SGraphNodeK2Var.h b/Engine/Source/Editor/GraphEditor/Public/KismetNodes/SGraphNodeK2Var.h
new file mode 100644
--- /dev/null
+++ b/Engine/Source/Editor/GraphEditor/Public/KismetNodes/SGraphNodeK2Var.h
@@ -0,0 +1,28 @@
+// Copyright Epic Games, Inc. All Rights Reserved.
+
+#pragma once
+
+#include "CoreMinimal.h"
+#include "Styling/SlateColor.h"
+#include "Widgets/DeclarativeSyntaxSupport.h"
+#include "KismetNodes/SGraphNodeK2Base.h"
+
+class UK2Node;
+
+class GRAPHEDITOR_API SGraphNodeK2Var : public SGraphNodeK2Base
+{
+public:
+	SLATE_BEGIN_ARGS(SGraphNodeK2Var){}
+	SLATE_END_ARGS()
+
+	void Construct( const FArguments& InArgs, UK2Node* InNode );
+
+	virtual const FSlateBrush* GetShadowBrush(bool bSelected) const override;
+
+	// SGraphNode interface
+	virtual void UpdateGraphNode() override;
+	// End of SGraphNode interface
+
+protected:
+	FSlateColor GetVariableColor() const;
+};
diff --git a/Engine/Source/Editor/Persona/Private/AnimGraphNodeDetails.cpp b/Engine/Source/Editor/Persona/Private/AnimGraphNodeDetails.cpp
--- a/Engine/Source/Editor/Persona/Private/AnimGraphNodeDetails.cpp
+++ b/Engine/Source/Editor/Persona/Private/AnimGraphNodeDetails.cpp
@@ -7,16 +7,13 @@
 #include "BoneContainer.h"
 #include "Engine/SkeletalMesh.h"
 #include "Animation/AnimationAsset.h"
-#include "Widgets/Layout/SSpacer.h"
 #include "DetailWidgetRow.h"
 #include "IDetailPropertyRow.h"
 #include "DetailLayoutBuilder.h"
 #include "DetailCategoryBuilder.h"
-#include "IDetailsView.h"
 #include "PropertyCustomizationHelpers.h"
 #include "SlateOptMacros.h"
 #include "Widgets/Images/SImage.h"
-#include "Widgets/Input/SEditableTextBox.h"
 #include "Widgets/Input/SButton.h"
 #include "Animation/AnimInstance.h"
 #include "Animation/EditorParentPlayerListObj.h"
@@ -27,10 +24,8 @@
 #include "AnimGraphNode_Base.h"
 #include "Widgets/Views/STreeView.h"
 #include "BoneSelectionWidget.h"
-#include "Widgets/Layout/SExpandableArea.h"
 #include "Animation/BlendProfile.h"
 #include "AnimGraphNode_AssetPlayerBase.h"
-#include "BlendProfilePicker.h"
 #include "ISkeletonEditorModule.h"
 #include "EdGraph/EdGraph.h"
 #include "BlueprintEditor.h"
@@ -43,15 +38,12 @@
 #include "AnimGraphNode_SequencePlayer.h"
 #include "Engine/SkeletalMeshSocket.h"
 #include "Styling/CoreStyle.h"
-#include "Widgets/Input/SCheckBox.h"
 #include "LODInfoUILayout.h"
 #include "IPersonaToolkit.h"
 #include "Interfaces/Interface_BoneReferenceSkeletonProvider.h"
 #include "IPropertyAccessEditor.h"
 #include "Algo/Accumulate.h"
 #include "Kismet2/BlueprintEditorUtils.h"
-#include "Framework/MultiBox/MultiBoxBuilder.h"
-#include "ScopedTransaction.h"
 
 #define LOCTEXT_NAMESPACE "KismetNodeWithOptionalPinsDetails"
 
@@ -1201,82 +1193,6 @@ bool FAnimGraphNodeBindingExtension::IsPropertyExtendable(const UClass* InObject
 	return false;
 }
 
-// Legacy binding widget
-class SShowAsWidget : public SCompoundWidget
-{
-	SLATE_BEGIN_ARGS(SShowAsWidget) {}
-
-	SLATE_END_ARGS()
-
-	void Construct(const FArguments& InArgs, const TSharedRef<IPropertyHandle>& InPropertyHandle)
-	{
-		PropertyHandle = InPropertyHandle;
-
-		TSharedRef<SHorizontalBox> HorizontalBox = 
-			SNew(SHorizontalBox)
-			.ToolTipText(LOCTEXT("AsPinTooltip", "Show/hide this property as a pin on the node"));
-
-		TWeakPtr<SWidget> WeakHorizontalBox = HorizontalBox;
-
-		HorizontalBox->AddSlot()
-			.AutoWidth()
-			.VAlign(VAlign_Center)
-			[
-				SNew(STextBlock)
-				.Text(LOCTEXT("ExposeAsPinLabel", "Expose"))
-				.Font(IDetailLayoutBuilder::GetDetailFont())
-				.Visibility_Lambda([WeakHorizontalBox](){ return WeakHorizontalBox.IsValid() && WeakHorizontalBox.Pin()->IsHovered() ? EVisibility::Visible : EVisibility::Collapsed; })
-			];
-
-		HorizontalBox->AddSlot()
-			.AutoWidth()
-			.VAlign(VAlign_Center)
-			.Padding(3.0f, 0.0f, 0.0f, 0.0f)
-			[
-				SNew(SCheckBox)
-				.IsChecked(this, &SShowAsWidget::IsChecked)
-				.OnCheckStateChanged(this, &SShowAsWidget::OnCheckStateChanged)
-			];
-
-		ChildSlot
-		[
-			HorizontalBox
-		];
-	}
-
-	ECheckBoxState IsChecked() const
-	{
-		bool bValue;
-		FPropertyAccess::Result Result = PropertyHandle->GetValue(bValue);
-		if(Result == FPropertyAccess::MultipleValues)
-		{
-			return ECheckBoxState::Undetermined;
-		}
-		else
-		{
-			return bValue ? ECheckBoxState::Checked : ECheckBoxState::Unchecked;
-		}
-
-		return ECheckBoxState::Unchecked;
-	}
-
-	void OnCheckStateChanged(ECheckBoxState InCheckBoxState)
-	{
-		bool bValue = InCheckBoxState == ECheckBoxState::Checked;
-		PropertyHandle->SetValue(bValue);
-	}
-
-	TSharedPtr<IPropertyHandle> PropertyHandle;
-};
-
-static FText MakeTextPath(const TArray<FString>& InPath)
-{
-	return FText::FromString(Algo::Accumulate(InPath, FString(), [](const FString& InResult, const FString& InSegment)
-		{ 
-			return InResult.IsEmpty() ? InSegment : (InResult + TEXT(".") + InSegment);
-		}));	
-}
-
 TSharedRef<SWidget> FAnimGraphNodeBindingExtension::GenerateExtensionWidget(const IDetailLayoutBuilder& InDetailBuilder, const UClass* InObjectClass, TSharedPtr<IPropertyHandle> InPropertyHandle)
 {
 	int32 OptionalPinIndex;
@@ -1287,410 +1203,35 @@ TSharedRef<SWidget> FAnimGraphNodeBindingExtension::GenerateExtensionWidget(cons
 	TArray<UObject*> OuterObjects;
 	InPropertyHandle->GetOuterObjects(OuterObjects);
 
-	FProperty* AnimNodeProperty = InPropertyHandle->GetProperty();
-	const FName PropertyName = AnimNodeProperty->GetFName();
-
-	const FName OptionalPinArrayEntryName(*FString::Printf(TEXT("ShowPinForProperties[%d].bShowPin"), OptionalPinIndex));
-	TSharedRef<IPropertyHandle> ShowPinPropertyHandle = InDetailBuilder.GetProperty(OptionalPinArrayEntryName, UAnimGraphNode_Base::StaticClass());
-	ShowPinPropertyHandle->MarkHiddenByCustomization();
+	TArray<UAnimGraphNode_Base*> AnimGraphNodes;
+	Algo::Transform(OuterObjects, AnimGraphNodes, [](UObject* InObject){ return Cast<UAnimGraphNode_Base>(InObject); });
 
-	UBlueprint* Blueprint = AnimGraphNode->GetAnimBlueprint();
 
-	if(IModularFeatures::Get().IsModularFeatureAvailable("PropertyAccessEditor"))
+	FProperty* AnimNodeProperty = InPropertyHandle->GetProperty();
+	TSharedPtr<IPropertyHandle> ParentHandle = InPropertyHandle->GetParentHandle();
+	FProperty* ParentProperty = ParentHandle.IsValid() ? ParentHandle->GetProperty() : nullptr;
+	
+	FName PropertyName;
+	if(FArrayProperty* ParentArrayProperty = CastField<FArrayProperty>(ParentProperty))
 	{
-		FPropertyBindingWidgetArgs Args;
-
-		Args.Property = InPropertyHandle->GetProperty();
-
-		Args.OnCanBindProperty = FOnCanBindProperty::CreateLambda([AnimNodeProperty](FProperty* InProperty)
-		{
-			// Note: We support type promotion here
-			IPropertyAccessEditor& PropertyAccessEditor = IModularFeatures::Get().GetModularFeature<IPropertyAccessEditor>("PropertyAccessEditor");
-			return PropertyAccessEditor.GetPropertyCompatibility(InProperty, AnimNodeProperty) != EPropertyAccessCompatibility::Incompatible;
-		});
-
-		Args.OnCanBindFunction = FOnCanBindFunction::CreateLambda([AnimNodeProperty](UFunction* InFunction)
-		{
-			IPropertyAccessEditor& PropertyAccessEditor = IModularFeatures::Get().GetModularFeature<IPropertyAccessEditor>("PropertyAccessEditor");
-
-			// Note: We support type promotion here
-			return InFunction->NumParms == 1 
-				&& PropertyAccessEditor.GetPropertyCompatibility(InFunction->GetReturnProperty(), AnimNodeProperty) != EPropertyAccessCompatibility::Incompatible
-				&& InFunction->HasAnyFunctionFlags(FUNC_BlueprintPure);
-		});
-
-		Args.OnCanBindToClass = FOnCanBindToClass::CreateLambda([](UClass* InClass)
-		{
-			return true;
-		});
-
-		Args.OnAddBinding = FOnAddBinding::CreateLambda([OuterObjects, Blueprint, ShowPinPropertyHandle, AnimNodeProperty](FName InPropertyName, const TArray<FBindingChainElement>& InBindingChain)
-		{
-			const UEdGraphSchema_K2* Schema = GetDefault<UEdGraphSchema_K2>();
-			IPropertyAccessEditor& PropertyAccessEditor = IModularFeatures::Get().GetModularFeature<IPropertyAccessEditor>("PropertyAccessEditor");
-
-			// Disconnect any connected pin since we will replace it with the bound property/function
-			ShowPinPropertyHandle->SetValue(false);
-
-			for(UObject* OuterObject : OuterObjects)
-			{
-				if(UAnimGraphNode_Base* AnimGraphNode = Cast<UAnimGraphNode_Base>(OuterObject))
-				{
-					AnimGraphNode->Modify();
-
-					const FFieldVariant& LeafField = InBindingChain.Last().Field;
-
-					FAnimGraphNodePropertyBinding Binding;
-					Binding.PropertyName = InPropertyName;
-					PropertyAccessEditor.MakeStringPath(InBindingChain, Binding.PropertyPath);
-					Binding.PathAsText = MakeTextPath(Binding.PropertyPath);
-					Binding.Type = LeafField.IsA<UFunction>() ? EAnimGraphNodePropertyBindingType::Function : EAnimGraphNodePropertyBindingType::Property;
-					Binding.bIsBound = true;
-					if(LeafField.IsA<FProperty>())
-					{
-						const FProperty* LeafProperty = LeafField.Get<FProperty>();
-						if(LeafProperty)
-						{
-							if(PropertyAccessEditor.GetPropertyCompatibility(LeafProperty, AnimNodeProperty) == EPropertyAccessCompatibility::Promotable)
-							{
-								Binding.bIsPromotion = true;
-								Schema->ConvertPropertyToPinType(LeafProperty, Binding.PromotedPinType);
-							}
-
-							Schema->ConvertPropertyToPinType(LeafProperty, Binding.PinType);
-						}
-					}
-					else if(LeafField.IsA<UFunction>())
-					{
-						const UFunction* LeafFunction = LeafField.Get<UFunction>();
-						if(LeafFunction)
-						{
-							if(FProperty* ReturnProperty = LeafFunction->GetReturnProperty())
-							{
-								if(PropertyAccessEditor.GetPropertyCompatibility(ReturnProperty, AnimNodeProperty) == EPropertyAccessCompatibility::Promotable)
-								{
-									Binding.bIsPromotion = true;
-									Schema->ConvertPropertyToPinType(ReturnProperty, Binding.PromotedPinType);
-								}
-
-								Schema->ConvertPropertyToPinType(ReturnProperty, Binding.PinType);
-							}
-						}
-					}
-					AnimGraphNode->PropertyBindings.Add(InPropertyName, Binding);
-				}
-			}
-			FBlueprintEditorUtils::MarkBlueprintAsModified(Blueprint);
-			
-			// Pins are exposed if we have a binding or not - and after running this we do.
-			ShowPinPropertyHandle->SetValue(true);
-		});
-
-		Args.OnRemoveBinding = FOnRemoveBinding::CreateLambda([OuterObjects, Blueprint](FName InPropertyName)
-		{
-			for(UObject* OuterObject : OuterObjects)
-			{
-				if(UAnimGraphNode_Base* AnimGraphNode = Cast<UAnimGraphNode_Base>(OuterObject))
-				{
-					AnimGraphNode->Modify();
-
-					AnimGraphNode->PropertyBindings.Remove(InPropertyName);
-				}
-			}
-
-			FBlueprintEditorUtils::MarkBlueprintAsModified(Blueprint);
-		});
-
-		Args.OnCanRemoveBinding = FOnCanRemoveBinding::CreateLambda([OuterObjects](FName InPropertyName)
-		{
-			for(UObject* OuterObject : OuterObjects)
-			{
-				if(UAnimGraphNode_Base* AnimGraphNode = Cast<UAnimGraphNode_Base>(OuterObject))
-				{
-					if(AnimGraphNode->PropertyBindings.Contains(InPropertyName))
-					{
-						return true;
-					}
-				}
-			}
-
-			return false;
-		});
-
-		enum class ECurrentValueType : int32
-		{
-			None,
-			Pin,
-			Binding,
-			MultipleValues,
-		};
-
-		Args.CurrentBindingText = MakeAttributeLambda([OuterObjects, PropertyName, ShowPinPropertyHandle]()
-		{
-			ECurrentValueType CurrentValueType = ECurrentValueType::None;
-
-			const FText MultipleValues = LOCTEXT("MultipleValues", "Multiple Values");
-			const FText Bind = LOCTEXT("Bind", "Bind");
-			const FText ExposedAsPin = LOCTEXT("ExposedAsPin", "Exposed As Pin");
-			FText CurrentValue = Bind;
-
-			auto SetAssignValue = [&CurrentValueType, &CurrentValue, &MultipleValues](const FText& InValue, ECurrentValueType InType)
-			{
-				if(CurrentValueType != ECurrentValueType::MultipleValues)
-				{
-					if(CurrentValueType == ECurrentValueType::None)
-					{
-						CurrentValueType = InType;
-						CurrentValue = InValue;
-					}
-					else if(CurrentValueType == InType)
-					{
-						if(!CurrentValue.EqualTo(InValue))
-						{
-							CurrentValueType = ECurrentValueType::MultipleValues;
-							CurrentValue = MultipleValues;
-						}
-					}
-					else
-					{
-						CurrentValueType = ECurrentValueType::MultipleValues;
-						CurrentValue = MultipleValues;
-					}
-				}
-			};
-
-			for(UObject* OuterObject : OuterObjects)
-			{
-				if(UAnimGraphNode_Base* AnimGraphNode = Cast<UAnimGraphNode_Base>(OuterObject))
-				{
-					if(FAnimGraphNodePropertyBinding* BindingPtr = AnimGraphNode->PropertyBindings.Find(PropertyName))
-					{
-						SetAssignValue(BindingPtr->PathAsText, ECurrentValueType::Binding);
-					}
-					else
-					{
-						bool bAsPin = false;
-						FPropertyAccess::Result Result = ShowPinPropertyHandle->GetValue(bAsPin);
-						if(Result == FPropertyAccess::MultipleValues)
-						{
-							SetAssignValue(MultipleValues, ECurrentValueType::MultipleValues);
-						}
-						else if(bAsPin)
-						{
-							SetAssignValue(ExposedAsPin, ECurrentValueType::Pin);
-						}
-						else
-						{
-							SetAssignValue(Bind, ECurrentValueType::None);
-						}
-					}
-				}
-			}
-
-			return CurrentValue;
-		});
-
-		Args.CurrentBindingImage = MakeAttributeLambda([OuterObjects, PropertyName, OptionalPinIndex]() -> const FSlateBrush*
-		{
-			static FName PropertyIcon(TEXT("Kismet.Tabs.Variables"));
-			static FName FunctionIcon(TEXT("GraphEditor.Function_16x"));
-
-			EAnimGraphNodePropertyBindingType BindingType = EAnimGraphNodePropertyBindingType::None;
-			for(UObject* OuterObject : OuterObjects)
-			{
-				if(UAnimGraphNode_Base* AnimGraphNode = Cast<UAnimGraphNode_Base>(OuterObject))
-				{
-					if(AnimGraphNode->ShowPinForProperties[OptionalPinIndex].bShowPin)
-					{
-						BindingType = EAnimGraphNodePropertyBindingType::None;
-						break;
-					}
-					else if(FAnimGraphNodePropertyBinding* BindingPtr = AnimGraphNode->PropertyBindings.Find(PropertyName))
-					{
-						if(BindingType == EAnimGraphNodePropertyBindingType::None)
-						{
-							BindingType = BindingPtr->Type;
-						}
-						else if(BindingType != BindingPtr->Type)
-						{
-							BindingType = EAnimGraphNodePropertyBindingType::None;
-							break;
-						}
-					}
-					else if(BindingType != EAnimGraphNodePropertyBindingType::None)
-					{
-						BindingType = EAnimGraphNodePropertyBindingType::None;
-						break;
-					}
-				}
-			}
-
-			if (BindingType == EAnimGraphNodePropertyBindingType::Function)
-			{
-				return FEditorStyle::GetBrush(FunctionIcon);
-			}
-			else
-			{
-				return FEditorStyle::GetBrush(PropertyIcon);
-			}
-		});
-
-		Args.CurrentBindingColor = MakeAttributeLambda([OuterObjects, InPropertyHandle, OptionalPinIndex, PropertyName]() -> FLinearColor
-		{
-			const UEdGraphSchema_K2* Schema = GetDefault<UEdGraphSchema_K2>();
-
-			FEdGraphPinType PinType;
-			Schema->ConvertPropertyToPinType(InPropertyHandle->GetProperty(), PinType);
-			FLinearColor BindingColor = Schema->GetPinTypeColor(PinType);
-
-			enum class EPromotionState
-			{
-				NotChecked,
-				NotPromoted,
-				Promoted,
-			} Promotion = EPromotionState::NotChecked;
-
-			for(UObject* OuterObject : OuterObjects)
-			{
-				if(UAnimGraphNode_Base* AnimGraphNode = Cast<UAnimGraphNode_Base>(OuterObject))
-				{
-					if(AnimGraphNode->ShowPinForProperties[OptionalPinIndex].bShowPin)
-					{
-						if(Promotion == EPromotionState::NotChecked)
-						{
-							Promotion = EPromotionState::NotPromoted;
-						}
-						else if(Promotion == EPromotionState::Promoted)
-						{
-							BindingColor = FLinearColor::Gray;
-							break;
-						}
-					}
-					else if(FAnimGraphNodePropertyBinding* BindingPtr = AnimGraphNode->PropertyBindings.Find(PropertyName))
-					{
-						if(Promotion == EPromotionState::NotChecked)
-						{
-							if(BindingPtr->bIsPromotion)
-							{
-								Promotion = EPromotionState::Promoted;
-								BindingColor = Schema->GetPinTypeColor(BindingPtr->PromotedPinType);
-							}
-							else
-							{
-								Promotion = EPromotionState::NotPromoted;
-							}
-						}
-						else
-						{
-							EPromotionState NewPromotion = BindingPtr->bIsPromotion ? EPromotionState::Promoted : EPromotionState::NotPromoted;
-							if(Promotion != NewPromotion)
-							{
-								BindingColor = FLinearColor::Gray;
-								break;
-							}
-						}
-					}
-				}
-			}
-
-			return BindingColor;
-		});
-
-		Args.MenuExtender = MakeShared<FExtender>();
-		Args.MenuExtender->AddMenuExtension("BindingActions", EExtensionHook::Before, nullptr, FMenuExtensionDelegate::CreateLambda([ShowPinPropertyHandle, OuterObjects, PropertyName, Blueprint](FMenuBuilder& InMenuBuilder)
-		{
-			InMenuBuilder.BeginSection("Pins", LOCTEXT("Pin", "Pin"));
-			{
-				InMenuBuilder.AddMenuEntry(
-					LOCTEXT("ExposeAsPin", "Expose As Pin"),
-					LOCTEXT("ExposeAsPinTooltip", "Show/hide this property as a pin on the node"),
-					FSlateIcon("EditorStyle", "GraphEditor.PinIcon"),
-					FUIAction(
-						FExecuteAction::CreateLambda([ShowPinPropertyHandle, OuterObjects, PropertyName, Blueprint]()
-						{
-							bool bValue = false;
-							ShowPinPropertyHandle->GetValue(bValue);
-
-							bool bHasBinding = false;
-
-							for(UObject* OuterObject : OuterObjects)
-							{
-								if(UAnimGraphNode_Base* AnimGraphNode = Cast<UAnimGraphNode_Base>(OuterObject))
-								{
-									bHasBinding |= AnimGraphNode->PropertyBindings.Find(PropertyName) != nullptr;
-								}
-							}
-
-							{
-								FScopedTransaction Transaction(LOCTEXT("PinExposure", "Pin Exposure"));
-
-								// Pins are exposed if we have a binding or not, so treat as unchecked only if we have
-								// no binding
-								ShowPinPropertyHandle->SetValue(!bValue || bHasBinding);
-
-								// Switching from non-pin to pin, remove any bindings
-								for(UObject* OuterObject : OuterObjects)
-								{
-									if(UAnimGraphNode_Base* AnimGraphNode = Cast<UAnimGraphNode_Base>(OuterObject))
-									{
-										AnimGraphNode->Modify();
-
-										AnimGraphNode->PropertyBindings.Remove(PropertyName);
-									}
-								}
-
-								FBlueprintEditorUtils::MarkBlueprintAsModified(Blueprint);
-							}
-						}),
-						FCanExecuteAction(),
-						FGetActionCheckState::CreateLambda([ShowPinPropertyHandle, OuterObjects, PropertyName]()
-						{
-							bool bValue;
-							FPropertyAccess::Result Result = ShowPinPropertyHandle->GetValue(bValue);
-							if(Result == FPropertyAccess::MultipleValues)
-							{
-								return ECheckBoxState::Undetermined;
-							}
-							else
-							{
-								bool bHasBinding = false;
-
-								for(UObject* OuterObject : OuterObjects)
-								{
-									if(UAnimGraphNode_Base* AnimGraphNode = Cast<UAnimGraphNode_Base>(OuterObject))
-									{
-										bHasBinding |= AnimGraphNode->PropertyBindings.Find(PropertyName) != nullptr;
-									}
-								}
-
-								// Pins are exposed if we have a binding or not, so treat as unchecked only if we have
-								// no binding
-								bValue = bValue && !bHasBinding;
-
-								return bValue ? ECheckBoxState::Checked : ECheckBoxState::Unchecked;
-							}
-
-							return ECheckBoxState::Unchecked;
-						})
-					),
-					NAME_None,
-					EUserInterfaceActionType::Check
-				);
-			}
-			InMenuBuilder.EndSection();
-		}));
-
-		Args.bAllowNewBindings = false;
-		Args.bAllowArrayElementBindings = true;
-		Args.bAllowUObjectFunctions = true;
-
-		IPropertyAccessEditor& PropertyAccessEditor = IModularFeatures::Get().GetModularFeature<IPropertyAccessEditor>("PropertyAccessEditor");
-		return PropertyAccessEditor.MakePropertyBindingWidget(AnimGraphNode->GetAnimBlueprint(), Args);
+		int32 ArrayIndex = InPropertyHandle->GetIndexInArray();
+		check(ArrayIndex != INDEX_NONE);
+		PropertyName = FName(ParentArrayProperty->GetFName(), InPropertyHandle->GetIndexInArray() + 1);
+		AnimNodeProperty = ParentArrayProperty;
 	}
 	else
 	{
-		return SNew(SShowAsWidget, ShowPinPropertyHandle);
+		PropertyName = AnimNodeProperty->GetFName();
 	}
+
+	const FName OptionalPinArrayEntryName(*FString::Printf(TEXT("ShowPinForProperties[%d].bShowPin"), OptionalPinIndex));
+	TSharedRef<IPropertyHandle> ShowPinPropertyHandle = InDetailBuilder.GetProperty(OptionalPinArrayEntryName, UAnimGraphNode_Base::StaticClass());
+	ShowPinPropertyHandle->MarkHiddenByCustomization();
+
+	UAnimGraphNode_Base::FAnimPropertyBindingWidgetArgs Args(AnimGraphNodes, AnimNodeProperty, PropertyName, OptionalPinIndex);
+	Args.bOnGraphNode = false;
+	
+	return UAnimGraphNode_Base::MakePropertyBindingWidget(Args);
 }
 
 #undef LOCTEXT_NAMESPACE
diff --git a/Engine/Source/Editor/UnrealEd/Public/IPropertyAccessBlueprintBinding.h b/Engine/Source/Editor/UnrealEd/Public/IPropertyAccessBlueprintBinding.h
new file mode 100644
--- /dev/null
+++ b/Engine/Source/Editor/UnrealEd/Public/IPropertyAccessBlueprintBinding.h
@@ -0,0 +1,65 @@
+// Copyright Epic Games, Inc. All Rights Reserved.
+
+#pragma once
+
+#include "CoreMinimal.h"
+
+#include "EdGraph/EdGraph.h"
+#include "Features/IModularFeatures.h"
+
+class FExtender;
+class UBlueprint;
+class UEdGraph;
+class UEdGraphNode;
+class UEdGraphPin;
+
+// Delegate used to set the context ID of a property access
+DECLARE_DELEGATE_OneParam(FOnSetPropertyAccessContextId, FName /*InContextId*/);
+
+// Delegate used to set the context ID of a property access
+DECLARE_DELEGATE_RetVal_OneParam(bool, FOnCanSetPropertyAccessContextId, FName /*InContextId*/);
+
+// Delegate used to get the context ID of a property access
+DECLARE_DELEGATE_RetVal(FName, FOnGetPropertyAccessContextId);
+
+// Modular feature allowing property access to bind to a Blueprint
+class IPropertyAccessBlueprintBinding : public IModularFeature
+{
+public:
+	virtual ~IPropertyAccessBlueprintBinding() {}
+	
+	// A context in which a Blueprint binding can take place
+	struct FContext
+	{
+		// The blueprint we would like to bind to
+		const UBlueprint* Blueprint = nullptr;
+
+		// The graph we would like to bind to
+		const UEdGraph* Graph = nullptr;
+
+		// The node we would like to bind to
+		const UEdGraphNode* Node = nullptr;
+
+		// The pin we would like to bind to
+		const UEdGraphPin* Pin = nullptr;
+	};
+	
+	// @return true if this binding can bind to the specified context
+	virtual bool CanBindToContext(const FContext& InContext) const = 0;
+
+	// Args passed to MakeBindingMenuExtender
+	struct FBindingMenuArgs
+	{
+		// Delegate used to set the context ID of a property access
+		FOnSetPropertyAccessContextId OnSetPropertyAccessContextId;
+
+		// Delegate used to control the enabled state of context IDs
+		FOnCanSetPropertyAccessContextId OnCanSetPropertyAccessContextId;
+		
+		// Delegate used to get the context ID of a property access
+		FOnGetPropertyAccessContextId OnGetPropertyAccessContextId;
+	};
+	
+	// Make a menu extender to be used in property access binding menus
+	virtual TSharedPtr<FExtender> MakeBindingMenuExtender(const FContext& InContext, const FBindingMenuArgs& InArgs) const = 0;
+};
\ No newline at end of file
diff --git a/Engine/Source/Editor/UnrealEd/Public/IPropertyAccessCompiler.h b/Engine/Source/Editor/UnrealEd/Public/IPropertyAccessCompiler.h
--- a/Engine/Source/Editor/UnrealEd/Public/IPropertyAccessCompiler.h
+++ b/Engine/Source/Editor/UnrealEd/Public/IPropertyAccessCompiler.h
@@ -10,7 +10,7 @@
 class UObject;
 struct FPropertyAccessLibrary;
 
-// The various batching of property copy
+// DEPRECATED - The various batching of property copy
 UENUM()
 enum class EPropertyAccessBatchType : uint8
 {
@@ -21,6 +21,69 @@ enum class EPropertyAccessBatchType : uint8
 	Batched,
 };
 
+// Handle used to track property accesses that are being compiled
+struct FPropertyAccessHandle
+{
+public:
+	// Get the internal index
+	int32 GetId() const { return Id; }
+
+	// Check this handle for validity
+	bool IsValid() const { return Id != INDEX_NONE; }
+
+	FPropertyAccessHandle() = default;
+	
+	FPropertyAccessHandle(int32 InId)
+		: Id(InId)
+	{
+	}
+
+	inline bool operator==(const FPropertyAccessHandle& Other) const
+	{
+		return Id == Other.Id;
+	} 
+
+private:
+	// Index used to track accesses by external systems
+	int32 Id = INDEX_NONE;
+};
+
+inline uint32 GetTypeHash(const FPropertyAccessHandle& Value)
+{
+	return GetTypeHash(Value.GetId());
+}
+
+// Handle used to describe property accesses that have been compiled
+struct FCompiledPropertyAccessHandle
+{
+public:
+	FCompiledPropertyAccessHandle() = default;
+	
+	// Get the index into the batch
+	int32 GetId() const { return Id; }
+
+	// Get the index of the batch
+	int32 GetBatchId() const { return BatchId; }
+
+	// Check this handle for validity
+	bool IsValid() const { return Id != INDEX_NONE && BatchId != INDEX_NONE; }
+
+private:
+	friend class FPropertyAccessLibraryCompiler;
+
+	FCompiledPropertyAccessHandle(int32 InId, int32 InBatchId)
+		: Id(InId)
+		, BatchId(InBatchId)
+	{
+	}
+
+	// Index into the batch
+	int32 Id = INDEX_NONE;
+
+	// Index of the batch
+	int32 BatchId = INDEX_NONE;
+};
+
 // A helper used to compile a property access library
 class IPropertyAccessLibraryCompiler
 {
@@ -30,12 +93,19 @@ public:
 	// Begin compilation - reset the library to its default state
 	virtual void BeginCompilation() = 0;
 	
-	UE_DEPRECATED(5.0, "Please use BeginCompilation without a class argument")
+	UE_DEPRECATED(5.0, "Please use BeginCompilation without a class arg")
 	virtual void BeginCompilation(const UClass* InClass) {}
 
 	// Add a copy to the property access library we are compiling
 	// @return an integer handle to the pending copy. This can be resolved to a true copy index by calling MapCopyIndex
-	virtual int32 AddCopy(TArrayView<FString> InSourcePath, TArrayView<FString> InDestPath, EPropertyAccessBatchType InBatchType, UObject* InAssociatedObject = nullptr) = 0;
+	virtual FPropertyAccessHandle AddCopy(TArrayView<FString> InSourcePath, TArrayView<FString> InDestPath, const FName& InContextId, UObject* InAssociatedObject = nullptr) = 0;
+	
+	UE_DEPRECATED(5.0, "Please use AddCopy with a context ID that returns a handle")
+	virtual int32 AddCopy(TArrayView<FString> InSourcePath, TArrayView<FString> InDestPath, EPropertyAccessBatchType InBatchType, UObject* InAssociatedObject = nullptr)
+	{
+		FPropertyAccessHandle Handle = AddCopy(InSourcePath, InDestPath, FName(NAME_None), InAssociatedObject); 
+		return Handle.GetId();
+	}
 
 	// Post-process the library to finish compilation. @return true if compilation succeeded.
 	virtual bool FinishCompilation() = 0;
@@ -43,6 +113,13 @@ public:
 	// Iterate any errors we have with compilation
 	virtual void IterateErrors(TFunctionRef<void(const FText&, UObject*)> InFunction) const = 0;
 
-	// Maps the initial integer copy handle to a true handle, post compilation
-	virtual int32 MapCopyIndex(int32 InIndex) const = 0;
+	// Maps the initial copy handle to a true handle, post compilation
+	virtual FCompiledPropertyAccessHandle GetCompiledHandle(FPropertyAccessHandle InHandle) const = 0;
+
+	UE_DEPRECATED(5.0, "Please use GetCompiledHandle")
+	virtual int32 MapCopyIndex(int32 InIndex) const
+	{
+		FCompiledPropertyAccessHandle CompiledHandle = GetCompiledHandle(FPropertyAccessHandle(InIndex));
+		return CompiledHandle.GetId();
+	}
 };
\ No newline at end of file
diff --git a/Engine/Source/Editor/UnrealEd/Public/IPropertyAccessEditor.h b/Engine/Source/Editor/UnrealEd/Public/IPropertyAccessEditor.h
--- a/Engine/Source/Editor/UnrealEd/Public/IPropertyAccessEditor.h
+++ b/Engine/Source/Editor/UnrealEd/Public/IPropertyAccessEditor.h
@@ -160,6 +160,9 @@ struct FPropertyBindingWidgetArgs
 	/** The current binding's text label */
 	TAttribute<FText> CurrentBindingText;
 
+	/** The current binding's tooltip text label */
+	TAttribute<FText> CurrentBindingToolTipText;
+	
 	/** The current binding's image */
 	TAttribute<const FSlateBrush*> CurrentBindingImage;
 
@@ -183,19 +186,38 @@ struct FPropertyBindingWidgetArgs
 
 	/** Whether to allow UObject functions as non-leaf nodes */
 	bool bAllowUObjectFunctions = false;
+
+	/** Whether to allow UScriptStruct functions as non-leaf nodes */
+	bool bAllowStructFunctions = false;	
 };
 
-/** Enum describing the result of ResolveLeafProperty */
+/** Enum describing the result of ResolvePropertyAccess */
 enum class EPropertyAccessResolveResult
 {
 	/** Resolution of the path failed */
 	Failed,
 
-	/** Resolution of the path failed and the property is internal to the initial context */
+	/** DEPRECATED - Resolution of the path succeeded and the property is internal to the initial context */
 	SucceededInternal,
 
-	/** Resolution of the path failed and the property is external to the initial context (i.e. uses an object/redirector indirection) */
+	/** DEPRECATED - Resolution of the path failed and the property is external to the initial context (i.e. uses an object/redirector indirection) */
 	SucceededExternal,
+
+	/** Resolution of the path succeeded */
+	Succeeded,
+};
+
+/**
+ * Result of a property access resolve. Provides information about what the compiler was able to determine about the
+ * property access.
+ */
+struct FPropertyAccessResolveResult
+{
+	// The success fail of the resolve
+	EPropertyAccessResolveResult Result = EPropertyAccessResolveResult::Failed;
+
+	// Whether the resolve was determined to be thread safe
+	bool bIsThreadSafe = false;
 };
 
 /** Enum describing property compatibility */
@@ -211,6 +233,32 @@ enum class EPropertyAccessCompatibility
 	Promotable,
 };
 
+// Context struct describing relevant characteristics of a property copy.
+// Used by client code to determine batch ID when called back via FOnPropertyAccessDetermineBatchId
+struct FPropertyAccessCopyContext
+{
+	// The object (usually a K2 node) in which the context takes place. Used for error reporting.
+	UObject* Object;
+	
+	// User-define context name, passed via IPropertyAccessLibraryCompiler::AddCopy
+	FName ContextId;
+
+	// Source path as text, for error reporting
+	FText SourcePathAsText;
+
+	// Dest path as text, for error reporting
+	FText DestPathAsText;
+	
+	// Whether the source path is thread safe
+	bool bSourceThreadSafe;
+
+	// Whether the dest path is thread safe
+	bool bDestThreadSafe;
+};
+
+// Delegate used to determine batch ID (index) for a particular copy context
+DECLARE_DELEGATE_RetVal_OneParam(int32, FOnPropertyAccessDetermineBatchId, const FPropertyAccessCopyContext& /*InContext*/);
+
 // Context used to create a property access library compiler
 struct FPropertyAccessLibraryCompilerArgs
 {
@@ -224,6 +272,10 @@ struct FPropertyAccessLibraryCompilerArgs
 
 	// The class that provides a root context for the library to be built in
 	const UClass* ClassContext;
+
+	// Delegate used to determine batch ID (index) for a particular copy context. If this is not set, then all copies
+	// will be batched together in batch 0
+	FOnPropertyAccessDetermineBatchId OnDetermineBatchId;
 };
 
 /** Editor support for property access system */
@@ -247,9 +299,35 @@ public:
 	 * @return a new binding widget
 	 */
 	virtual TSharedRef<SWidget> MakePropertyBindingWidget(const TArray<FBindingContextStruct>& InBindingContextStructs, const FPropertyBindingWidgetArgs& InArgs = FPropertyBindingWidgetArgs()) const = 0;
+	
+	/** Resolve a property access, returning the leaf property and array index if any. @return result structure for more information about the result */
+	virtual FPropertyAccessResolveResult ResolvePropertyAccess(const UStruct* InStruct, TArrayView<FString> InPath, FProperty*& OutProperty, int32& OutArrayIndex) const = 0;
+
+	/** Args used to resolve property access segments. The various functions will be called per-segment when resolved. */
+	struct FResolvePropertyAccessArgs
+	{
+		/** Function called when a property is resolved. Array index is valid for static array properties. For dynamic array properties, use ArrayFunction. */
+		TFunction<void(int32 /*SegmentIndex*/, FProperty* /*Property*/, int32 /*StaticArrayIndex*/)> PropertyFunction;
+
+		/** Function called when a dynamic array is resolved. */
+		TFunction<void(int32 /*SegmentIndex*/, FArrayProperty* /*Property*/, int32 /*ArrayIndex*/)> ArrayFunction;
+
+		/** Function called when a function is resolved. */
+		TFunction<void(int32 /*SegmentIndex*/, UFunction* /*Function*/, FProperty* /*ReturnProperty*/)> FunctionFunction;
+	};
+	
+	/**
+	 * Resolve a property path to a structure, calling back for each segment in path segment order if resolution succeed
+	 * @return true if resolution succeeded
+	 */
+	virtual FPropertyAccessResolveResult ResolvePropertyAccess(const UStruct* InStruct, TArrayView<FString> InPath, const FResolvePropertyAccessArgs& InArgs) const = 0;
 
-	/** Resolve a property path to a structure, returning the leaf property and array index if any. @return true if resolution succeeded */
-	virtual EPropertyAccessResolveResult ResolveLeafProperty(const UStruct* InStruct, TArrayView<FString> InPath, FProperty*& OutProperty, int32& OutArrayIndex) const = 0;
+	UE_DEPRECATED(5.0, "Please use ResolvePropertyAccess")
+	virtual EPropertyAccessResolveResult ResolveLeafProperty(const UStruct* InStruct, TArrayView<FString> InPath, FProperty*& OutProperty, int32& OutArrayIndex) const
+	{
+		const FPropertyAccessResolveResult Result = ResolvePropertyAccess(InStruct, InPath, OutProperty, OutArrayIndex);
+		return Result.Result;
+	}
 
 	// Get the compatibility of the two supplied properties. Ordering matters for promotion (A->B).
 	virtual EPropertyAccessCompatibility GetPropertyCompatibility(const FProperty* InPropertyA, const FProperty* InPropertyB) const = 0;
@@ -262,4 +340,7 @@ public:
 
 	// Make a property access library compiler, used for building a FPropertyAccessLibrary
 	virtual TUniquePtr<IPropertyAccessLibraryCompiler> MakePropertyAccessCompiler(const FPropertyAccessLibraryCompilerArgs& InArgs) const = 0;
+
+	// Make a text representation of a property path 
+	virtual FText MakeTextPath(const TArray<FString>& InPath) const = 0;
 };
\ No newline at end of file
diff --git a/Engine/Source/Runtime/AnimGraphRuntime/Private/AnimNodes/AnimNode_CallFunction.cpp b/Engine/Source/Runtime/AnimGraphRuntime/Private/AnimNodes/AnimNode_CallFunction.cpp
new file mode 100644
--- /dev/null
+++ b/Engine/Source/Runtime/AnimGraphRuntime/Private/AnimNodes/AnimNode_CallFunction.cpp
@@ -0,0 +1,92 @@
+﻿// Copyright Epic Games, Inc. All Rights Reserved.
+
+#include "AnimNodes/AnimNode_CallFunction.h"
+#include "Animation/AnimInstanceProxy.h"
+
+void FAnimNode_CallFunction::OnInitializeAnimInstance(const FAnimInstanceProxy* InProxy, const UAnimInstance* InAnimInstance)
+{
+	if(FunctionName != NAME_None)
+	{
+		Function = InAnimInstance->FindFunction(FunctionName);
+	}
+	
+	Counter.Reset();
+	CurrentWeight = 0.0f;
+}
+
+void FAnimNode_CallFunction::GatherDebugData(FNodeDebugData& DebugData)
+{
+	DebugData.AddDebugItem(TEXT("CallFunction"));
+	
+	Source.GatherDebugData(DebugData);
+}
+
+void FAnimNode_CallFunction::CallFunctionFromCallSite(EAnimFunctionCallSite InCallSite, const FAnimationBaseContext& InContext) const
+{
+	if(CallSite == InCallSite && Function != nullptr)
+	{
+		UObject* AnimInstanceObject = InContext.AnimInstanceProxy->GetAnimInstanceObject();
+		AnimInstanceObject->ProcessEvent(Function, nullptr);
+	}
+}
+
+void FAnimNode_CallFunction::Update_AnyThread(const FAnimationUpdateContext& InContext)
+{
+	CallFunctionFromCallSite(EAnimFunctionCallSite::OnUpdate, InContext);
+
+	bool bBecameRelevant = false;
+	if(!Counter.HasEverBeenUpdated() || !Counter.WasSynchronizedCounter(InContext.AnimInstanceProxy->GetUpdateCounter()))
+	{
+		bBecameRelevant = true;
+		CurrentWeight = 0.0f;
+		CallFunctionFromCallSite(EAnimFunctionCallSite::OnBecomeRelevant, InContext);
+	}
+
+	const float NewWeight = InContext.GetFinalBlendWeight();
+	if(FAnimWeight::IsFullWeight(CurrentWeight) && !FAnimWeight::IsFullWeight(NewWeight))
+	{
+		CallFunctionFromCallSite(EAnimFunctionCallSite::OnStartedBlendingOut, InContext);
+	}
+
+	if(!FAnimWeight::IsRelevant(CurrentWeight) && FAnimWeight::IsRelevant(NewWeight))
+	{
+		CallFunctionFromCallSite(EAnimFunctionCallSite::OnStartedBlendingIn, InContext);
+	}
+
+	if(!FAnimWeight::IsFullWeight(CurrentWeight) && FAnimWeight::IsFullWeight(NewWeight))
+	{
+		CallFunctionFromCallSite(EAnimFunctionCallSite::OnFinishedBlendingIn, InContext);
+	}
+
+	Source.Update(InContext);
+
+	if(bBecameRelevant)
+	{
+		CallFunctionFromCallSite(EAnimFunctionCallSite::OnBecomeRelevantPostRecursion, InContext);
+	}
+
+	CallFunctionFromCallSite(EAnimFunctionCallSite::OnUpdatePostRecursion, InContext);
+
+	Counter.SynchronizeWith(InContext.AnimInstanceProxy->GetUpdateCounter());
+	CurrentWeight = NewWeight;
+}
+
+void FAnimNode_CallFunction::Evaluate_AnyThread(FPoseContext& InContext)
+{
+	CallFunctionFromCallSite(EAnimFunctionCallSite::OnEvaluate, InContext);
+	
+	Source.Evaluate(InContext);
+
+	CallFunctionFromCallSite(EAnimFunctionCallSite::OnEvaluatePostRecursion, InContext);
+}
+
+void FAnimNode_CallFunction::Initialize_AnyThread(const FAnimationInitializeContext& InContext)
+{
+	CurrentWeight = 0.0f;
+	
+	CallFunctionFromCallSite(EAnimFunctionCallSite::OnInitialize, InContext);
+	
+	Source.Initialize(InContext);
+
+	CallFunctionFromCallSite(EAnimFunctionCallSite::OnInitializePostRecursion, InContext);
+}
diff --git a/Engine/Source/Runtime/AnimGraphRuntime/Private/KismetAnimationLibrary.cpp b/Engine/Source/Runtime/AnimGraphRuntime/Private/KismetAnimationLibrary.cpp
--- a/Engine/Source/Runtime/AnimGraphRuntime/Private/KismetAnimationLibrary.cpp
+++ b/Engine/Source/Runtime/AnimGraphRuntime/Private/KismetAnimationLibrary.cpp
@@ -215,26 +215,37 @@ float UKismetAnimationLibrary::K2_CalculateVelocityFromSockets(
 	return VelocityMin;
 }
 
+struct FK2ProfilingTimer
+{
+	double LastTime;
+	double AccummulatedTime;
+};
 
-TArray<UKismetAnimationLibrary::FK2ProfilingTimer> UKismetAnimationLibrary::sProfilingTimers;
+class FProfilingTimerPerThread : public TThreadSingleton<FProfilingTimerPerThread>
+{
+public:
+	TArray<FK2ProfilingTimer> ProfilingTimers;
+};
 
 void UKismetAnimationLibrary::K2_StartProfilingTimer()
 {
 	FK2ProfilingTimer Timer;
 	Timer.LastTime = FPlatformTime::Seconds() * 1000.0;
 	Timer.AccummulatedTime = 0.0;
-	sProfilingTimers.Add(Timer);
+	FProfilingTimerPerThread::Get().ProfilingTimers.Add(Timer);
 }
 
 float UKismetAnimationLibrary::K2_EndProfilingTimer(bool bLog, const FString& LogPrefix)
 {
-	if (sProfilingTimers.Num() == 0)
+	TArray<FK2ProfilingTimer>& ProfilingTimers = FProfilingTimerPerThread::Get().ProfilingTimers;
+	
+	if (ProfilingTimers.Num() == 0)
 	{
 		UE_LOG(LogAnimation, Warning, TEXT("Unbalanced use of Start & End Profiling Timer nodes."));
 		return 0.f;
 	}
 
-	FK2ProfilingTimer Timer = sProfilingTimers.Pop();
+	FK2ProfilingTimer Timer = ProfilingTimers.Pop();
 	double CurrentTimer = FPlatformTime::Seconds() * 1000.0;
 	Timer.AccummulatedTime = CurrentTimer - Timer.LastTime;
 
diff --git a/Engine/Source/Runtime/AnimGraphRuntime/Public/AnimNodes/AnimNode_CallFunction.h b/Engine/Source/Runtime/AnimGraphRuntime/Public/AnimNodes/AnimNode_CallFunction.h
new file mode 100644
--- /dev/null
+++ b/Engine/Source/Runtime/AnimGraphRuntime/Public/AnimNodes/AnimNode_CallFunction.h
@@ -0,0 +1,90 @@
+﻿// Copyright Epic Games, Inc. All Rights Reserved.
+
+#include "CoreMinimal.h"
+#include "Animation/AnimNodeBase.h"
+#include "AnimNode_CallFunction.generated.h"
+
+// When to call the function during the execution of the animation graph
+UENUM()
+enum class EAnimFunctionCallSite
+{
+	// Called when the node is initialized - i.e. it becomes weighted/relevant in the graph (before child nodes are initialized)
+	OnInitialize,
+	
+	// Called when the node is updated (before child nodes are updated)
+	OnUpdate,
+
+	// Called when the node is updated for the first time with a valid weight
+	OnBecomeRelevant,
+	
+	// Called when the node is evaluated (before child nodes are evaluated)
+    OnEvaluate,
+
+	// Called when the node is initialized - i.e. it becomes weighted/relevant in the graph (after child nodes are initialized)
+    OnInitializePostRecursion UMETA(DisplayName = "On Initialize (Post Recursion)"),
+	
+	// Called when the node is updated (after child nodes are updated)
+	OnUpdatePostRecursion UMETA(DisplayName = "On Update (Post Recursion)"),
+
+	// Called when the node is updated for the first time with a valid weight (after child nodes are updated)
+    OnBecomeRelevantPostRecursion UMETA(DisplayName = "On Become Relevant (Post Recursion)"),
+
+	// Called when the node is evaluated (after child nodes are evaluated)
+	OnEvaluatePostRecursion UMETA(DisplayName = "On Evaluate (Post Recursion)"),
+
+	// Called when the node is updated, was at full weight and beings to blend out. Called before child nodes are
+	// updated
+	OnStartedBlendingOut,
+
+	// Called when the node is updated, was at zero weight and beings to blend in. Called before child nodes are updated
+    OnStartedBlendingIn,
+
+	// Called when the node is updated, was at non-zero weight and finishes blending out. Called before child nodes are
+	// updated (note that this is necessarily not called within the graph update but from outside)
+	// @TODO: Not currently supported, needs subsystem support!
+    OnFinishedBlendingOut UMETA(Hidden),
+
+	// Called when the node is updated, was at non-zero weight and becomes full weight. Called before child nodes are
+    // updated
+    OnFinishedBlendingIn,
+};
+
+/** Calls specified user-defined events/functions during anim graph execution */
+USTRUCT(BlueprintInternalUseOnly)
+struct ANIMGRAPHRUNTIME_API FAnimNode_CallFunction : public FAnimNode_Base
+{
+	GENERATED_BODY()
+
+	UPROPERTY(EditAnywhere, Category = Links)
+	FPoseLink Source;
+
+	// The name of the function to call
+	UPROPERTY()
+	FName FunctionName = NAME_None;
+
+	// The function to call, recovered by looking for a function of name FunctionName
+	UPROPERTY()
+	UFunction* Function = nullptr;
+
+	// Counter used to determine relevancy
+	FGraphTraversalCounter Counter;
+
+	// Weight to determine blend-related call sites
+	float CurrentWeight = 0.0f;
+	
+	//  When to call the function during the execution of the animation graph
+	UPROPERTY(EditAnywhere, Category="Function")
+	EAnimFunctionCallSite CallSite = EAnimFunctionCallSite::OnUpdate;
+		
+	// FAnimNode_Base interface
+	virtual void OnInitializeAnimInstance(const FAnimInstanceProxy* InProxy, const UAnimInstance* InAnimInstance) override;
+	virtual bool NeedsOnInitializeAnimInstance() const override { return true; }
+	virtual void GatherDebugData(FNodeDebugData& DebugData) override;
+	virtual void Update_AnyThread(const FAnimationUpdateContext& Context) override;
+	virtual void Evaluate_AnyThread(FPoseContext& Context) override;
+	virtual void Initialize_AnyThread(const FAnimationInitializeContext& Context) override;
+	// End of FAnimNode_Base interface
+
+	// Calls the function we hold if the callsite matches the one we have set
+	void CallFunctionFromCallSite(EAnimFunctionCallSite InCallSite, const FAnimationBaseContext& InContext) const;
+};
\ No newline at end of file
diff --git a/Engine/Source/Runtime/AnimGraphRuntime/Public/KismetAnimationLibrary.h b/Engine/Source/Runtime/AnimGraphRuntime/Public/KismetAnimationLibrary.h
--- a/Engine/Source/Runtime/AnimGraphRuntime/Public/KismetAnimationLibrary.h
+++ b/Engine/Source/Runtime/AnimGraphRuntime/Public/KismetAnimationLibrary.h
@@ -14,7 +14,7 @@ class USkeletalMeshComponent;
 /**
  *	A library of the most common animation blueprint functions.
  */
-UCLASS(meta=(ScriptName="AnimGraphLibrary", DocumentationPolicy="Strict"))
+UCLASS(meta=(ScriptName="AnimGraphLibrary", BlueprintThreadSafe, DocumentationPolicy="Strict"))
 class ANIMGRAPHRUNTIME_API UKismetAnimationLibrary : public UBlueprintFunctionLibrary
 {
 	GENERATED_UCLASS_BODY()
@@ -81,7 +81,7 @@ class ANIMGRAPHRUNTIME_API UKismetAnimationLibrary : public UBlueprintFunctionLi
 	 *	@param OutRangeMin The minimum for the output range (commonly == 0.0)
 	 *	@param OutRangeMax The maximum for the output range (commonly == 1.0)
 	 */
-	UFUNCTION(BlueprintPure, Category = "Utilities|Animation", meta = (DisplayName = "Get Distance Between Two Sockets", ScriptName = "DistanceBetweenSockets", bRemapRange= "false"))
+	UFUNCTION(BlueprintPure, Category = "Utilities|Animation", meta = (NotBlueprintThreadSafe, DisplayName = "Get Distance Between Two Sockets", ScriptName = "DistanceBetweenSockets", bRemapRange= "false"))
 	static float K2_DistanceBetweenTwoSocketsAndMapRange(
 		const USkeletalMeshComponent* Component, 
 		const FName SocketOrBoneNameA, 
@@ -101,7 +101,7 @@ class ANIMGRAPHRUNTIME_API UKismetAnimationLibrary : public UBlueprintFunctionLi
 	 *	@param SocketOrBoneNameFrom The name of the first socket / bone
 	 *	@param SocketOrBoneNameTo The name of the second socket / bone
 	 */
-	UFUNCTION(BlueprintPure, Category = "Utilities|Animation", meta = (DisplayName = "Get Direction Between Sockets", ScriptName = "DirectionBetweenSockets"))
+	UFUNCTION(BlueprintPure, Category = "Utilities|Animation", meta = (NotBlueprintThreadSafe, DisplayName = "Get Direction Between Sockets", ScriptName = "DirectionBetweenSockets"))
 	static FVector K2_DirectionBetweenSockets(const USkeletalMeshComponent* Component, const FName SocketOrBoneNameFrom, const FName SocketOrBoneNameTo);
 
 	/** 
@@ -183,7 +183,7 @@ class ANIMGRAPHRUNTIME_API UKismetAnimationLibrary : public UBlueprintFunctionLi
 	 *	@param EasingType The easing function to use
 	 *	@param CustomCurve The curve to use if the easing type is "Custom"
 	 */
-	UFUNCTION(BlueprintCallable, Category = "Utilities|Animation", meta = (DisplayName = "Calculate Velocity From Sockets", ScriptName = "CalculateVelocityFromSockets", NumberOfSamples = "16", VelocityMin = "0.f", VelocityMax = "128.f"))
+	UFUNCTION(BlueprintCallable, Category = "Utilities|Animation", meta = (NotBlueprintThreadSafe, DisplayName = "Calculate Velocity From Sockets", ScriptName = "CalculateVelocityFromSockets", NumberOfSamples = "16", VelocityMin = "0.f", VelocityMax = "128.f"))
 	static float K2_CalculateVelocityFromSockets(
 		float DeltaSeconds,
 		USkeletalMeshComponent * Component,
@@ -221,16 +221,7 @@ class ANIMGRAPHRUNTIME_API UKismetAnimationLibrary : public UBlueprintFunctionLi
 	 * @param	Velocity		The velocity to use as direction relative to BaseRotation
 	 * @param	BaseRotation	The base rotation, e.g. of a pawn
 	 */
-	UFUNCTION(BlueprintPure, Category="Utilities|Animation", meta = (BlueprintThreadSafe))
+	UFUNCTION(BlueprintPure, Category="Utilities|Animation")
 	static float CalculateDirection(const FVector& Velocity, const FRotator& BaseRotation);
-
-private:
-
-	struct FK2ProfilingTimer
-	{
-		double LastTime;
-		double AccummulatedTime;
-	};
-	static TArray<FK2ProfilingTimer> sProfilingTimers;
 };
 
diff --git a/Engine/Source/Runtime/Engine/Classes/Animation/AnimInstance.h b/Engine/Source/Runtime/Engine/Classes/Animation/AnimInstance.h
--- a/Engine/Source/Runtime/Engine/Classes/Animation/AnimInstance.h
+++ b/Engine/Source/Runtime/Engine/Classes/Animation/AnimInstance.h
@@ -593,6 +593,10 @@ public:
 	UFUNCTION(BlueprintImplementableEvent)
 	void BlueprintLinkedAnimationLayersInitialized();
 
+	/** Executed when the Animation Blueprint is updated on a worker thread, just prior to graph update */
+	UFUNCTION(BlueprintImplementableEvent, meta=(BlueprintThreadSafe))
+	void BlueprintThreadSafeUpdateAnimation(float DeltaTime);
+	
 	bool CanTransitionSignature() const;
 	
 	/*********************************************************************************************
@@ -1336,6 +1340,10 @@ public:
 	/** Check whether we have any active curves */
 	bool HasActiveCurves() const;
 
+	/** Get the current delta time */
+	UFUNCTION(BlueprintPure, Category="Animation", meta=(BlueprintThreadSafe))
+	float GetDeltaSeconds() const;
+	
 	/** 
 	 * Append the type of curve to the OutCurveList specified by Curve Flags
 	 */
diff --git a/Engine/Source/Runtime/Engine/Classes/Kismet/KismetSystemLibrary.h b/Engine/Source/Runtime/Engine/Classes/Kismet/KismetSystemLibrary.h
--- a/Engine/Source/Runtime/Engine/Classes/Kismet/KismetSystemLibrary.h
+++ b/Engine/Source/Runtime/Engine/Classes/Kismet/KismetSystemLibrary.h
@@ -380,6 +380,16 @@ class ENGINE_API UKismetSystemLibrary : public UBlueprintFunctionLibrary
 	UFUNCTION(BlueprintPure, Category="Utilities|Text", meta=(BlueprintThreadSafe))
 	static FText MakeLiteralText(FText Value);
 
+	/**
+	 * Prints a string to the log
+	 * If Print To Log is true, it will be visible in the Output Log window.  Otherwise it will be logged only as 'Verbose', so it generally won't show up.
+	 *
+	 * @param	InString		The string to log out
+	 * @param	bPrintToLog		Whether or not to print the output to the log
+	 */
+	UFUNCTION(BlueprintCallable, Category="Utilities|String", meta=(BlueprintThreadSafe, Keywords = "log print", DevelopmentOnly))
+	static void LogString(const FString& InString = FString(TEXT("Hello")), bool bPrintToLog = true);
+	
 	/**
 	 * Prints a string to the log, and optionally, to the screen
 	 * If Print To Log is true, it will be visible in the Output Log window.  Otherwise it will be logged only as 'Verbose', so it generally won't show up.
diff --git a/Engine/Source/Runtime/Engine/Private/Animation/AnimBlueprintGeneratedClass.cpp b/Engine/Source/Runtime/Engine/Private/Animation/AnimBlueprintGeneratedClass.cpp
--- a/Engine/Source/Runtime/Engine/Private/Animation/AnimBlueprintGeneratedClass.cpp
+++ b/Engine/Source/Runtime/Engine/Private/Animation/AnimBlueprintGeneratedClass.cpp
@@ -471,6 +471,7 @@ void UAnimBlueprintGeneratedClass::Link(FArchive& Ar, bool bRelinkExistingProper
 	ConstantProperties.Empty();
 	MutableSubsystemProperties.Empty();
 	ConstantSubsystemProperties.Empty();
+	MutableNodeDataProperty = nullptr;
 
 #if WITH_EDITOR
 	// This relies on the entire class being fully loaded, this is not the case with EDL async-loading, in which case the functions are generated in PostLoad
diff --git a/Engine/Source/Runtime/Engine/Private/Animation/AnimInstance.cpp b/Engine/Source/Runtime/Engine/Private/Animation/AnimInstance.cpp
--- a/Engine/Source/Runtime/Engine/Private/Animation/AnimInstance.cpp
+++ b/Engine/Source/Runtime/Engine/Private/Animation/AnimInstance.cpp
@@ -494,6 +494,16 @@ void UAnimInstance::UpdateAnimation(float DeltaSeconds, bool bNeedsValidRootMoti
 		UpdateMontageEvaluationData();
 	}
 
+	if (IAnimClassInterface* AnimBlueprintClass = IAnimClassInterface::GetFromClass(GetClass()))
+	{
+		AnimBlueprintClass->ForEachSubsystem(this, [this, DeltaSeconds](const FAnimSubsystemInstanceContext& InContext)
+		{
+			FAnimSubsystemUpdateContext Context(InContext, this, DeltaSeconds);
+			InContext.Subsystem.OnPreUpdate_GameThread(Context);
+			return EAnimSubsystemEnumeration::Continue;
+		});
+	}
+	
 	{
 		SCOPE_CYCLE_COUNTER(STAT_NativeUpdateAnimation);
 		CSV_SCOPED_TIMING_STAT(Animation, NativeUpdate);
@@ -515,7 +525,7 @@ PRAGMA_ENABLE_DEPRECATION_WARNINGS
 		AnimBlueprintClass->ForEachSubsystem(this, [this, DeltaSeconds](const FAnimSubsystemInstanceContext& InContext)
 		{
 			FAnimSubsystemUpdateContext Context(InContext, this, DeltaSeconds);
-			InContext.Subsystem.OnUpdate(Context);
+			InContext.Subsystem.OnPostUpdate_GameThread(Context);
 			return EAnimSubsystemEnumeration::Continue;
 		});
 	}
@@ -1354,6 +1364,11 @@ bool UAnimInstance::HasActiveCurves() const
 	return GetProxyOnAnyThread<FAnimInstanceProxy>().HasActiveCurves();
 }
 
+float UAnimInstance::GetDeltaSeconds() const
+{
+	return GetProxyOnAnyThread<FAnimInstanceProxy>().GetDeltaSeconds();
+}
+
 void UAnimInstance::TriggerAnimNotifies(float DeltaSeconds)
 {
 	SCOPE_CYCLE_COUNTER(STAT_AnimTriggerAnimNotifies);
diff --git a/Engine/Source/Runtime/Engine/Private/Animation/AnimInstanceProxy.cpp b/Engine/Source/Runtime/Engine/Private/Animation/AnimInstanceProxy.cpp
--- a/Engine/Source/Runtime/Engine/Private/Animation/AnimInstanceProxy.cpp
+++ b/Engine/Source/Runtime/Engine/Private/Animation/AnimInstanceProxy.cpp
@@ -1035,16 +1035,6 @@ void FAnimInstanceProxy::UpdateAnimation()
 
 	FMemMark Mark(FMemStack::Get());
 
-	if(AnimClassInterface)
-	{
-		AnimClassInterface->ForEachSubsystem(GetAnimInstanceObject(), [this](const FAnimSubsystemInstanceContext& InContext)
-		{
-		    FAnimSubsystemParallelUpdateContext Context(InContext, *this, CurrentDeltaSeconds);
-		    InContext.Subsystem.OnParallelUpdate(Context);
-		    return EAnimSubsystemEnumeration::Continue;
-		});
-	}
-
 #if WITH_EDITORONLY_DATA
 	UpdatedNodesThisFrame.Reset();
 	NodeInputAttributesThisFrame.Reset();
@@ -1115,9 +1105,36 @@ void FAnimInstanceProxy::UpdateAnimation_WithRoot(const FAnimationUpdateContext&
 		// from other linked instances with grouped layers
 		if(FrameCounterForUpdate != GFrameCounter)
 		{
-			SCOPE_CYCLE_COUNTER(STAT_NativeUpdateAnimation);
-			Update(CurrentDeltaSeconds);
+			if(AnimClassInterface)
+			{
+				AnimClassInterface->ForEachSubsystem(GetAnimInstanceObject(), [this](const FAnimSubsystemInstanceContext& InContext)
+				{
+					FAnimSubsystemParallelUpdateContext Context(InContext, *this, CurrentDeltaSeconds);
+					InContext.Subsystem.OnPreUpdate_WorkerThread(Context);
+					return EAnimSubsystemEnumeration::Continue;
+				});
+			}
+			
+			{
+				SCOPE_CYCLE_COUNTER(STAT_BlueprintUpdateAnimation);
+				CastChecked<UAnimInstance>(GetAnimInstanceObject())->BlueprintThreadSafeUpdateAnimation(CurrentDeltaSeconds);
+			}
+
+			{
+				SCOPE_CYCLE_COUNTER(STAT_NativeUpdateAnimation);
+				Update(CurrentDeltaSeconds);
+			}
 
+			if(AnimClassInterface)
+			{
+				AnimClassInterface->ForEachSubsystem(GetAnimInstanceObject(), [this](const FAnimSubsystemInstanceContext& InContext)
+				{
+					FAnimSubsystemParallelUpdateContext Context(InContext, *this, CurrentDeltaSeconds);
+					InContext.Subsystem.OnPostUpdate_WorkerThread(Context);
+					return EAnimSubsystemEnumeration::Continue;
+				});
+			}
+			
 			FrameCounterForUpdate = GFrameCounter;
 		}
 	}
diff --git a/Engine/Source/Runtime/Engine/Private/Animation/AnimSubsystem_PropertyAccess.cpp b/Engine/Source/Runtime/Engine/Private/Animation/AnimSubsystem_PropertyAccess.cpp
--- a/Engine/Source/Runtime/Engine/Private/Animation/AnimSubsystem_PropertyAccess.cpp
+++ b/Engine/Source/Runtime/Engine/Private/Animation/AnimSubsystem_PropertyAccess.cpp
@@ -4,16 +4,28 @@
 #include "Animation/AnimInstance.h"
 #include "Animation/AnimInstanceProxy.h"
 
-void FAnimSubsystem_PropertyAccess::OnUpdate(FAnimSubsystemUpdateContext& InContext) const
+void FAnimSubsystem_PropertyAccess::OnPreUpdate_GameThread(FAnimSubsystemUpdateContext& InContext) const
 {
 	// Process internal batched property copies
-	PropertyAccess::ProcessCopies(InContext.AnimInstance, Library, EPropertyAccessCopyBatch::ExternalBatched);
+	PropertyAccess::ProcessCopies(InContext.AnimInstance, Library, PropertyAccess::FCopyBatchId((int32)EAnimPropertyAccessCallSite::GameThread_Batched_PreEventGraph));
 }
 
-void FAnimSubsystem_PropertyAccess::OnParallelUpdate(FAnimSubsystemParallelUpdateContext& InContext) const
+void FAnimSubsystem_PropertyAccess::OnPostUpdate_GameThread(FAnimSubsystemUpdateContext& InContext) const
 {
 	// Process internal batched property copies
-	PropertyAccess::ProcessCopies(InContext.Proxy.GetAnimInstanceObject(), Library, EPropertyAccessCopyBatch::InternalBatched);
+	PropertyAccess::ProcessCopies(InContext.AnimInstance, Library, PropertyAccess::FCopyBatchId((int32)EAnimPropertyAccessCallSite::GameThread_Batched_PostEventGraph));
+}
+
+void FAnimSubsystem_PropertyAccess::OnPreUpdate_WorkerThread(FAnimSubsystemParallelUpdateContext& InContext) const
+{
+	// Process internal batched property copies
+	PropertyAccess::ProcessCopies(InContext.Proxy.GetAnimInstanceObject(), Library, PropertyAccess::FCopyBatchId((int32)EAnimPropertyAccessCallSite::WorkerThread_Batched_PreEventGraph));
+}
+
+void FAnimSubsystem_PropertyAccess::OnPostUpdate_WorkerThread(FAnimSubsystemParallelUpdateContext& InContext) const
+{
+	// Process internal batched property copies
+	PropertyAccess::ProcessCopies(InContext.Proxy.GetAnimInstanceObject(), Library, PropertyAccess::FCopyBatchId((int32)EAnimPropertyAccessCallSite::WorkerThread_Batched_PostEventGraph));
 }
 
 void FAnimSubsystem_PropertyAccess::OnPostLoad(FAnimSubsystemPostLoadContext& InContext)
diff --git a/Engine/Source/Runtime/Engine/Private/Animation/ExposedValueHandler.cpp b/Engine/Source/Runtime/Engine/Private/Animation/ExposedValueHandler.cpp
--- a/Engine/Source/Runtime/Engine/Private/Animation/ExposedValueHandler.cpp
+++ b/Engine/Source/Runtime/Engine/Private/Animation/ExposedValueHandler.cpp
@@ -58,9 +58,10 @@ void FExposedValueHandler::Execute(const FAnimationBaseContext& Context) const
 		if(PropertyAccessLibrary != nullptr)
 		{
 			UObject* AnimInstanceObject = Context.AnimInstanceProxy->GetAnimInstanceObject();
+			PropertyAccess::FCopyBatchId CopyBatch((int32)EAnimPropertyAccessCallSite::WorkerThread_Unbatched);
 			for(const FExposedValueCopyRecord& CopyRecord : CopyRecords)
 			{
-				PropertyAccess::ProcessCopy(AnimInstanceObject, *PropertyAccessLibrary, EPropertyAccessCopyBatch::InternalUnbatched, CopyRecord.CopyIndex, [&CopyRecord](const FProperty* InProperty, void* InAddress)
+				PropertyAccess::ProcessCopy(AnimInstanceObject, *PropertyAccessLibrary, CopyBatch, CopyRecord.CopyIndex, [&CopyRecord](const FProperty* InProperty, void* InAddress)
 				{
 					if(CopyRecord.PostCopyOperation == EPostCopyOperation::LogicalNegateBool)
 					{
diff --git a/Engine/Source/Runtime/Engine/Private/KismetSystemLibrary.cpp b/Engine/Source/Runtime/Engine/Private/KismetSystemLibrary.cpp
--- a/Engine/Source/Runtime/Engine/Private/KismetSystemLibrary.cpp
+++ b/Engine/Source/Runtime/Engine/Private/KismetSystemLibrary.cpp
@@ -251,6 +251,20 @@ UObject* UKismetSystemLibrary::Conv_InterfaceToObject(const FScriptInterface& In
 	return Interface.GetObject();
 }
 
+void UKismetSystemLibrary::LogString(const FString& InString, bool bPrintToLog)
+{
+#if !(UE_BUILD_SHIPPING || UE_BUILD_TEST) // Do not Print in Shipping or Test
+	if(bPrintToLog)
+	{
+		UE_LOG(LogBlueprintUserMessages, Log, TEXT("%s"), *InString);
+	}
+	else
+	{
+		UE_LOG(LogBlueprintUserMessages, Verbose, TEXT("%s"), *InString);
+	}	
+#endif
+}
+
 void UKismetSystemLibrary::PrintString(const UObject* WorldContextObject, const FString& InString, bool bPrintToScreen, bool bPrintToLog, FLinearColor TextColor, float Duration)
 {
 #if !(UE_BUILD_SHIPPING || UE_BUILD_TEST) // Do not Print in Shipping or Test
diff --git a/Engine/Source/Runtime/Engine/Private/PropertyAccess.cpp b/Engine/Source/Runtime/Engine/Private/PropertyAccess.cpp
--- a/Engine/Source/Runtime/Engine/Private/PropertyAccess.cpp
+++ b/Engine/Source/Runtime/Engine/Private/PropertyAccess.cpp
@@ -563,90 +563,53 @@ struct FPropertyAccessSystem
 	}
 
 	// Process a single copy
-	static void ProcessCopy(UStruct* InStruct, void* InContainer, const FPropertyAccessLibrary& InLibrary, int32 InCopyIndex, EPropertyAccessCopyBatch InBatchType, TFunctionRef<void(const FProperty*, void*)> InPostCopyOperation)
+	static void ProcessCopy(UStruct* InStruct, void* InContainer, const FPropertyAccessLibrary& InLibrary, int32 InCopyIndex, int32 InBatchId, TFunctionRef<void(const FProperty*, void*)> InPostCopyOperation)
 	{
 		if(InLibrary.bHasBeenPostLoaded)
 		{
-			const FPropertyAccessCopy& Copy = InLibrary.CopyBatches[(__underlying_type(EPropertyAccessCopyBatch))InBatchType].Copies[InCopyIndex];
-			const FPropertyAccessIndirectionChain& SrcAccess = InLibrary.SrcAccesses[Copy.AccessIndex];
-			if(SrcAccess.Property.Get())
+			if(InLibrary.CopyBatchArray.IsValidIndex(InBatchId))
 			{
-				GetAccessAddress(InContainer, InLibrary, SrcAccess, [InContainer, &InLibrary, &Copy, &SrcAccess, &InPostCopyOperation](void* InSrcAddress)
+				const FPropertyAccessCopy& Copy = InLibrary.CopyBatchArray[InBatchId].Copies[InCopyIndex];
+				const FPropertyAccessIndirectionChain& SrcAccess = InLibrary.SrcAccesses[Copy.AccessIndex];
+				if(SrcAccess.Property.Get())
 				{
-					for(int32 DestAccessIndex = Copy.DestAccessStartIndex; DestAccessIndex < Copy.DestAccessEndIndex; ++DestAccessIndex)
+					GetAccessAddress(InContainer, InLibrary, SrcAccess, [InContainer, &InLibrary, &Copy, &SrcAccess, &InPostCopyOperation](void* InSrcAddress)
 					{
-						const FPropertyAccessIndirectionChain& DestAccess = InLibrary.DestAccesses[DestAccessIndex];
-						if(DestAccess.Property.Get())
+						for(int32 DestAccessIndex = Copy.DestAccessStartIndex; DestAccessIndex < Copy.DestAccessEndIndex; ++DestAccessIndex)
 						{
-							GetAccessAddress(InContainer, InLibrary, DestAccess, [&InSrcAddress, &Copy, &SrcAccess, &DestAccess, &InPostCopyOperation](void* InDestAddress)
+							const FPropertyAccessIndirectionChain& DestAccess = InLibrary.DestAccesses[DestAccessIndex];
+							if(DestAccess.Property.Get())
 							{
-								PerformCopy(Copy, SrcAccess.Property.Get(), InSrcAddress, DestAccess.Property.Get(), InDestAddress, InPostCopyOperation);
-							});
+								GetAccessAddress(InContainer, InLibrary, DestAccess, [&InSrcAddress, &Copy, &SrcAccess, &DestAccess, &InPostCopyOperation](void* InDestAddress)
+								{
+									PerformCopy(Copy, SrcAccess.Property.Get(), InSrcAddress, DestAccess.Property.Get(), InDestAddress, InPostCopyOperation);
+								});
+							}
 						}
-					}
-				});
+					});
+				}
 			}
 		}
 	}
 
-	static void ProcessCopies(UStruct* InStruct, void* InContainer, const FPropertyAccessLibrary& InLibrary, EPropertyAccessCopyBatch InBatchType)
+	static void ProcessCopies(UStruct* InStruct, void* InContainer, const FPropertyAccessLibrary& InLibrary, int32 InBatchId)
 	{
 		if(InLibrary.bHasBeenPostLoaded)
 		{
 			FMemMark Mark(FMemStack::Get());
 
-			// Copy all valid properties
-			// Parallelization opportunity: ParallelFor all the property copies we need to make
-			const int32 NumCopies = InLibrary.CopyBatches[(__underlying_type(EPropertyAccessCopyBatch))InBatchType].Copies.Num();
-			for(int32 CopyIndex = 0; CopyIndex < NumCopies; ++CopyIndex)
+			if(InLibrary.CopyBatchArray.IsValidIndex(InBatchId))
 			{
-				ProcessCopy(InStruct, InContainer, InLibrary, CopyIndex, InBatchType, [](const FProperty*, void*){});
-			}
-		}
-	}
-
-	/** A node in a class property tree */
-	struct FPropertyNode
-	{
-		FPropertyNode(TArray<FPropertyNode>& InPropertyTree, const FProperty* InProperty, FName InName, int32 InId, int32 InParentId)
-			: Property(InProperty)
-			, Name(InName)
-			, Id(InId)
-			, ParentId(InParentId)
-		{
-			if(InParentId != INDEX_NONE)
-			{
-				if(InPropertyTree[InParentId].FirstChildId == INDEX_NONE)
+				// Copy all valid properties
+				// Parallelization opportunity: ParallelFor all the property copies we need to make
+				const int32 NumCopies = InLibrary.CopyBatchArray[InBatchId].Copies.Num();
+				for(int32 CopyIndex = 0; CopyIndex < NumCopies; ++CopyIndex)
 				{
-					InPropertyTree[InParentId].FirstChildId = InId;
-					InPropertyTree[InParentId].LastChildId = InId;
+					ProcessCopy(InStruct, InContainer, InLibrary, CopyIndex, InBatchId, [](const FProperty*, void*){});
 				}
-				InPropertyTree[InParentId].LastChildId++;
 			}
 		}
-
-		FPropertyNode(TArray<FPropertyNode>& InPropertyTree, const FProperty* InProperty, int32 InId, int32 InParentId)
-			: FPropertyNode(InPropertyTree, InProperty, InProperty->GetFName(), InId, InParentId)
-		{
-		}
-
-		FPropertyNode(TArray<FPropertyNode>& InPropertyTree, int32 InId)
-			: FPropertyNode(InPropertyTree, nullptr, NAME_None, InId, INDEX_NONE)
-		{
-		}
-
-		const FProperty* Property = nullptr;
-
-		FName Name = NAME_None;
-
-		int32 Id = INDEX_NONE;
-
-		int32 ParentId = INDEX_NONE;
-
-		int32 FirstChildId = INDEX_NONE;
-
-		int32 LastChildId = INDEX_NONE;
-	};
+	}
 };
 
 namespace PropertyAccess
@@ -657,15 +620,25 @@ namespace PropertyAccess
 	}
 
 	void ProcessCopies(UObject* InObject, const FPropertyAccessLibrary& InLibrary, EPropertyAccessCopyBatch InBatchType)
+	{
+		ProcessCopies(InObject, InLibrary, FCopyBatchId((__underlying_type(EPropertyAccessCopyBatch))InBatchType));
+	}
+
+	void ProcessCopies(UObject* InObject, const FPropertyAccessLibrary& InLibrary, const FCopyBatchId& InBatchId)
 	{
 		QUICK_SCOPE_CYCLE_COUNTER(PropertyAccess_ProcessCopies);
 
-		::FPropertyAccessSystem::ProcessCopies(InObject->GetClass(), InObject, InLibrary, InBatchType);
+		::FPropertyAccessSystem::ProcessCopies(InObject->GetClass(), InObject, InLibrary, InBatchId.Id);
 	}
 
 	void ProcessCopy(UObject* InObject, const FPropertyAccessLibrary& InLibrary, EPropertyAccessCopyBatch InBatchType, int32 InCopyIndex, TFunctionRef<void(const FProperty*, void*)> InPostCopyOperation)
 	{
-		::FPropertyAccessSystem::ProcessCopy(InObject->GetClass(), InObject, InLibrary, InCopyIndex, InBatchType, InPostCopyOperation);
+		ProcessCopy(InObject, InLibrary, FCopyBatchId((__underlying_type(EPropertyAccessCopyBatch))InBatchType), InCopyIndex, InPostCopyOperation);
+	}
+
+	void ProcessCopy(UObject* InObject, const FPropertyAccessLibrary& InLibrary, const FCopyBatchId& InBatchId, int32 InCopyIndex, TFunctionRef<void(const FProperty*, void*)> InPostCopyOperation)
+	{
+		::FPropertyAccessSystem::ProcessCopy(InObject->GetClass(), InObject, InLibrary, InCopyIndex, InBatchId.Id, InPostCopyOperation);
 	}
 
 	void BindEvents(UObject* InObject, const FPropertyAccessLibrary& InLibrary)
diff --git a/Engine/Source/Runtime/Engine/Public/Animation/AnimSubsystem.h b/Engine/Source/Runtime/Engine/Public/Animation/AnimSubsystem.h
--- a/Engine/Source/Runtime/Engine/Public/Animation/AnimSubsystem.h
+++ b/Engine/Source/Runtime/Engine/Public/Animation/AnimSubsystem.h
@@ -92,12 +92,18 @@ struct FAnimSubsystem
 
 	/** Virtual destructor */
 	virtual ~FAnimSubsystem() = default;
-	
-	/** Override point to process game-thread data per-frame */
-	virtual void OnUpdate(FAnimSubsystemUpdateContext& InContext) const {}
 
-	/** Override point to process worker-thread data per-frame */
-	virtual void OnParallelUpdate(FAnimSubsystemParallelUpdateContext& InContext) const {}
+	/** Override point to process game-thread data per-frame. Called before event-graph-related work (e.g. NativeUpdateAnimation and BlueprintUpdateAnimation) */
+	virtual void OnPreUpdate_GameThread(FAnimSubsystemUpdateContext& InContext) const {}
+
+	/** Override point to process game-thread data per-frame. Called after event-graph-related work (e.g. NativeUpdateAnimation and BlueprintUpdateAnimation) */
+	virtual void OnPostUpdate_GameThread(FAnimSubsystemUpdateContext& InContext) const {}
+
+	/** Override point to process worker-thread data per-frame. Called before proxy Update and BlueprintThreadSafeUpdateAnimation. */
+	virtual void OnPreUpdate_WorkerThread(FAnimSubsystemParallelUpdateContext& InContext) const {}
+
+	/** Override point to process worker-thread data per-frame. Called after proxy Update and BlueprintThreadSafeUpdateAnimation. */
+	virtual void OnPostUpdate_WorkerThread(FAnimSubsystemParallelUpdateContext& InContext) const {}	
 
 	/** Override point to perform subsystem instance data initialization post-load */
 	virtual void OnPostLoadDefaults(FAnimSubsystemPostLoadDefaultsContext& InContext) {}
diff --git a/Engine/Source/Runtime/Engine/Public/Animation/AnimSubsystem_PropertyAccess.h b/Engine/Source/Runtime/Engine/Public/Animation/AnimSubsystem_PropertyAccess.h
--- a/Engine/Source/Runtime/Engine/Public/Animation/AnimSubsystem_PropertyAccess.h
+++ b/Engine/Source/Runtime/Engine/Public/Animation/AnimSubsystem_PropertyAccess.h
@@ -7,6 +7,26 @@
 #include "PropertyAccess.h"
 #include "AnimSubsystem_PropertyAccess.generated.h"
 
+// The various call sites we can call into the property access library
+UENUM()
+enum class EAnimPropertyAccessCallSite
+{
+	// Access is made on a worker thread in the anim graph or in a BP function
+	WorkerThread_Unbatched UMETA(DisplayName="Thread Safe"),
+
+	// Access is made on a worker thread before BlueprintThreadSafeUpdateAnimation is run
+	WorkerThread_Batched_PreEventGraph UMETA(DisplayName="Pre-Thread Safe Update Animation"),
+
+	// Access is made on a worker thread after BlueprintThreadSafeUpdateAnimation is run
+	WorkerThread_Batched_PostEventGraph UMETA(DisplayName="Post-Thread Safe Update Animation"),
+	
+	// Access is made on the game thread before the event graph (and NativeUpdateAnimation) is run
+	GameThread_Batched_PreEventGraph UMETA(DisplayName="Pre-Event Graph"),
+
+	// Access is made on the game thread after the event graph (and NativeUpdateAnimation) is run
+	GameThread_Batched_PostEventGraph UMETA(DisplayName="Post-Event Graph"),
+};
+
 USTRUCT()
 struct ENGINE_API FAnimSubsystem_PropertyAccess : public FAnimSubsystem
 {
@@ -15,8 +35,10 @@ struct ENGINE_API FAnimSubsystem_PropertyAccess : public FAnimSubsystem
 	friend class UAnimBlueprintExtension_PropertyAccess;
 
 	/** FAnimSubsystem interface */
-	virtual void OnUpdate(FAnimSubsystemUpdateContext& InContext) const override;
-	virtual void OnParallelUpdate(FAnimSubsystemParallelUpdateContext& InContext) const override;
+	virtual void OnPreUpdate_GameThread(FAnimSubsystemUpdateContext& InContext) const override;
+	virtual void OnPostUpdate_GameThread(FAnimSubsystemUpdateContext& InContext) const override;
+	virtual void OnPreUpdate_WorkerThread(FAnimSubsystemParallelUpdateContext& InContext) const override;
+	virtual void OnPostUpdate_WorkerThread(FAnimSubsystemParallelUpdateContext& InContext) const override;
 	virtual void OnPostLoad(FAnimSubsystemPostLoadContext& InContext) override;
 
 	/** Access the property access library */
diff --git a/Engine/Source/Runtime/Engine/Public/PropertyAccess.h b/Engine/Source/Runtime/Engine/Public/PropertyAccess.h
--- a/Engine/Source/Runtime/Engine/Public/PropertyAccess.h
+++ b/Engine/Source/Runtime/Engine/Public/PropertyAccess.h
@@ -34,6 +34,16 @@ enum class EPropertyAccessCopyBatch : uint8
 
 namespace PropertyAccess
 {
+	/** Batch ID - organisation/batching strategy determined by client systems via compilation */
+	struct FCopyBatchId
+	{
+		explicit FCopyBatchId(int32 InId)
+			: Id(InId)
+		{}
+		
+		int32 Id;
+	};
+
 	/** 
 	 * Called to patch up library after it is loaded.
 	 * This converts all FName-based paths into node-based paths that provide an optimized way of accessing properties.
@@ -44,12 +54,18 @@ namespace PropertyAccess
 	 * Process a 'tick' of a property access instance. 
 	 * Note internally allocates via FMemStack and pushes its own FMemMark
 	 */
+	ENGINE_API extern void ProcessCopies(UObject* InObject, const FPropertyAccessLibrary& InLibrary, const FCopyBatchId& InBatchId);
+
+	UE_DEPRECATED(5.0, "Please use the signature that takes a FCopyBatchId BatchId")
 	ENGINE_API extern void ProcessCopies(UObject* InObject, const FPropertyAccessLibrary& InLibrary, EPropertyAccessCopyBatch InBatchType);
 
 	/** 
 	 * Process a single copy 
 	 * Note that this can potentially allocate via FMemStack, so inserting FMemMark before a number of these calls is recommended
 	 */
+	ENGINE_API extern void ProcessCopy(UObject* InObject, const FPropertyAccessLibrary& InLibrary, const FCopyBatchId& InBatchId, int32 InCopyIndex, TFunctionRef<void(const FProperty*, void*)> InPostCopyOperation);
+
+	UE_DEPRECATED(5.0, "Please use the signature that takes a FCopyBatchId BatchId")
 	ENGINE_API extern void ProcessCopy(UObject* InObject, const FPropertyAccessLibrary& InLibrary, EPropertyAccessCopyBatch InBatchType, int32 InCopyIndex, TFunctionRef<void(const FProperty*, void*)> InPostCopyOperation);
 
 	UE_DEPRECATED(5.0, "Property Access Events are no longer supported")
@@ -391,10 +407,15 @@ private:
 	UPROPERTY()
 	TArray<FPropertyAccessPath> DestPaths;
 
-	// All copy operations
+#if WITH_EDITORONLY_DATA
 	UPROPERTY()
-	FPropertyAccessCopyBatch CopyBatches[(uint8)EPropertyAccessCopyBatch::Count];
+	FPropertyAccessCopyBatch CopyBatches_DEPRECATED[(uint8)EPropertyAccessCopyBatch::Count];
+#endif
 
+	// All copy operations
+	UPROPERTY()
+	TArray<FPropertyAccessCopyBatch> CopyBatchArray;
+	
 	// All source property accesses
 	UPROPERTY(Transient)
 	TArray<FPropertyAccessIndirectionChain> SrcAccesses;
