Fixed unloaded assets in asset players not showing up in anim BP context menus

Unified the approach to gathering assets and spawning nodes from them. Previously each had its own slightly unique code doing similar things.
Note that with this change it is now possible to directly spawn 'evaluator' style nodes from the context menu instead of relying on a context menu conversion step.
Also tweaked icons, tooltips and titles for consistency

diff --git a/Engine/Source/Editor/AnimGraph/Private/AnimGraphNode_AssetPlayerBase.cpp b/Engine/Source/Editor/AnimGraph/Private/AnimGraphNode_AssetPlayerBase.cpp
--- a/Engine/Source/Editor/AnimGraph/Private/AnimGraphNode_AssetPlayerBase.cpp
+++ b/Engine/Source/Editor/AnimGraph/Private/AnimGraphNode_AssetPlayerBase.cpp
@@ -3,25 +3,19 @@
 #include "AnimGraphNode_AssetPlayerBase.h"
 #include "EdGraphSchema_K2.h"
 #include "Animation/AnimComposite.h"
-#include "Animation/BlendSpace.h"
-#include "Animation/AimOffsetBlendSpace.h"
-#include "Animation/AimOffsetBlendSpace1D.h"
-#include "AnimGraphNode_SequencePlayer.h"
-#include "AnimGraphNode_SequenceEvaluator.h"
-#include "AnimGraphNode_RotationOffsetBlendSpace.h"
-#include "AnimGraphNode_BlendSpacePlayer.h"
-#include "AnimGraphNode_BlendSpaceEvaluator.h"
-#include "Animation/PoseAsset.h"
-#include "AnimGraphNode_PoseBlendNode.h"
-#include "AnimGraphNode_PoseByName.h"
-#include "AnimGraphNode_PoseDriver.h"
+#include "BlueprintActionFilter.h"
+#include "BlueprintNodeSpawner.h"
 #include "UObject/UObjectIterator.h"
 #include "Animation/AnimLayerInterface.h"
 #include "IAnimBlueprintGeneratedClassCompiledData.h"
 #include "IAnimBlueprintCompilationContext.h"
 #include "Animation/AnimSync.h"
 #include "Animation/AnimAttributes.h"
+#include "AssetRegistry/AssetRegistryModule.h"
 #include "UObject/UE5MainStreamObjectVersion.h"
+#include "BlueprintActionDatabaseRegistrar.h"
+#include "BlueprintNodeTemplateCache.h"
+#include "Animation/AnimNode_AssetPlayerBase.h"
 
 #define LOCTEXT_NAMESPACE "UAnimGraphNode_AssetPlayerBase"
 
@@ -99,6 +93,20 @@ void UAnimGraphNode_AssetPlayerBase::PinDefaultValueChanged(UEdGraphPin* Pin)
 	}
 }
 
+FText UAnimGraphNode_AssetPlayerBase::GetTooltipText() const
+{
+	bool const bIsTemplateNode = GetGraph() == nullptr || FBlueprintNodeTemplateCache::IsTemplateOuter(GetGraph());
+	if(bIsTemplateNode)
+	{
+		return FText::GetEmpty();
+	}
+	else
+	{
+		// FText::Format() is slow, so we utilize the cached list title
+		return GetNodeTitle(ENodeTitleType::ListView);
+	}
+}
+
 void UAnimGraphNode_AssetPlayerBase::OnProcessDuringCompilation(IAnimBlueprintCompilationContext& InCompilationContext, IAnimBlueprintGeneratedClassCompiledData& OutCompiledData)
 {
 	UBlueprint* Blueprint = GetBlueprint();
@@ -210,4 +218,103 @@ bool SupportNodeClassForAsset(const UClass* AssetClass, UClass* NodeClass)
 	return (NodeCDO->SupportsAssetClass(AssetClass) != EAnimAssetHandlerType::NotSupported);
 }
 
+void UAnimGraphNode_AssetPlayerBase::SetupNewNode(UEdGraphNode* InNewNode, bool bInIsTemplateNode, const FAssetData InAssetData)
+{
+	UAnimGraphNode_AssetPlayerBase* GraphNode = CastChecked<UAnimGraphNode_AssetPlayerBase>(InNewNode);
+	InAssetData.GetTagValue("Skeleton", GraphNode->UnloadedSkeletonName);
+	
+	if(!bInIsTemplateNode)
+	{
+		GraphNode->SetAnimationAsset(CastChecked<UAnimationAsset>(InAssetData.GetAsset()));
+	}
+}
+
+void UAnimGraphNode_AssetPlayerBase::GetMenuActionsHelper(FBlueprintActionDatabaseRegistrar& InActionRegistrar, TSubclassOf<UAnimGraphNode_Base> InNodeClass, const TArray<TSubclassOf<UObject>>& InAssetTypes, const TArray<TSubclassOf<UObject>>& InExcludedAssetTypes, const TFunctionRef<FText(const FAssetData&)>& InMenuNameFunction, const TFunctionRef<FText(const FAssetData&)>& InMenuTooltipFunction, const TFunction<void(UEdGraphNode*, bool, const FAssetData)>& InSetupNewNodeFunction)
+{
+	auto MakeAction = [&InActionRegistrar, &InMenuNameFunction, &InMenuTooltipFunction, InSetupNewNodeFunction, InNodeClass](const FAssetData& InAssetData)
+	{
+		auto AssetSetup = [InSetupNewNodeFunction, InAssetData](UEdGraphNode* InNewNode, bool bInIsTemplateNode)
+		{
+			InSetupNewNodeFunction(InNewNode, bInIsTemplateNode, InAssetData);
+		};
+
+		UBlueprintNodeSpawner* NodeSpawner = UBlueprintNodeSpawner::Create(InNodeClass.Get());
+		NodeSpawner->CustomizeNodeDelegate = UBlueprintNodeSpawner::FCustomizeNodeDelegate::CreateLambda(AssetSetup);
+		NodeSpawner->DefaultMenuSignature.MenuName = InMenuNameFunction(InAssetData);
+		NodeSpawner->DefaultMenuSignature.Tooltip = InMenuTooltipFunction(InAssetData);
+		InActionRegistrar.AddBlueprintAction(InAssetData, NodeSpawner);
+	};	
+
+	const UObject* QueryObject = InActionRegistrar.GetActionKeyFilter();
+	bool bIsObjectOfAssetType = false;
+	for(const TSubclassOf<UObject>& AssetType : InAssetTypes)
+	{
+		if(QueryObject && QueryObject->IsA(AssetType.Get()))
+		{
+			bIsObjectOfAssetType = true;
+			break;
+		}
+	}
+	
+	if (QueryObject == nullptr || QueryObject == InNodeClass.Get())
+	{
+		FAssetRegistryModule& AssetRegistryModule = FModuleManager::LoadModuleChecked<FAssetRegistryModule>(TEXT("AssetRegistry"));
+
+		FARFilter Filter;
+		for(const TSubclassOf<UObject>& AssetType : InAssetTypes)
+		{
+			Filter.ClassNames.Add(AssetType.Get()->GetFName());
+		}
+		for(const TSubclassOf<UObject>& ExcludedAssetType : InExcludedAssetTypes)
+		{
+			Filter.RecursiveClassesExclusionSet.Add(ExcludedAssetType.Get()->GetFName());
+		}	
+		Filter.bRecursiveClasses = true;
+		
+		TArray<FAssetData> AnimBlueprints;
+		AssetRegistryModule.Get().GetAssets(Filter, AnimBlueprints);
+
+		for (const FAssetData& AssetData : AnimBlueprints)
+		{
+			if(AssetData.IsUAsset())
+			{
+				MakeAction(AssetData);
+			}
+		}
+	}
+	else if (bIsObjectOfAssetType)
+	{
+		MakeAction(FAssetData(QueryObject));
+	}
+}
+
+bool UAnimGraphNode_AssetPlayerBase::IsActionFilteredOut(class FBlueprintActionFilter const& Filter)
+{
+	bool bIsFilteredOut = false;
+
+	if(!UnloadedSkeletonName.IsEmpty())
+	{
+		FBlueprintActionContext const& FilterContext = Filter.Context;
+
+		for (UBlueprint* Blueprint : FilterContext.Blueprints)
+		{
+			if (UAnimBlueprint* AnimBlueprint = Cast<UAnimBlueprint>(Blueprint))
+			{
+				if(!AnimBlueprint->TargetSkeleton->IsCompatibleSkeletonByAssetString(UnloadedSkeletonName))
+				{
+					bIsFilteredOut = true;
+					break;
+				}
+			}
+			else
+			{
+				// Not an animation Blueprint, cannot use
+				bIsFilteredOut = true;
+				break;
+			}
+		}
+	}
+	return bIsFilteredOut;
+}
+
 #undef LOCTEXT_NAMESPACE
\ No newline at end of file
diff --git a/Engine/Source/Editor/AnimGraph/Private/AnimGraphNode_BlendSpaceBase.cpp b/Engine/Source/Editor/AnimGraph/Private/AnimGraphNode_BlendSpaceBase.cpp
--- a/Engine/Source/Editor/AnimGraph/Private/AnimGraphNode_BlendSpaceBase.cpp
+++ b/Engine/Source/Editor/AnimGraph/Private/AnimGraphNode_BlendSpaceBase.cpp
@@ -24,6 +24,11 @@ FLinearColor UAnimGraphNode_BlendSpaceBase::GetNodeTitleColor() const
 	return FLinearColor(0.2f, 0.8f, 0.2f);
 }
 
+FSlateIcon UAnimGraphNode_BlendSpaceBase::GetIconAndTint(FLinearColor& OutColor) const
+{
+	return FSlateIcon("EditorStyle", "ClassIcon.BlendSpace");
+}
+
 void UAnimGraphNode_BlendSpaceBase::CustomizePinData(UEdGraphPin* Pin, FName SourcePropertyName, int32 ArrayIndex) const
 {
 	UBlendSpace * BlendSpace = GetBlendSpace();
diff --git a/Engine/Source/Editor/AnimGraph/Private/AnimGraphNode_BlendSpaceEvaluator.cpp b/Engine/Source/Editor/AnimGraph/Private/AnimGraphNode_BlendSpaceEvaluator.cpp
--- a/Engine/Source/Editor/AnimGraph/Private/AnimGraphNode_BlendSpaceEvaluator.cpp
+++ b/Engine/Source/Editor/AnimGraph/Private/AnimGraphNode_BlendSpaceEvaluator.cpp
@@ -3,8 +3,12 @@
 #include "AnimGraphNode_BlendSpaceEvaluator.h"
 #include "ToolMenus.h"
 #include "AnimGraphCommands.h"
+#include "BlueprintNodeSpawner.h"
 #include "Kismet2/CompilerResultsLog.h"
 #include "IAnimBlueprintNodeOverrideAssetsContext.h"
+#include "AssetRegistry/AssetRegistryModule.h"
+#include "BlueprintActionDatabaseRegistrar.h"
+#include "BlueprintNodeTemplateCache.h"
 
 /////////////////////////////////////////////////////
 // UAnimGraphNode_BlendSpaceEvaluator
@@ -16,12 +20,6 @@ UAnimGraphNode_BlendSpaceEvaluator::UAnimGraphNode_BlendSpaceEvaluator(const FOb
 {
 }
 
-FText UAnimGraphNode_BlendSpaceEvaluator::GetTooltipText() const
-{
-	// FText::Format() is slow, so we utilize the cached list title
-	return GetNodeTitle(ENodeTitleType::ListView);
-}
-
 FText UAnimGraphNode_BlendSpaceEvaluator::GetNodeTitleForBlendSpace(ENodeTitleType::Type TitleType, UBlendSpace* InBlendSpace) const
 {
 	const FText BlendSpaceName = FText::FromString(InBlendSpace->GetName());
@@ -88,10 +86,25 @@ FText UAnimGraphNode_BlendSpaceEvaluator::GetNodeTitle(ENodeTitleType::Type Titl
 	}
 }
 
-void UAnimGraphNode_BlendSpaceEvaluator::GetMenuActions(FBlueprintActionDatabaseRegistrar& ActionRegistrar) const
+void UAnimGraphNode_BlendSpaceEvaluator::GetMenuActions(FBlueprintActionDatabaseRegistrar& InActionRegistrar) const
 {
-	// Intentionally empty so that we don't get duplicate blend space entries.
-	// You can convert a regular blend space player to an evaluator via the right click context menu
+	GetMenuActionsHelper(
+		InActionRegistrar,
+		GetClass(),
+		{ UBlendSpace::StaticClass() },
+		{ },
+		[](const FAssetData& InAssetData)
+		{
+			return FText::Format(LOCTEXT("MenuDescFormat", "Blendspace Evaluator '{0}'"), FText::FromName(InAssetData.AssetName));
+		},
+		[](const FAssetData& InAssetData)
+		{
+			return FText::Format(LOCTEXT("MenuDescTooltipFormat", "Blendspace Evaluator\n'{0}'"), FText::FromName(InAssetData.ObjectPath));
+		},
+		[](UEdGraphNode* InNewNode, bool bInIsTemplateNode, const FAssetData InAssetData)
+		{
+			UAnimGraphNode_AssetPlayerBase::SetupNewNode(InNewNode, bInIsTemplateNode, InAssetData);
+		});	
 }
 
 void UAnimGraphNode_BlendSpaceEvaluator::ValidateAnimNodeDuringCompilation(class USkeleton* ForSkeleton, class FCompilerResultsLog& MessageLog)
diff --git a/Engine/Source/Editor/AnimGraph/Private/AnimGraphNode_BlendSpaceGraph.cpp b/Engine/Source/Editor/AnimGraph/Private/AnimGraphNode_BlendSpaceGraph.cpp
--- a/Engine/Source/Editor/AnimGraph/Private/AnimGraphNode_BlendSpaceGraph.cpp
+++ b/Engine/Source/Editor/AnimGraph/Private/AnimGraphNode_BlendSpaceGraph.cpp
@@ -1,10 +1,13 @@
 // Copyright Epic Games, Inc. All Rights Reserved.
 
 #include "AnimGraphNode_BlendSpaceGraph.h"
+
+#include "AnimGraphNode_AssetPlayerBase.h"
 #include "BlueprintNodeSpawner.h"
 #include "BlueprintActionDatabaseRegistrar.h"
 #include "Animation/AimOffsetBlendSpace.h"
 #include "Animation/AimOffsetBlendSpace1D.h"
+#include "AssetRegistry/AssetRegistryModule.h"
 
 #define LOCTEXT_NAMESPACE "UAnimGraphNode_BlendSpaceGraph"
 
@@ -51,82 +54,26 @@ FText UAnimGraphNode_BlendSpaceGraph::GetNodeTitle(ENodeTitleType::Type TitleTyp
 	}
 }
 
-void UAnimGraphNode_BlendSpaceGraph::GetMenuActions(FBlueprintActionDatabaseRegistrar& ActionRegistrar) const
+void UAnimGraphNode_BlendSpaceGraph::GetMenuActions(FBlueprintActionDatabaseRegistrar& InActionRegistrar) const
 {
-	struct GetMenuActions_Utils
-	{
-		static void SetNodeBlendSpace(UEdGraphNode* NewNode, bool bIsTemplateNode, TWeakObjectPtr<UBlendSpace> BlendSpace)
+	UAnimGraphNode_AssetPlayerBase::GetMenuActionsHelper(
+		InActionRegistrar,
+		GetClass(),
+		{ UBlendSpace::StaticClass(), UBlendSpace1D::StaticClass() },
+		{ UAimOffsetBlendSpace::StaticClass(), UAimOffsetBlendSpace1D::StaticClass() },
+		[](const FAssetData& InAssetData)
 		{
-			UAnimGraphNode_BlendSpaceGraph* BlendSpaceNode = CastChecked<UAnimGraphNode_BlendSpaceGraph>(NewNode);
-			BlendSpaceNode->SetupFromAsset(BlendSpace.Get(), bIsTemplateNode);
-		}
-
-		static UBlueprintNodeSpawner* MakeBlendSpaceAction(TSubclassOf<UEdGraphNode> const NodeClass, const UBlendSpace* InBlendSpace)
+			return FText::Format(LOCTEXT("MenuDescFormat", "Blendspace '{0}'"), FText::FromName(InAssetData.AssetName));
+		},
+		[](const FAssetData& InAssetData)
 		{
-			UBlueprintNodeSpawner* NodeSpawner = nullptr;
-
-			bool const bIsAimOffset = InBlendSpace->IsA(UAimOffsetBlendSpace::StaticClass()) ||
-									  InBlendSpace->IsA(UAimOffsetBlendSpace1D::StaticClass());
-			if (!bIsAimOffset)
-			{
-				NodeSpawner = UBlueprintNodeSpawner::Create(NodeClass);
-				check(NodeSpawner != nullptr);
-
-				TWeakObjectPtr<UBlendSpace> BlendSpacePtr = MakeWeakObjectPtr(const_cast<UBlendSpace*>(InBlendSpace));
-				NodeSpawner->CustomizeNodeDelegate = UBlueprintNodeSpawner::FCustomizeNodeDelegate::CreateStatic(GetMenuActions_Utils::SetNodeBlendSpace, BlendSpacePtr);
-			}	
-			return NodeSpawner;
-		}
-
-		static void SetNodeBlendSpaceClass(UEdGraphNode* NewNode, bool bIsTemplateNode, TSubclassOf<UBlendSpace> InBlendSpaceClass)
-		{
-			UAnimGraphNode_BlendSpaceGraph* BlendSpaceNode = CastChecked<UAnimGraphNode_BlendSpaceGraph>(NewNode);
-			BlendSpaceNode->SetupFromClass(InBlendSpaceClass, bIsTemplateNode);
-		}
-
-		static UBlueprintNodeSpawner* MakeBlendSpaceAction(TSubclassOf<UEdGraphNode> const NodeClass, TSubclassOf<UBlendSpace> InBlendSpaceClass)
+			return FText::Format(LOCTEXT("MenuDescTooltipFormat", "Blendspace\n'{0}'"), FText::FromName(InAssetData.ObjectPath));
+		},
+		[](UEdGraphNode* InNewNode, bool bInIsTemplateNode, const FAssetData InAssetData)
 		{
-			UBlueprintNodeSpawner* NodeSpawner = UBlueprintNodeSpawner::Create(NodeClass);
-			check(NodeSpawner != nullptr);
-
-			NodeSpawner->CustomizeNodeDelegate = UBlueprintNodeSpawner::FCustomizeNodeDelegate::CreateStatic(GetMenuActions_Utils::SetNodeBlendSpaceClass, InBlendSpaceClass);
-	
-			return NodeSpawner;
-		}
-	};
-
-	// Add the non-asset based cases
-	ActionRegistrar.AddBlueprintAction(GetMenuActions_Utils::MakeBlendSpaceAction(GetClass(), UBlendSpace::StaticClass()));
-	ActionRegistrar.AddBlueprintAction(GetMenuActions_Utils::MakeBlendSpaceAction(GetClass(), UBlendSpace1D::StaticClass()));
-
-	if (const UObject* RegistrarTarget = ActionRegistrar.GetActionKeyFilter())
-	{
-		if (const UBlendSpace* TargetBlendSpace = Cast<UBlendSpace>(RegistrarTarget))
-		{
-			if(TargetBlendSpace->IsAsset())
-			{
-				if (UBlueprintNodeSpawner* NodeSpawner = GetMenuActions_Utils::MakeBlendSpaceAction(GetClass(), TargetBlendSpace))
-				{
-					ActionRegistrar.AddBlueprintAction(TargetBlendSpace, NodeSpawner);
-				}
-			}
-		}
-		// else, the Blueprint database is specifically looking for actions pertaining to something different (not a BlendSpace asset)
-	}
-	else
-	{
-		UClass* NodeClass = GetClass();
-		for (TObjectIterator<UBlendSpace> BlendSpaceIt; BlendSpaceIt; ++BlendSpaceIt)
-		{
-			if(BlendSpaceIt->IsAsset())
-			{
-				if (UBlueprintNodeSpawner* NodeSpawner = GetMenuActions_Utils::MakeBlendSpaceAction(NodeClass, *BlendSpaceIt))
-				{
-					ActionRegistrar.AddBlueprintAction(*BlendSpaceIt, NodeSpawner);
-				}
-			}
-		}
-	}
+			UAnimGraphNode_BlendSpaceGraph* GraphNode = CastChecked<UAnimGraphNode_BlendSpaceGraph>(InNewNode);
+			GraphNode->SetupFromAsset(InAssetData, bInIsTemplateNode);
+		});
 }
 
 void UAnimGraphNode_BlendSpaceGraph::BakeDataDuringCompilation(class FCompilerResultsLog& MessageLog)
diff --git a/Engine/Source/Editor/AnimGraph/Private/AnimGraphNode_BlendSpaceGraphBase.cpp b/Engine/Source/Editor/AnimGraph/Private/AnimGraphNode_BlendSpaceGraphBase.cpp
--- a/Engine/Source/Editor/AnimGraph/Private/AnimGraphNode_BlendSpaceGraphBase.cpp
+++ b/Engine/Source/Editor/AnimGraph/Private/AnimGraphNode_BlendSpaceGraphBase.cpp
@@ -28,6 +28,7 @@
 #include "BlueprintEditor.h"
 #include "Animation/AnimSequence.h"
 #include "AnimBlueprintExtension_BlendSpaceGraph.h"
+#include "BlueprintNodeTemplateCache.h"
 
 #define LOCTEXT_NAMESPACE "UAnimGraphNode_BlendSpaceGraphBase"
 
@@ -46,9 +47,22 @@ FLinearColor UAnimGraphNode_BlendSpaceGraphBase::GetNodeTitleColor() const
 	return FLinearColor(0.8f, 0.8f, 0.8f);
 }
 
+FSlateIcon UAnimGraphNode_BlendSpaceGraphBase::GetIconAndTint(FLinearColor& OutColor) const
+{
+	return FSlateIcon("EditorStyle", "ClassIcon.BlendSpace");
+}
+
 FText UAnimGraphNode_BlendSpaceGraphBase::GetTooltipText() const
 {
-	return GetNodeTitle(ENodeTitleType::ListView);
+	bool const bIsTemplateNode = GetGraph() == nullptr || FBlueprintNodeTemplateCache::IsTemplateOuter(GetGraph());
+	if(bIsTemplateNode)
+	{
+		return FText::GetEmpty();
+	}
+	else
+	{
+		return GetNodeTitle(ENodeTitleType::ListView);
+	}
 }
 
 UAnimGraphNode_Base* UAnimGraphNode_BlendSpaceGraphBase::ExpandGraphAndProcessNodes(UEdGraph* SourceGraph, UAnimGraphNode_Base* SourceRootNode, IAnimBlueprintCompilationContext& InCompilationContext, IAnimBlueprintGeneratedClassCompiledData& OutCompiledData)
@@ -129,19 +143,19 @@ void UAnimGraphNode_BlendSpaceGraphBase::OnCopyTermDefaultsToDefaultObject(IAnim
 	DestinationNode->BlendSpace = Extension->AddBlendSpace(BlendSpace);
 }
 
-void UAnimGraphNode_BlendSpaceGraphBase::SetupFromAsset(UBlendSpace* InBlendSpace, bool bInIsTemplateNode)
+void UAnimGraphNode_BlendSpaceGraphBase::SetupFromAsset(const FAssetData& InAssetData, bool bInIsTemplateNode)
 {
-	if(bInIsTemplateNode)
-	{
-		BlendSpace = InBlendSpace;
-	}
-	else
+	InAssetData.GetTagValue("Skeleton", SkeletonName);
+	
+	if(!bInIsTemplateNode)
 	{
+		UBlendSpace* Asset = CastChecked<UBlendSpace>(InAssetData.GetAsset());
+		
 		UAnimBlueprint* AnimBlueprint = GetAnimBlueprint();
 		const UAnimationGraphSchema* AnimationGraphSchema = GetDefault<UAnimationGraphSchema>();
 
-		BlendSpaceGraph = CastChecked<UBlendSpaceGraph>(FBlueprintEditorUtils::CreateNewGraph(this, InBlendSpace->GetFName(), UBlendSpaceGraph::StaticClass(), UEdGraphSchema::StaticClass()));
-		BlendSpaceGraph->BlendSpace = BlendSpace = DuplicateObject(InBlendSpace, BlendSpaceGraph);
+		BlendSpaceGraph = CastChecked<UBlendSpaceGraph>(FBlueprintEditorUtils::CreateNewGraph(this, Asset->GetFName(), UBlendSpaceGraph::StaticClass(), UEdGraphSchema::StaticClass()));
+		BlendSpaceGraph->BlendSpace = BlendSpace = DuplicateObject(Asset, BlendSpaceGraph);
 		BlendSpaceGraph->BlendSpace->ClearFlags(RF_Public | RF_Standalone | RF_Transient);
 		BlendSpaceGraph->BlendSpace->SetFlags(RF_Transactional);
 		BlendSpaceGraph->BlendSpace->ResetSkeleton(nullptr);
@@ -515,4 +529,33 @@ void UAnimGraphNode_BlendSpaceGraphBase::GetRequiredExtensions(TArray<TSubclassO
 	OutExtensions.Add(UAnimBlueprintExtension_BlendSpaceGraph::StaticClass());
 }
 
+bool UAnimGraphNode_BlendSpaceGraphBase::IsActionFilteredOut(class FBlueprintActionFilter const& Filter)
+{
+	bool bIsFilteredOut = false;
+
+	if(!SkeletonName.IsEmpty())
+	{
+		FBlueprintActionContext const& FilterContext = Filter.Context;
+
+		for (UBlueprint* Blueprint : FilterContext.Blueprints)
+		{
+			if (UAnimBlueprint* AnimBlueprint = Cast<UAnimBlueprint>(Blueprint))
+			{
+				if(!AnimBlueprint->TargetSkeleton->IsCompatibleSkeletonByAssetString(SkeletonName))
+				{
+					bIsFilteredOut = true;
+					break;
+				}
+			}
+			else
+			{
+				// Not an animation Blueprint, cannot use
+				bIsFilteredOut = true;
+				break;
+			}
+		}
+	}
+	return bIsFilteredOut;
+}
+
 #undef LOCTEXT_NAMESPACE
diff --git a/Engine/Source/Editor/AnimGraph/Private/AnimGraphNode_BlendSpacePlayer.cpp b/Engine/Source/Editor/AnimGraph/Private/AnimGraphNode_BlendSpacePlayer.cpp
--- a/Engine/Source/Editor/AnimGraph/Private/AnimGraphNode_BlendSpacePlayer.cpp
+++ b/Engine/Source/Editor/AnimGraph/Private/AnimGraphNode_BlendSpacePlayer.cpp
@@ -13,6 +13,7 @@
 #include "AnimGraphCommands.h"
 #include "IAnimBlueprintCopyTermDefaultsContext.h"
 #include "IAnimBlueprintNodeOverrideAssetsContext.h"
+#include "AssetRegistry/AssetRegistryModule.h"
 
 /////////////////////////////////////////////////////
 // UAnimGraphNode_BlendSpacePlayer
@@ -24,12 +25,6 @@ UAnimGraphNode_BlendSpacePlayer::UAnimGraphNode_BlendSpacePlayer(const FObjectIn
 {
 }
 
-FText UAnimGraphNode_BlendSpacePlayer::GetTooltipText() const
-{
-	// FText::Format() is slow, so we utilize the cached list title
-	return GetNodeTitle(ENodeTitleType::ListView);
-}
-
 FText UAnimGraphNode_BlendSpacePlayer::GetNodeTitleForBlendSpace(ENodeTitleType::Type TitleType, UBlendSpace* InBlendSpace) const
 {
 	const FText BlendSpaceName = FText::FromString(InBlendSpace->GetName());
@@ -178,63 +173,25 @@ void UAnimGraphNode_BlendSpacePlayer::GetNodeContextMenuActions(UToolMenu* Menu,
 	}
 }
 
-void UAnimGraphNode_BlendSpacePlayer::GetMenuActions(FBlueprintActionDatabaseRegistrar& ActionRegistrar) const
+void UAnimGraphNode_BlendSpacePlayer::GetMenuActions(FBlueprintActionDatabaseRegistrar& InActionRegistrar) const
 {
-	struct GetMenuActions_Utils
-	{
-		static void SetNodeBlendSpace(UEdGraphNode* NewNode, bool /*bIsTemplateNode*/, TWeakObjectPtr<UBlendSpace> BlendSpace)
+	GetMenuActionsHelper(
+		InActionRegistrar,
+		GetClass(),
+		{ UBlendSpace::StaticClass() },
+		{ },
+		[](const FAssetData& InAssetData)
 		{
-			UAnimGraphNode_BlendSpacePlayer* BlendSpaceNode = CastChecked<UAnimGraphNode_BlendSpacePlayer>(NewNode);
-			BlendSpaceNode->Node.SetBlendSpace(BlendSpace.Get());
-		}
-
-		static UBlueprintNodeSpawner* MakeBlendSpaceAction(TSubclassOf<UEdGraphNode> const NodeClass, const UBlendSpace* BlendSpace)
+			return FText::Format(LOCTEXT("MenuDescFormat", "Blendspace Player '{0}'"), FText::FromName(InAssetData.AssetName));
+		},
+		[](const FAssetData& InAssetData)
 		{
-			UBlueprintNodeSpawner* NodeSpawner = nullptr;
-
-			bool const bIsAimOffset = BlendSpace->IsA(UAimOffsetBlendSpace::StaticClass()) ||
-				BlendSpace->IsA(UAimOffsetBlendSpace1D::StaticClass());
-			if (!bIsAimOffset)
-			{
-				NodeSpawner = UBlueprintNodeSpawner::Create(NodeClass);
-				check(NodeSpawner != nullptr);
-
-				TWeakObjectPtr<UBlendSpace> BlendSpacePtr = MakeWeakObjectPtr(const_cast<UBlendSpace*>(BlendSpace));
-				NodeSpawner->CustomizeNodeDelegate = UBlueprintNodeSpawner::FCustomizeNodeDelegate::CreateStatic(GetMenuActions_Utils::SetNodeBlendSpace, BlendSpacePtr);
-			}	
-			return NodeSpawner;
-		}
-	};
-
-	if (const UObject* RegistrarTarget = ActionRegistrar.GetActionKeyFilter())
-	{
-		if (const UBlendSpace* TargetBlendSpace = Cast<UBlendSpace>(RegistrarTarget))
+			return FText::Format(LOCTEXT("MenuDescTooltipFormat", "Blendspace Player\n'{0}'"), FText::FromName(InAssetData.ObjectPath));
+		},
+		[](UEdGraphNode* InNewNode, bool bInIsTemplateNode, const FAssetData InAssetData)
 		{
-			if(TargetBlendSpace->IsAsset())
-			{
-				if (UBlueprintNodeSpawner* NodeSpawner = GetMenuActions_Utils::MakeBlendSpaceAction(GetClass(), TargetBlendSpace))
-				{
-					ActionRegistrar.AddBlueprintAction(TargetBlendSpace, NodeSpawner);
-				}
-			}
-		}
-		// else, the Blueprint database is specifically looking for actions pertaining to something different (not a BlendSpace asset)
-	}
-	else
-	{
-		UClass* NodeClass = GetClass();
-		for (TObjectIterator<UBlendSpace> BlendSpaceIt; BlendSpaceIt; ++BlendSpaceIt)
-		{
-			UBlendSpace* BlendSpace = *BlendSpaceIt;
-			if(BlendSpace->IsAsset())
-			{
-				if (UBlueprintNodeSpawner* NodeSpawner = GetMenuActions_Utils::MakeBlendSpaceAction(NodeClass, BlendSpace))
-				{
-					ActionRegistrar.AddBlueprintAction(BlendSpace, NodeSpawner);
-				}
-			}
-		}
-	}
+			UAnimGraphNode_AssetPlayerBase::SetupNewNode(InNewNode, bInIsTemplateNode, InAssetData);
+		});
 }
 
 FBlueprintNodeSignature UAnimGraphNode_BlendSpacePlayer::GetSignature() const
diff --git a/Engine/Source/Editor/AnimGraph/Private/AnimGraphNode_RotationOffsetBlendSpace.cpp b/Engine/Source/Editor/AnimGraph/Private/AnimGraphNode_RotationOffsetBlendSpace.cpp
--- a/Engine/Source/Editor/AnimGraph/Private/AnimGraphNode_RotationOffsetBlendSpace.cpp
+++ b/Engine/Source/Editor/AnimGraph/Private/AnimGraphNode_RotationOffsetBlendSpace.cpp
@@ -15,6 +15,8 @@
 #include "ScopedTransaction.h"
 #include "Kismet2/BlueprintEditorUtils.h"
 #include "AnimGraphCommands.h"
+#include "BlueprintNodeTemplateCache.h"
+#include "AssetRegistry/AssetRegistryModule.h"
 
 /////////////////////////////////////////////////////
 // UAnimGraphNode_RotationOffsetBlendSpace
@@ -26,12 +28,6 @@ UAnimGraphNode_RotationOffsetBlendSpace::UAnimGraphNode_RotationOffsetBlendSpace
 {
 }
 
-FText UAnimGraphNode_RotationOffsetBlendSpace::GetTooltipText() const
-{
-	// FText::Format() is slow, so we utilize the cached list title
-	return GetNodeTitle(ENodeTitleType::ListView);
-}
-
 FText UAnimGraphNode_RotationOffsetBlendSpace::GetNodeTitle(ENodeTitleType::Type TitleType) const
 {
 	UBlendSpace* BlendSpaceToCheck = Node.GetBlendSpace();
@@ -45,11 +41,11 @@ FText UAnimGraphNode_RotationOffsetBlendSpace::GetNodeTitle(ENodeTitleType::Type
 	{
 		if (TitleType == ENodeTitleType::ListView || TitleType == ENodeTitleType::MenuTitle)
 		{
-			return LOCTEXT("RotationOffsetBlend_NONE_ListTitle", "AimOffset '(None)'");
+			return LOCTEXT("RotationOffsetBlend_NONE_ListTitle", "AimOffset Player '(None)'");
 		}
 		else
 		{
-			return LOCTEXT("RotationOffsetBlend_NONE_Title", "(None)\nAimOffset");
+			return LOCTEXT("RotationOffsetBlend_NONE_Title", "(None)\nAimOffset Player");
 		}
 	}
 	// @TODO: the bone can be altered in the property editor, so we have to 
@@ -64,73 +60,35 @@ FText UAnimGraphNode_RotationOffsetBlendSpace::GetNodeTitle(ENodeTitleType::Type
 		// FText::Format() is slow, so we cache this to save on performance
 		if (TitleType == ENodeTitleType::ListView || TitleType == ENodeTitleType::MenuTitle)
 		{
-			CachedNodeTitles.SetCachedTitle(TitleType, FText::Format(LOCTEXT("AimOffsetListTitle", "AimOffset '{BlendSpaceName}'"), Args), this);
+			CachedNodeTitles.SetCachedTitle(TitleType, FText::Format(LOCTEXT("AimOffsetListTitle", "AimOffset Player '{BlendSpaceName}'"), Args), this);
 		}
 		else
 		{
-			CachedNodeTitles.SetCachedTitle(TitleType, FText::Format(LOCTEXT("AimOffsetFullTitle", "{BlendSpaceName}\nAimOffset"), Args), this);
+			CachedNodeTitles.SetCachedTitle(TitleType, FText::Format(LOCTEXT("AimOffsetFullTitle", "{BlendSpaceName}\nAimOffset Player"), Args), this);
 		}
 	}
 	return CachedNodeTitles[TitleType];
 }
 
-void UAnimGraphNode_RotationOffsetBlendSpace::GetMenuActions(FBlueprintActionDatabaseRegistrar& ActionRegistrar) const
+void UAnimGraphNode_RotationOffsetBlendSpace::GetMenuActions(FBlueprintActionDatabaseRegistrar& InActionRegistrar) const
 {
-	struct GetMenuActions_Utils
-	{
-		static void SetNodeBlendSpace(UEdGraphNode* NewNode, bool /*bIsTemplateNode*/, TWeakObjectPtr<UBlendSpace> BlendSpace)
-		{
-			UAnimGraphNode_RotationOffsetBlendSpace* BlendSpaceNode = CastChecked<UAnimGraphNode_RotationOffsetBlendSpace>(NewNode);
-			BlendSpaceNode->Node.SetBlendSpace(BlendSpace.Get());
-		}
-
-		static UBlueprintNodeSpawner* MakeBlendSpaceAction(TSubclassOf<UEdGraphNode> const NodeClass, const UBlendSpace* BlendSpace)
+	GetMenuActionsHelper(
+		InActionRegistrar,
+		GetClass(),
+		{ UAimOffsetBlendSpace::StaticClass(), UAimOffsetBlendSpace1D::StaticClass() },
+		{ },
+		[](const FAssetData& InAssetData)
 		{
-			UBlueprintNodeSpawner* NodeSpawner = nullptr;
-
-			bool const bIsAimOffset = BlendSpace->IsA(UAimOffsetBlendSpace::StaticClass()) ||
-				BlendSpace->IsA(UAimOffsetBlendSpace1D::StaticClass());
-			if (bIsAimOffset)
-			{
-				NodeSpawner = UBlueprintNodeSpawner::Create(NodeClass);
-				check(NodeSpawner != nullptr);
-
-				TWeakObjectPtr<UBlendSpace> BlendSpacePtr = MakeWeakObjectPtr(const_cast<UBlendSpace*>(BlendSpace));
-				NodeSpawner->CustomizeNodeDelegate = UBlueprintNodeSpawner::FCustomizeNodeDelegate::CreateStatic(GetMenuActions_Utils::SetNodeBlendSpace, BlendSpacePtr);
-			}
-			return NodeSpawner;
-		}
-	};
-
-	if (const UObject* RegistrarTarget = ActionRegistrar.GetActionKeyFilter())
-	{
-		if (const UBlendSpace* TargetBlendSpace = Cast<UBlendSpace>(RegistrarTarget))
+			return FText::Format(LOCTEXT("MenuDescFormat", "AimOffset Player '{0}'"), FText::FromName(InAssetData.AssetName));
+		},
+		[](const FAssetData& InAssetData)
 		{
-			if(TargetBlendSpace->IsAsset())
-			{
-				if (UBlueprintNodeSpawner* NodeSpawner = GetMenuActions_Utils::MakeBlendSpaceAction(GetClass(), TargetBlendSpace))
-				{
-					ActionRegistrar.AddBlueprintAction(TargetBlendSpace, NodeSpawner);
-				}
-			}
-		}
-		// else, the Blueprint database is specifically looking for actions pertaining to something different (not a BlendSpace asset)
-	}
-	else
-	{
-		UClass* NodeClass = GetClass();
-		for (TObjectIterator<UBlendSpace> BlendSpaceIt; BlendSpaceIt; ++BlendSpaceIt)
+			return FText::Format(LOCTEXT("MenuDescTooltipFormat", "AimOffset Player\n'{0}'"), FText::FromName(InAssetData.ObjectPath));
+		},
+		[](UEdGraphNode* InNewNode, bool bInIsTemplateNode, const FAssetData InAssetData)
 		{
-			UBlendSpace* BlendSpace = *BlendSpaceIt;
-			if(BlendSpace->IsAsset())
-			{
-				if (UBlueprintNodeSpawner* NodeSpawner = GetMenuActions_Utils::MakeBlendSpaceAction(NodeClass, BlendSpace))
-				{
-					ActionRegistrar.AddBlueprintAction(BlendSpace, NodeSpawner);
-				}
-			}
-		}
-	}
+			UAnimGraphNode_AssetPlayerBase::SetupNewNode(InNewNode, bInIsTemplateNode, InAssetData);
+		});
 }
 
 FBlueprintNodeSignature UAnimGraphNode_RotationOffsetBlendSpace::GetSignature() const
diff --git a/Engine/Source/Editor/AnimGraph/Private/AnimGraphNode_RotationOffsetBlendSpaceGraph.cpp b/Engine/Source/Editor/AnimGraph/Private/AnimGraphNode_RotationOffsetBlendSpaceGraph.cpp
--- a/Engine/Source/Editor/AnimGraph/Private/AnimGraphNode_RotationOffsetBlendSpaceGraph.cpp
+++ b/Engine/Source/Editor/AnimGraph/Private/AnimGraphNode_RotationOffsetBlendSpaceGraph.cpp
@@ -6,7 +6,9 @@
 #include "Animation/AimOffsetBlendSpace.h"
 #include "Animation/AimOffsetBlendSpace1D.h"
 #include "AnimGraphNodeAlphaOptions.h"
+#include "AnimGraphNode_AssetPlayerBase.h"
 #include "DetailLayoutBuilder.h"
+#include "AssetRegistry/AssetRegistryModule.h"
 
 #define LOCTEXT_NAMESPACE "UAnimGraphNode_RotationOffsetBlendSpaceGraph"
 
@@ -49,86 +51,31 @@ FText UAnimGraphNode_RotationOffsetBlendSpaceGraph::GetNodeTitle(ENodeTitleType:
 	}
 	else
 	{
-		return LOCTEXT("EmptyBlendspaceListTitle", "AimOffset");
+		// Template node, so use the tooltip generated in GetMenuActions
+		return FText::GetEmpty();
 	}
 }
 
-void UAnimGraphNode_RotationOffsetBlendSpaceGraph::GetMenuActions(FBlueprintActionDatabaseRegistrar& ActionRegistrar) const
+void UAnimGraphNode_RotationOffsetBlendSpaceGraph::GetMenuActions(FBlueprintActionDatabaseRegistrar& InActionRegistrar) const
 {
-	struct GetMenuActions_Utils
-	{
-		static void SetNodeBlendSpace(UEdGraphNode* NewNode, bool bIsTemplateNode, TWeakObjectPtr<UBlendSpace> BlendSpace)
-		{
-			UAnimGraphNode_RotationOffsetBlendSpaceGraph* BlendSpaceNode = CastChecked<UAnimGraphNode_RotationOffsetBlendSpaceGraph>(NewNode);
-			BlendSpaceNode->SetupFromAsset(BlendSpace.Get(), bIsTemplateNode);
-		}
-
-		static UBlueprintNodeSpawner* MakeBlendSpaceAction(TSubclassOf<UEdGraphNode> const NodeClass, const UBlendSpace* InBlendSpace)
-		{
-			UBlueprintNodeSpawner* NodeSpawner = nullptr;
-
-			bool const bIsAimOffset = InBlendSpace->IsA(UAimOffsetBlendSpace::StaticClass()) ||
-									  InBlendSpace->IsA(UAimOffsetBlendSpace1D::StaticClass());
-			if (bIsAimOffset)
-			{
-				NodeSpawner = UBlueprintNodeSpawner::Create(NodeClass);
-				check(NodeSpawner != nullptr);
-
-				TWeakObjectPtr<UBlendSpace> BlendSpacePtr = MakeWeakObjectPtr(const_cast<UBlendSpace*>(InBlendSpace));
-				NodeSpawner->CustomizeNodeDelegate = UBlueprintNodeSpawner::FCustomizeNodeDelegate::CreateStatic(GetMenuActions_Utils::SetNodeBlendSpace, BlendSpacePtr);
-			}	
-			return NodeSpawner;
-		}
-
-		static void SetNodeBlendSpaceClass(UEdGraphNode* NewNode, bool bIsTemplateNode, TSubclassOf<UBlendSpace> InBlendSpaceClass)
+	UAnimGraphNode_AssetPlayerBase::GetMenuActionsHelper(
+		InActionRegistrar,
+		GetClass(),
+		{ UAimOffsetBlendSpace::StaticClass(), UAimOffsetBlendSpace1D::StaticClass() },
+		{ },
+		[](const FAssetData& InAssetData)
 		{
-			UAnimGraphNode_RotationOffsetBlendSpaceGraph* BlendSpaceNode = CastChecked<UAnimGraphNode_RotationOffsetBlendSpaceGraph>(NewNode);
-			BlendSpaceNode->SetupFromClass(InBlendSpaceClass, bIsTemplateNode);
-		}
-
-		static UBlueprintNodeSpawner* MakeBlendSpaceAction(TSubclassOf<UEdGraphNode> const NodeClass, TSubclassOf<UBlendSpace> InBlendSpaceClass)
-		{
-			UBlueprintNodeSpawner* NodeSpawner = UBlueprintNodeSpawner::Create(NodeClass);
-			check(NodeSpawner != nullptr);
-
-			NodeSpawner->CustomizeNodeDelegate = UBlueprintNodeSpawner::FCustomizeNodeDelegate::CreateStatic(GetMenuActions_Utils::SetNodeBlendSpaceClass, InBlendSpaceClass);
-	
-			return NodeSpawner;
-		}
-	};
-
-	// Add the non-asset based cases
-	ActionRegistrar.AddBlueprintAction(GetMenuActions_Utils::MakeBlendSpaceAction(GetClass(), UAimOffsetBlendSpace::StaticClass()));
-	ActionRegistrar.AddBlueprintAction(GetMenuActions_Utils::MakeBlendSpaceAction(GetClass(), UAimOffsetBlendSpace1D::StaticClass()));
-
-	if (const UObject* RegistrarTarget = ActionRegistrar.GetActionKeyFilter())
-	{
-		if (const UBlendSpace* TargetBlendSpace = Cast<UBlendSpace>(RegistrarTarget))
+			return FText::Format(LOCTEXT("MenuDescFormat", "AimOffset '{0}'"), FText::FromName(InAssetData.AssetName));
+		},
+		[](const FAssetData& InAssetData)
 		{
-			if(TargetBlendSpace->IsAsset())
-			{
-				if (UBlueprintNodeSpawner* NodeSpawner = GetMenuActions_Utils::MakeBlendSpaceAction(GetClass(), TargetBlendSpace))
-				{
-					ActionRegistrar.AddBlueprintAction(TargetBlendSpace, NodeSpawner);
-				}
-			}
-		}
-		// else, the Blueprint database is specifically looking for actions pertaining to something different (not a BlendSpace asset)
-	}
-	else
-	{
-		UClass* NodeClass = GetClass();
-		for (TObjectIterator<UBlendSpace> BlendSpaceIt; BlendSpaceIt; ++BlendSpaceIt)
+			return FText::Format(LOCTEXT("MenuDescTooltipFormat", "AimOffset\n'{0}'"), FText::FromName(InAssetData.ObjectPath));
+		},
+		[](UEdGraphNode* InNewNode, bool bInIsTemplateNode, const FAssetData InAssetData)
 		{
-			if(BlendSpaceIt->IsAsset())
-			{
-				if (UBlueprintNodeSpawner* NodeSpawner = GetMenuActions_Utils::MakeBlendSpaceAction(NodeClass, *BlendSpaceIt))
-				{
-					ActionRegistrar.AddBlueprintAction(*BlendSpaceIt, NodeSpawner);
-				}
-			}
-		}
-	}
+			UAnimGraphNode_RotationOffsetBlendSpaceGraph* GraphNode = CastChecked<UAnimGraphNode_RotationOffsetBlendSpaceGraph>(InNewNode);
+			GraphNode->SetupFromAsset(InAssetData, bInIsTemplateNode);
+		});
 }
 
 void UAnimGraphNode_RotationOffsetBlendSpaceGraph::BakeDataDuringCompilation(FCompilerResultsLog& MessageLog)
diff --git a/Engine/Source/Editor/AnimGraph/Private/AnimGraphNode_SequenceEvaluator.cpp b/Engine/Source/Editor/AnimGraph/Private/AnimGraphNode_SequenceEvaluator.cpp
--- a/Engine/Source/Editor/AnimGraph/Private/AnimGraphNode_SequenceEvaluator.cpp
+++ b/Engine/Source/Editor/AnimGraph/Private/AnimGraphNode_SequenceEvaluator.cpp
@@ -8,7 +8,13 @@
 #include "Animation/AnimComposite.h"
 #include "Animation/AnimSequence.h"
 #include "AnimGraphNode_SequenceEvaluator.h"
+
+#include "BlueprintNodeSpawner.h"
+#include "EditorCategoryUtils.h"
 #include "IAnimBlueprintNodeOverrideAssetsContext.h"
+#include "AssetRegistry/AssetRegistryModule.h"
+#include "BlueprintActionDatabaseRegistrar.h"
+#include "BlueprintNodeTemplateCache.h"
 
 /////////////////////////////////////////////////////
 // UAnimGraphNode_SequenceEvaluator
@@ -46,10 +52,9 @@ void UAnimGraphNode_SequenceEvaluator::ReplaceReferredAnimations(const TMap<UAni
 	HandleAnimReferenceReplacement(Node.Sequence, AnimAssetReplacementMap);
 }
 
-FText UAnimGraphNode_SequenceEvaluator::GetTooltipText() const
+FText UAnimGraphNode_SequenceEvaluator::GetMenuCategory() const
 {
-	// FText::Format() is slow, so we utilize the cached list title
-	return GetNodeTitle(ENodeTitleType::ListView);
+	return FEditorCategoryUtils::GetCommonCategory(FCommonEditorCategory::Animation);
 }
 
 FText UAnimGraphNode_SequenceEvaluator::GetNodeTitleForSequence(ENodeTitleType::Type TitleType, UAnimSequenceBase* InSequence) const
@@ -103,9 +108,46 @@ FText UAnimGraphNode_SequenceEvaluator::GetNodeTitle(ENodeTitleType::Type TitleT
 	return CachedNodeTitle;
 }
 
-void UAnimGraphNode_SequenceEvaluator::GetMenuActions(FBlueprintActionDatabaseRegistrar& ActionRegistrar) const
+FSlateIcon UAnimGraphNode_SequenceEvaluator::GetIconAndTint(FLinearColor& OutColor) const
+{
+	return FSlateIcon("EditorStyle", "ClassIcon.AnimSequence");
+}
+
+void UAnimGraphNode_SequenceEvaluator::GetMenuActions(FBlueprintActionDatabaseRegistrar& InActionRegistrar) const
 {
-	// Intentionally empty; you can drop down a regular sequence player and convert into a sequence evaluator in the right-click menu.
+	GetMenuActionsHelper(
+		InActionRegistrar,
+		GetClass(),
+		{ UAnimSequence::StaticClass() },
+		{ },
+		[](const FAssetData& InAssetData)
+		{
+			const FString TagValue = InAssetData.GetTagValueRef<FString>(GET_MEMBER_NAME_CHECKED(UAnimSequence, AdditiveAnimType));
+			if(const bool bKnownToBeAdditive = (!TagValue.IsEmpty() && !TagValue.Equals(TEXT("AAT_None"))))
+			{
+				return FText::Format(LOCTEXT("MenuDescFormat", "Evaluate '{0}' (additive)"), FText::FromName(InAssetData.AssetName));
+			}
+			else
+			{
+				return FText::Format(LOCTEXT("MenuDescFormat", "Evaluate '{0}'"), FText::FromName(InAssetData.AssetName));
+			}
+		},
+		[](const FAssetData& InAssetData)
+		{
+			const FString TagValue = InAssetData.GetTagValueRef<FString>(GET_MEMBER_NAME_CHECKED(UAnimSequence, AdditiveAnimType));
+			if(const bool bKnownToBeAdditive = (!TagValue.IsEmpty() && !TagValue.Equals(TEXT("AAT_None"))))
+			{
+				return FText::Format(LOCTEXT("MenuDescTooltipFormat", "Evaluate (additive)\n'{0}'"), FText::FromName(InAssetData.ObjectPath));
+			}
+			else
+			{
+				return FText::Format(LOCTEXT("MenuDescTooltipFormat", "Evaluate\n'{0}'"), FText::FromName(InAssetData.ObjectPath));
+			}
+		},
+		[](UEdGraphNode* InNewNode, bool bInIsTemplateNode, const FAssetData InAssetData)
+		{
+			UAnimGraphNode_AssetPlayerBase::SetupNewNode(InNewNode, bInIsTemplateNode, InAssetData);
+		});
 }
 
 void UAnimGraphNode_SequenceEvaluator::SetAnimationAsset(UAnimationAsset* Asset)
diff --git a/Engine/Source/Editor/AnimGraph/Private/AnimGraphNode_SequencePlayer.cpp b/Engine/Source/Editor/AnimGraph/Private/AnimGraphNode_SequencePlayer.cpp
--- a/Engine/Source/Editor/AnimGraph/Private/AnimGraphNode_SequencePlayer.cpp
+++ b/Engine/Source/Editor/AnimGraph/Private/AnimGraphNode_SequencePlayer.cpp
@@ -23,33 +23,6 @@
 
 #define LOCTEXT_NAMESPACE "A3Nodes"
 
-/////////////////////////////////////////////////////
-// FNewSequencePlayerAction
-
-// Action to add a sequence player node to the graph
-struct FNewSequencePlayerAction : public FEdGraphSchemaAction_K2NewNode
-{
-protected:
-	FAssetData AssetInfo;
-public:
-	FNewSequencePlayerAction(const FAssetData& InAssetInfo, FText Title)
-		: FEdGraphSchemaAction_K2NewNode(LOCTEXT("Animation", "Animations"), Title, LOCTEXT("EvalAnimSequenceToMakePose", "Evaluates an animation sequence to produce a pose"), 0, FText::FromName(InAssetInfo.ObjectPath))
-	{
-		AssetInfo = InAssetInfo;
-
-		UAnimGraphNode_SequencePlayer* Template = NewObject<UAnimGraphNode_SequencePlayer>();
-		NodeTemplate = Template;
-	}
-
-	virtual UEdGraphNode* PerformAction(class UEdGraph* ParentGraph, UEdGraphPin* FromPin, const FVector2D Location, bool bSelectNewNode = true) override
-	{
-		UAnimGraphNode_SequencePlayer* SpawnedNode = CastChecked<UAnimGraphNode_SequencePlayer>(FEdGraphSchemaAction_K2NewNode::PerformAction(ParentGraph, FromPin, Location, bSelectNewNode));
-		SpawnedNode->Node.SetSequence(Cast<UAnimSequence>(AssetInfo.GetAsset()));
-
-		return SpawnedNode;
-	}
-};
-
 /////////////////////////////////////////////////////
 // UAnimGraphNode_SequencePlayer
 
@@ -90,16 +63,9 @@ FLinearColor UAnimGraphNode_SequencePlayer::GetNodeTitleColor() const
 	}
 }
 
-FText UAnimGraphNode_SequencePlayer::GetTooltipText() const
+FSlateIcon UAnimGraphNode_SequencePlayer::GetIconAndTint(FLinearColor& OutColor) const
 {
-	UAnimSequenceBase* Sequence = Node.GetSequence();
-	if (!Sequence)
-	{
-		return FText();
-	}
-
-	const bool bAdditive = Sequence->IsValidAdditive();
-	return GetTitleGivenAssetInfo(FText::FromString(Sequence->GetPathName()), bAdditive);
+	return FSlateIcon("EditorStyle", "ClassIcon.AnimSequence");
 }
 
 FText UAnimGraphNode_SequencePlayer::GetNodeTitleForSequence(ENodeTitleType::Type TitleType, UAnimSequenceBase* InSequence) const
@@ -214,109 +180,41 @@ FText UAnimGraphNode_SequencePlayer::GetMenuCategory() const
 	return FEditorCategoryUtils::GetCommonCategory(FCommonEditorCategory::Animation);
 }
 
-void UAnimGraphNode_SequencePlayer::GetMenuActions(FBlueprintActionDatabaseRegistrar& ActionRegistrar) const
+void UAnimGraphNode_SequencePlayer::GetMenuActions(FBlueprintActionDatabaseRegistrar& InActionRegistrar) const
 {
-	auto LoadedAssetSetup = [](UEdGraphNode* NewNode, bool /*bIsTemplateNode*/, TWeakObjectPtr<UAnimSequence> SequencePtr)
-	{
-		UAnimGraphNode_SequencePlayer* SequencePlayerNode = CastChecked<UAnimGraphNode_SequencePlayer>(NewNode);
-		SequencePlayerNode->Node.SetSequence(SequencePtr.Get());
-	};
-
-	auto UnloadedAssetSetup = [](UEdGraphNode* NewNode, bool bIsTemplateNode, const FAssetData AssetData)
-	{
-		UAnimGraphNode_SequencePlayer* SequencePlayerNode = CastChecked<UAnimGraphNode_SequencePlayer>(NewNode);
-		if (bIsTemplateNode)
-		{
-			AssetData.GetTagValue("Skeleton", SequencePlayerNode->UnloadedSkeletonName);
-		}
-		else
-		{
-			UAnimSequence* Sequence = Cast<UAnimSequence>(AssetData.GetAsset());
-			check(Sequence != nullptr);
-			SequencePlayerNode->Node.SetSequence(Sequence);
-		}
-	};	
-
-	const UObject* QueryObject = ActionRegistrar.GetActionKeyFilter();
-	if (QueryObject == nullptr)
-	{
-		FAssetRegistryModule& AssetRegistryModule = FModuleManager::LoadModuleChecked<FAssetRegistryModule>(TEXT("AssetRegistry"));
-		// define a filter to help in pulling UAnimSequence asset data from the registry
-		FARFilter Filter;
-		Filter.ClassNames.Add(UAnimSequence::StaticClass()->GetFName());
-		Filter.bRecursiveClasses = true;
-		// Find matching assets and add an entry for each one
-		TArray<FAssetData> SequenceList;
-		AssetRegistryModule.Get().GetAssets(Filter, /*out*/SequenceList);
-
-		for (auto AssetIt = SequenceList.CreateConstIterator(); AssetIt; ++AssetIt)
+	GetMenuActionsHelper(
+		InActionRegistrar,
+		GetClass(),
+		{ UAnimSequence::StaticClass() },
+		{ },
+		[](const FAssetData& InAssetData)
 		{
-			const FAssetData& Asset = *AssetIt;			
-
-			UBlueprintNodeSpawner* NodeSpawner = UBlueprintNodeSpawner::Create(GetClass());
-			if (Asset.IsAssetLoaded())
+			const FString TagValue = InAssetData.GetTagValueRef<FString>(GET_MEMBER_NAME_CHECKED(UAnimSequence, AdditiveAnimType));
+			if(const bool bKnownToBeAdditive = (!TagValue.IsEmpty() && !TagValue.Equals(TEXT("AAT_None"))))
 			{
-				UAnimSequence* AnimSequence = Cast<UAnimSequence>(Asset.GetAsset());
-				if(AnimSequence)
-				{
-					NodeSpawner->CustomizeNodeDelegate = UBlueprintNodeSpawner::FCustomizeNodeDelegate::CreateStatic(LoadedAssetSetup, TWeakObjectPtr<UAnimSequence>(AnimSequence));
-					NodeSpawner->DefaultMenuSignature.MenuName = GetTitleGivenAssetInfo(FText::FromName(AnimSequence->GetFName()), AnimSequence->IsValidAdditive());
-					NodeSpawner->DefaultMenuSignature.Tooltip = GetTitleGivenAssetInfo(FText::FromString(AnimSequence->GetPathName()), AnimSequence->IsValidAdditive());
-				}
+				return FText::Format(LOCTEXT("MenuDescFormat", "Play '{0}' (additive)"), FText::FromName(InAssetData.AssetName));
 			}
 			else
 			{
-				const FString TagValue = Asset.GetTagValueRef<FString>(GET_MEMBER_NAME_CHECKED(UAnimSequence, AdditiveAnimType));
-				const bool bKnownToBeAdditive = (!TagValue.IsEmpty() && !TagValue.Equals(TEXT("AAT_None")));
-
-				NodeSpawner->CustomizeNodeDelegate = UBlueprintNodeSpawner::FCustomizeNodeDelegate::CreateStatic(UnloadedAssetSetup, Asset);
-				NodeSpawner->DefaultMenuSignature.MenuName = GetTitleGivenAssetInfo(FText::FromName(Asset.AssetName), bKnownToBeAdditive);
-				NodeSpawner->DefaultMenuSignature.Tooltip = GetTitleGivenAssetInfo(FText::FromName(Asset.ObjectPath), bKnownToBeAdditive);
+				return FText::Format(LOCTEXT("MenuDescFormat", "Play '{0}'"), FText::FromName(InAssetData.AssetName));
 			}
-			ActionRegistrar.AddBlueprintAction(Asset, NodeSpawner);
-		}
-	}
-	else if (const UAnimSequence* AnimSequence = Cast<UAnimSequence>(QueryObject))
-	{
-		UBlueprintNodeSpawner* NodeSpawner = UBlueprintNodeSpawner::Create(GetClass());
-
-		TWeakObjectPtr<UAnimSequence> SequencePtr = MakeWeakObjectPtr(const_cast<UAnimSequence*>(AnimSequence));
-		NodeSpawner->CustomizeNodeDelegate = UBlueprintNodeSpawner::FCustomizeNodeDelegate::CreateStatic(LoadedAssetSetup, SequencePtr);
-		NodeSpawner->DefaultMenuSignature.MenuName = GetTitleGivenAssetInfo(FText::FromName(AnimSequence->GetFName()), AnimSequence->IsValidAdditive());
-		NodeSpawner->DefaultMenuSignature.Tooltip = GetTitleGivenAssetInfo(FText::FromString(AnimSequence->GetPathName()), AnimSequence->IsValidAdditive());
-
-		ActionRegistrar.AddBlueprintAction(QueryObject, NodeSpawner);
-	}
-	else if (QueryObject == GetClass())
-	{
-		FAssetRegistryModule& AssetRegistryModule = FModuleManager::LoadModuleChecked<FAssetRegistryModule>(TEXT("AssetRegistry"));
-		// define a filter to help in pulling UAnimSequence asset data from the registry
-		FARFilter Filter;
-		Filter.ClassNames.Add(UAnimSequence::StaticClass()->GetFName());
-		Filter.bRecursiveClasses = true;
-		// Find matching assets and add an entry for each one
-		TArray<FAssetData> SequenceList;
-		AssetRegistryModule.Get().GetAssets(Filter, /*out*/SequenceList);
-
-		for (auto AssetIt = SequenceList.CreateConstIterator(); AssetIt; ++AssetIt)
+		},
+		[](const FAssetData& InAssetData)
 		{
-			const FAssetData& Asset = *AssetIt;
-			if (Asset.IsAssetLoaded())
+			const FString TagValue = InAssetData.GetTagValueRef<FString>(GET_MEMBER_NAME_CHECKED(UAnimSequence, AdditiveAnimType));
+			if(const bool bKnownToBeAdditive = (!TagValue.IsEmpty() && !TagValue.Equals(TEXT("AAT_None"))))
 			{
-				continue;
+				return FText::Format(LOCTEXT("MenuDescTooltipFormat", "Play (additive)\n'{0}'"), FText::FromName(InAssetData.ObjectPath));
 			}
-
-			UBlueprintNodeSpawner* NodeSpawner = UBlueprintNodeSpawner::Create(GetClass());
-			NodeSpawner->CustomizeNodeDelegate = UBlueprintNodeSpawner::FCustomizeNodeDelegate::CreateStatic(UnloadedAssetSetup, Asset);
-
-			const FString TagValue = Asset.GetTagValueRef<FString>(GET_MEMBER_NAME_CHECKED(UAnimSequence, AdditiveAnimType));
-			const bool bKnownToBeAdditive = (!TagValue.IsEmpty() && !TagValue.Equals(TEXT("AAT_None")));
-
-			NodeSpawner->DefaultMenuSignature.MenuName = GetTitleGivenAssetInfo(FText::FromName(Asset.AssetName), bKnownToBeAdditive);
-			NodeSpawner->DefaultMenuSignature.Tooltip = GetTitleGivenAssetInfo(FText::FromName(Asset.ObjectPath), bKnownToBeAdditive);
-			ActionRegistrar.AddBlueprintAction(Asset, NodeSpawner);
-		}
-	}	
+			else
+			{
+				return FText::Format(LOCTEXT("MenuDescTooltipFormat", "Play\n'{0}'"), FText::FromName(InAssetData.ObjectPath));
+			}
+		},
+		[](UEdGraphNode* InNewNode, bool bInIsTemplateNode, const FAssetData InAssetData)
+		{
+			UAnimGraphNode_AssetPlayerBase::SetupNewNode(InNewNode, bInIsTemplateNode, InAssetData);
+		});	
 }
 
 bool UAnimGraphNode_SequencePlayer::IsActionFilteredOut(class FBlueprintActionFilter const& Filter)
diff --git a/Engine/Source/Editor/AnimGraph/Public/AnimGraphNode_AssetPlayerBase.h b/Engine/Source/Editor/AnimGraph/Public/AnimGraphNode_AssetPlayerBase.h
--- a/Engine/Source/Editor/AnimGraph/Public/AnimGraphNode_AssetPlayerBase.h
+++ b/Engine/Source/Editor/AnimGraph/Public/AnimGraphNode_AssetPlayerBase.h
@@ -31,10 +31,24 @@ public:
 	/** UEdGraphNode interface */
 	virtual void PinConnectionListChanged(UEdGraphPin* Pin) override;
 	virtual void PinDefaultValueChanged(UEdGraphPin* Pin) override;
+	virtual FText GetTooltipText() const override;
+
+	/** UK2Node interface */
+	virtual bool IsActionFilteredOut(class FBlueprintActionFilter const& Filter) override;
 
 	/** UAnimGraphNode_Base interface */
 	virtual void OnProcessDuringCompilation(IAnimBlueprintCompilationContext& InCompilationContext, IAnimBlueprintGeneratedClassCompiledData& OutCompiledData) override;
 	virtual void GetOutputLinkAttributes(FNodeAttributeArray& OutAttributes) const override;
 	virtual void ValidateAnimNodeDuringCompilation(USkeleton* ForSkeleton, FCompilerResultsLog& MessageLog) override;
 	virtual void SetAnimationAsset(UAnimationAsset* Asset) { check(false); /*Base function called*/ }
+
+	// Helper function to gather menu actions from specific asset types supported by this node
+	static void GetMenuActionsHelper(FBlueprintActionDatabaseRegistrar& InActionRegistrar, TSubclassOf<UAnimGraphNode_Base> InNodeClass, const TArray<TSubclassOf<UObject>>& InAssetTypes, const TArray<TSubclassOf<UObject>>& InExcludedAssetTypes, const TFunctionRef<FText(const FAssetData&)>& InMenuNameFunction, const TFunctionRef<FText(const FAssetData&)>& InMenuTooltipFunction, const TFunction<void(UEdGraphNode*, bool, const FAssetData)>& InSetupNewNodeFunction);
+	
+protected:
+	// Default setup function that can be used with GetMenuActionsHalper
+	static void SetupNewNode(UEdGraphNode* InNewNode, bool bInIsTemplateNode, const FAssetData InAssetData);
+
+	/** Used for filtering in the Blueprint context menu when the sequence asset this node uses is unloaded */
+	FString UnloadedSkeletonName;
 };
diff --git a/Engine/Source/Editor/AnimGraph/Public/AnimGraphNode_BlendSpaceBase.h b/Engine/Source/Editor/AnimGraph/Public/AnimGraphNode_BlendSpaceBase.h
--- a/Engine/Source/Editor/AnimGraph/Public/AnimGraphNode_BlendSpaceBase.h
+++ b/Engine/Source/Editor/AnimGraph/Public/AnimGraphNode_BlendSpaceBase.h
@@ -15,6 +15,7 @@ class UAnimGraphNode_BlendSpaceBase : public UAnimGraphNode_AssetPlayerBase
 
 	// UEdGraphNode interface
 	ANIMGRAPH_API virtual FLinearColor GetNodeTitleColor() const override;
+	ANIMGRAPH_API virtual FSlateIcon GetIconAndTint(FLinearColor& OutColor) const override;
 	ANIMGRAPH_API virtual FText GetMenuCategory() const override;
 	// End of UEdGraphNode interface
 
diff --git a/Engine/Source/Editor/AnimGraph/Public/AnimGraphNode_BlendSpaceEvaluator.h b/Engine/Source/Editor/AnimGraph/Public/AnimGraphNode_BlendSpaceEvaluator.h
--- a/Engine/Source/Editor/AnimGraph/Public/AnimGraphNode_BlendSpaceEvaluator.h
+++ b/Engine/Source/Editor/AnimGraph/Public/AnimGraphNode_BlendSpaceEvaluator.h
@@ -20,7 +20,6 @@ class UAnimGraphNode_BlendSpaceEvaluator : public UAnimGraphNode_BlendSpaceBase
 	FAnimNode_BlendSpaceEvaluator Node;
 
 	// UEdGraphNode interface
-	virtual FText GetTooltipText() const override;
 	virtual FText GetNodeTitle(ENodeTitleType::Type TitleType) const override;
 	virtual void GetMenuActions(FBlueprintActionDatabaseRegistrar& ActionRegistrar) const override;
 	// End of UEdGraphNode interface
diff --git a/Engine/Source/Editor/AnimGraph/Public/AnimGraphNode_BlendSpaceGraphBase.h b/Engine/Source/Editor/AnimGraph/Public/AnimGraphNode_BlendSpaceGraphBase.h
--- a/Engine/Source/Editor/AnimGraph/Public/AnimGraphNode_BlendSpaceGraphBase.h
+++ b/Engine/Source/Editor/AnimGraph/Public/AnimGraphNode_BlendSpaceGraphBase.h
@@ -15,7 +15,7 @@ class IAnimBlueprintGeneratedClassCompiledData;
 class UBlendSpaceGraph;
 class UAnimationBlendSpaceSampleGraph;
 
-UCLASS()
+UCLASS(Abstract)
 class ANIMGRAPH_API UAnimGraphNode_BlendSpaceGraphBase : public UAnimGraphNode_Base
 {
 	GENERATED_BODY()
@@ -42,7 +42,7 @@ public:
 	void ReplaceGraph(int32 InSampleIndex, UAnimSequence* InSequence);
 
 	// Setup this node from the specified asset
-	void SetupFromAsset(UBlendSpace* InBlendSpace, bool bInIsTemplateNode);
+	void SetupFromAsset(const FAssetData& InAssetData, bool bInIsTemplateNode);
 
 	// UEdGraphNode interface
 	virtual void PostPlacedNewNode() override;
@@ -79,10 +79,14 @@ protected:
 	UPROPERTY()
 	TArray<TObjectPtr<UEdGraph>> Graphs;
 
+	// Skeleton name used for filtering unloaded assets 
+	FString SkeletonName;
+
 protected:
 	// UEdGraphNode interface
 	virtual FText GetMenuCategory() const override;
 	virtual FLinearColor GetNodeTitleColor() const override;
+	virtual FSlateIcon GetIconAndTint(FLinearColor& OutColor) const override;
 	virtual FText GetTooltipText() const override;
 	virtual UObject* GetJumpTargetForDoubleClick() const override;
 	virtual void JumpToDefinition() const override;
@@ -103,6 +107,7 @@ protected:
 
 	// UK2Node interface
 	virtual void PreloadRequiredAssets() override;
+	virtual bool IsActionFilteredOut(class FBlueprintActionFilter const& Filter) override;
 
 	// Helper function for compilation
 	UAnimGraphNode_Base* ExpandGraphAndProcessNodes(UEdGraph* SourceGraph, UAnimGraphNode_Base* SourceRootNode, IAnimBlueprintCompilationContext& InCompilationContext, IAnimBlueprintGeneratedClassCompiledData& OutCompiledData);
diff --git a/Engine/Source/Editor/AnimGraph/Public/AnimGraphNode_BlendSpacePlayer.h b/Engine/Source/Editor/AnimGraph/Public/AnimGraphNode_BlendSpacePlayer.h
--- a/Engine/Source/Editor/AnimGraph/Public/AnimGraphNode_BlendSpacePlayer.h
+++ b/Engine/Source/Editor/AnimGraph/Public/AnimGraphNode_BlendSpacePlayer.h
@@ -21,7 +21,6 @@ class UAnimGraphNode_BlendSpacePlayer : public UAnimGraphNode_BlendSpaceBase
 	FAnimNode_BlendSpacePlayer Node;
 
 	// UEdGraphNode interface
-	virtual FText GetTooltipText() const override;
 	virtual FText GetNodeTitle(ENodeTitleType::Type TitleType) const override;
 	// End of UEdGraphNode interface
 
diff --git a/Engine/Source/Editor/AnimGraph/Public/AnimGraphNode_RotationOffsetBlendSpace.h b/Engine/Source/Editor/AnimGraph/Public/AnimGraphNode_RotationOffsetBlendSpace.h
--- a/Engine/Source/Editor/AnimGraph/Public/AnimGraphNode_RotationOffsetBlendSpace.h
+++ b/Engine/Source/Editor/AnimGraph/Public/AnimGraphNode_RotationOffsetBlendSpace.h
@@ -21,7 +21,6 @@ class UAnimGraphNode_RotationOffsetBlendSpace: public UAnimGraphNode_BlendSpaceB
 	FAnimNode_RotationOffsetBlendSpace Node;
 
 	// UEdGraphNode interface
-	virtual FText GetTooltipText() const override;
 	virtual FText GetNodeTitle(ENodeTitleType::Type TitleType) const override;
 	virtual void PostEditChangeProperty(struct FPropertyChangedEvent& PropertyChangedEvent) override;
 	// End of UEdGraphNode interface
diff --git a/Engine/Source/Editor/AnimGraph/Public/AnimGraphNode_SequenceEvaluator.h b/Engine/Source/Editor/AnimGraph/Public/AnimGraphNode_SequenceEvaluator.h
--- a/Engine/Source/Editor/AnimGraph/Public/AnimGraphNode_SequenceEvaluator.h
+++ b/Engine/Source/Editor/AnimGraph/Public/AnimGraphNode_SequenceEvaluator.h
@@ -21,7 +21,8 @@ class UAnimGraphNode_SequenceEvaluator : public UAnimGraphNode_AssetPlayerBase
 
 	// UEdGraphNode interface
 	virtual FText GetNodeTitle(ENodeTitleType::Type TitleType) const override;
-	virtual FText GetTooltipText() const override;
+	virtual FSlateIcon GetIconAndTint(FLinearColor& OutColor) const override;
+	virtual FText GetMenuCategory() const override;
 	// End of UEdGraphNode
 
 	// UAnimGraphNode_Base interface
diff --git a/Engine/Source/Editor/AnimGraph/Public/AnimGraphNode_SequencePlayer.h b/Engine/Source/Editor/AnimGraph/Public/AnimGraphNode_SequencePlayer.h
--- a/Engine/Source/Editor/AnimGraph/Public/AnimGraphNode_SequencePlayer.h
+++ b/Engine/Source/Editor/AnimGraph/Public/AnimGraphNode_SequencePlayer.h
@@ -24,7 +24,7 @@ class UAnimGraphNode_SequencePlayer : public UAnimGraphNode_AssetPlayerBase
 
 	// UEdGraphNode interface
 	virtual FLinearColor GetNodeTitleColor() const override;
-	virtual FText GetTooltipText() const override;
+	virtual FSlateIcon GetIconAndTint(FLinearColor& OutColor) const override;
 	virtual FText GetNodeTitle(ENodeTitleType::Type TitleType) const override;
 	virtual FText GetMenuCategory() const override;
 	virtual void PostEditChangeProperty(struct FPropertyChangedEvent& PropertyChangedEvent) override;
@@ -63,7 +63,4 @@ private:
 
 	/** Constructing FText strings can be costly, so we cache the node's title */
 	FNodeTitleTextTable CachedNodeTitles;
-
-	/** Used for filtering in the Blueprint context menu when the sequence asset this node uses is unloaded */
-	FString UnloadedSkeletonName;
 };
