Linked anim graphs/layers now appear in the "my blueprint" outline and context menu

Titles & icons tweaked to emphasise the linked content, not the type of the node.
Also fixes bounding box init issue when spawning new input pose nodes in linked anim graphs
Added the ability to override the tag checked when vierufying skeleton compatibility by FAssetData. This allows anim BPs to use the functionality as they have a "TargetSkeleton" tag vs. other assets with a "Skeleton" tag.

diff --git a/Engine/Source/Editor/AnimGraph/Private/AnimGraphNode_LinkedAnimGraph.cpp b/Engine/Source/Editor/AnimGraph/Private/AnimGraphNode_LinkedAnimGraph.cpp
--- a/Engine/Source/Editor/AnimGraph/Private/AnimGraphNode_LinkedAnimGraph.cpp
+++ b/Engine/Source/Editor/AnimGraph/Private/AnimGraphNode_LinkedAnimGraph.cpp
@@ -1,8 +1,15 @@
 // Copyright Epic Games, Inc. All Rights Reserved.
 
 #include "AnimGraphNode_LinkedAnimGraph.h"
+
+#include "AnimGraphNode_AssetPlayerBase.h"
+#include "BlueprintNodeSpawner.h"
 #include "Animation/AnimBlueprint.h"
+#include "AssetRegistry/AssetRegistryModule.h"
 #include "Engine/Blueprint.h"
+#include "BlueprintActionDatabaseRegistrar.h"
+
+#define LOCTEXT_NAMESPACE "UAnimGraphNode_LinkedAnimGraph"
 
 void UAnimGraphNode_LinkedAnimGraph::PostPasteNode()
 {
@@ -20,4 +27,71 @@ void UAnimGraphNode_LinkedAnimGraph::PostPasteNode()
 			}
 		}
 	}
-}
\ No newline at end of file
+}
+
+void UAnimGraphNode_LinkedAnimGraph::SetupFromAsset(const FAssetData& InAssetData, bool bInIsTemplateNode)
+{
+	InAssetData.GetTagValue("TargetSkeleton", SkeletonName);
+	
+	if(!bInIsTemplateNode)
+	{
+		UAnimBlueprint* AnimBlueprint = CastChecked<UAnimBlueprint>(InAssetData.GetAsset());
+		Node.InstanceClass = AnimBlueprint->GeneratedClass.Get();
+	}
+}
+
+void UAnimGraphNode_LinkedAnimGraph::GetMenuActions(FBlueprintActionDatabaseRegistrar& InActionRegistrar) const
+{
+	// Anim graph node base class will allow us to spawn an 'empty' node
+	UAnimGraphNode_Base::GetMenuActions(InActionRegistrar);
+
+	UAnimGraphNode_AssetPlayerBase::GetMenuActionsHelper(
+		InActionRegistrar,
+		GetClass(),
+		{ UAnimBlueprint::StaticClass()},
+		{ },
+		[](const FAssetData& InAssetData)
+		{
+			return FText::Format(LOCTEXT("MenuDescFormat", "{0} - Linked Anim Graph"), FText::FromName(InAssetData.AssetName));
+		},
+		[](const FAssetData& InAssetData)
+		{
+			return FText::Format(LOCTEXT("MenuDescTooltipFormat", "Linked Anim Graph\n'{0}'"), FText::FromName(InAssetData.ObjectPath));
+		},
+		[](UEdGraphNode* InNewNode, bool bInIsTemplateNode, const FAssetData InAssetData)
+		{
+			UAnimGraphNode_LinkedAnimGraph* GraphNode = CastChecked<UAnimGraphNode_LinkedAnimGraph>(InNewNode);
+			GraphNode->SetupFromAsset(InAssetData, bInIsTemplateNode);
+		});
+}
+
+bool UAnimGraphNode_LinkedAnimGraph::IsActionFilteredOut(class FBlueprintActionFilter const& Filter)
+{
+	bool bIsFilteredOut = false;
+
+	if(!SkeletonName.IsEmpty())
+	{
+		FBlueprintActionContext const& FilterContext = Filter.Context;
+
+		for (UBlueprint* Blueprint : FilterContext.Blueprints)
+		{
+			if (UAnimBlueprint* AnimBlueprint = Cast<UAnimBlueprint>(Blueprint))
+			{
+				if(!AnimBlueprint->TargetSkeleton->IsCompatibleSkeletonByAssetString(SkeletonName))
+				{
+					bIsFilteredOut = true;
+					break;
+				}
+			}
+			else
+			{
+				// Not an animation Blueprint, cannot use
+				bIsFilteredOut = true;
+				break;
+			}
+		}
+	}
+	return bIsFilteredOut;
+}
+
+#undef LOCTEXT_NAMESPACE
\ No newline at end of file
diff --git a/Engine/Source/Editor/AnimGraph/Private/AnimGraphNode_LinkedAnimGraphBase.cpp b/Engine/Source/Editor/AnimGraph/Private/AnimGraphNode_LinkedAnimGraphBase.cpp
--- a/Engine/Source/Editor/AnimGraph/Private/AnimGraphNode_LinkedAnimGraphBase.cpp
+++ b/Engine/Source/Editor/AnimGraph/Private/AnimGraphNode_LinkedAnimGraphBase.cpp
@@ -19,6 +19,7 @@
 #include "AnimGraphAttributes.h"
 #include "IAnimBlueprintCopyTermDefaultsContext.h"
 #include "AnimBlueprintExtension_LinkedAnimGraph.h"
+#include "EdGraphSchema_K2_Actions.h"
 
 #define LOCTEXT_NAMESPACE "LinkedAnimGraph"
 
@@ -55,49 +56,67 @@ FLinearColor UAnimGraphNode_LinkedAnimGraphBase::GetNodeTitleColor() const
 	return LinkedAnimGraphGraphNodeConstants::TitleColor;
 }
 
+FSlateIcon UAnimGraphNode_LinkedAnimGraphBase::GetIconAndTint(FLinearColor& OutColor) const
+{
+	return FSlateIcon("EditorStyle", "ClassIcon.AnimBlueprint");
+}
+
 FText UAnimGraphNode_LinkedAnimGraphBase::GetTooltipText() const
 {
 	return LOCTEXT("ToolTip", "Runs a linked anim graph in another instance to process animation");
 }
 
+FText UAnimGraphNode_LinkedAnimGraphBase::GetMenuCategory() const
+{
+	return LOCTEXT("LinkedAnimGraphCategory", "Linked Anim Blueprints");
+}
+
 FText UAnimGraphNode_LinkedAnimGraphBase::GetNodeTitle(ENodeTitleType::Type TitleType) const
 {
 	UClass* TargetClass = GetTargetClass();
 	UAnimBlueprint* TargetAnimBlueprint = TargetClass ? CastChecked<UAnimBlueprint>(TargetClass->ClassGeneratedBy) : nullptr;
-
+	const FName TargetAnimBlueprintName = TargetAnimBlueprint ? TargetAnimBlueprint->GetFName() : NAME_None;
+	
 	const FAnimNode_LinkedAnimGraph& Node = *GetLinkedAnimGraphNode();
 
 	FFormatNamedArguments Args;
-	Args.Add(TEXT("NodeTitle"), LOCTEXT("Title", "Linked Anim Graph"));
-	Args.Add(TEXT("TargetClass"), TargetAnimBlueprint ? FText::FromString(TargetAnimBlueprint->GetName()) : LOCTEXT("ClassNone", "None"));
+	Args.Add(TEXT("NodeType"), LOCTEXT("Title", "Linked Anim Graph"));
+	Args.Add(TEXT("TargetClass"), FText::FromName(TargetAnimBlueprintName));
 
 	if(TitleType == ENodeTitleType::MenuTitle)
 	{
-		return LOCTEXT("NodeTitle", "Linked Anim Graph");
+		return LOCTEXT("Title", "Linked Anim Graph");
 	}
-	if(TitleType == ENodeTitleType::ListView)
+	else if(TitleType == ENodeTitleType::ListView)
 	{
 		if(Node.Tag != NAME_None)
 		{
 			Args.Add(TEXT("Tag"), FText::FromName(Node.Tag));
-			return FText::Format(LOCTEXT("TitleListFormatTagged", "{NodeTitle} ({Tag}) - {TargetClass}"), Args);
+			return FText::Format(LOCTEXT("TitleListFormatTagged", "{TargetClass} ({Tag})"), Args);
+		}
+		else if(TargetClass)
+		{
+			return FText::Format(LOCTEXT("TitleListFormat", "{TargetClass}"), Args);
 		}
 		else
 		{
-			return FText::Format(LOCTEXT("TitleListFormat", "{NodeTitle} - {TargetClass}"), Args);
+			return LOCTEXT("Title", "Linked Anim Graph");
 		}
 	}
 	else
 	{
-
 		if(Node.Tag != NAME_None)
 		{
 			Args.Add(TEXT("Tag"), FText::FromName(Node.Tag));
-			return FText::Format(LOCTEXT("TitleFormatTagged", "{NodeTitle} ({Tag})\n{TargetClass}"), Args);
+			return FText::Format(LOCTEXT("TitleFormatTagged", "{TargetClass} ({Tag})\n{NodeType} "), Args);
+		}
+		else if(TargetClass)
+		{
+			return FText::Format(LOCTEXT("TitleFormat", "{TargetClass}\n{NodeType}"), Args);
 		}
 		else
 		{
-			return FText::Format(LOCTEXT("TitleFormat", "{NodeTitle}\n{TargetClass}"), Args);
+			return LOCTEXT("Title", "Linked Anim Graph");
 		}
 	}
 }
@@ -495,4 +514,11 @@ void UAnimGraphNode_LinkedAnimGraphBase::GetRequiredExtensions(TArray<TSubclassO
 	OutExtensions.Add(UAnimBlueprintExtension_LinkedAnimGraph::StaticClass());
 }
 
+TSharedPtr<FEdGraphSchemaAction> UAnimGraphNode_LinkedAnimGraphBase::GetEventNodeAction(const FText& ActionCategory)
+{
+	TSharedPtr<FEdGraphSchemaAction_K2Event> NodeAction = MakeShareable(new FEdGraphSchemaAction_K2Event(ActionCategory, GetNodeTitle(ENodeTitleType::ListView), GetTooltipText(), 0));
+	NodeAction->NodeTemplate = this;
+	return NodeAction;
+}
+
 #undef LOCTEXT_NAMESPACE
diff --git a/Engine/Source/Editor/AnimGraph/Private/AnimGraphNode_LinkedAnimLayer.cpp b/Engine/Source/Editor/AnimGraph/Private/AnimGraphNode_LinkedAnimLayer.cpp
--- a/Engine/Source/Editor/AnimGraph/Private/AnimGraphNode_LinkedAnimLayer.cpp
+++ b/Engine/Source/Editor/AnimGraph/Private/AnimGraphNode_LinkedAnimLayer.cpp
@@ -17,7 +17,10 @@
 #include "UObject/CoreRedirects.h"
 #include "Kismet2/KismetEditorUtilities.h"
 #include "AnimationStateGraph.h"
+#include "BlueprintNodeSpawner.h"
 #include "UObject/FortniteMainBranchObjectVersion.h"
+#include "BlueprintActionDatabaseRegistrar.h"
+#include "AssetRegistry/AssetRegistryModule.h"
 
 #define LOCTEXT_NAMESPACE "LinkedAnimLayerNode"
 
@@ -51,6 +54,11 @@ void UAnimGraphNode_LinkedAnimLayer::ReconstructNode()
 	Super::ReconstructNode();
 }
 
+FSlateIcon UAnimGraphNode_LinkedAnimLayer::GetIconAndTint(FLinearColor& OutColor) const
+{
+	return FSlateIcon("EditorStyle", "ClassIcon.AnimLayerInterface");
+}
+
 FText UAnimGraphNode_LinkedAnimLayer::GetTooltipText() const
 {
 	return LOCTEXT("ToolTip", "Runs another linked animation layer graph to process animation");
@@ -92,30 +100,47 @@ FText UAnimGraphNode_LinkedAnimLayer::GetNodeTitle(ENodeTitleType::Type TitleTyp
 	}
 	else
 	{
-	FFormatNamedArguments Args;
-		Args.Add(TEXT("NodeTitle"), LOCTEXT("NodeTitle", "Linked Anim Layer"));
-		Args.Add(TEXT("TargetClass"), TargetAnimBlueprintInterface ? FText::FromString(TargetAnimBlueprintInterface->GetName()) : LOCTEXT("InterfaceNone", "None"));
+		bool bIsSelf = TargetAnimBlueprintInterface == nullptr; 
+		
+		FFormatNamedArguments Args;
+		Args.Add(TEXT("NodeType"), LOCTEXT("NodeTitle", "Linked Anim Layer"));
+		Args.Add(TEXT("TargetClass"), bIsSelf ? LOCTEXT("ClassSelf", "Self") : FText::FromString(TargetAnimBlueprintInterface->GetName()));
 		Args.Add(TEXT("Layer"), (Node.Layer == NAME_None) ? LOCTEXT("LayerNone", "None") : FText::FromName(Node.Layer));
 
-		if (FAnimNode_LinkedAnimLayer* PreviewNode = GetPreviewNode())
+		if (TitleType == ENodeTitleType::ListView)
 		{
-			if (UAnimInstance* PreviewAnimInstance = PreviewNode->GetTargetInstance<UAnimInstance>())
+			if(bIsSelf)
 			{
-				if (UClass* PreviewTargetClass = PreviewAnimInstance->GetClass())
-	{
-					Args.Add(TEXT("TargetClass"), PreviewTargetClass == GetAnimBlueprint()->GeneratedClass ? LOCTEXT("ClassSelf", "Self") : FText::FromName(PreviewTargetClass->GetFName()));
-				}
-	}
+				return FText::Format(LOCTEXT("TitleListViewFormatSelf", "{Layer}"), Args);
+			}
+			else
+			{
+				return FText::Format(LOCTEXT("TitleListViewFormat", "{TargetClass} - {Layer}"), Args);
+			}
 		}
+		else
+		{
+			if (FAnimNode_LinkedAnimLayer* PreviewNode = GetPreviewNode())
+			{
+				if (UAnimInstance* PreviewAnimInstance = PreviewNode->GetTargetInstance<UAnimInstance>())
+				{
+					if (UClass* PreviewTargetClass = PreviewAnimInstance->GetClass())
+					{
+						bIsSelf = PreviewTargetClass == GetAnimBlueprint()->GeneratedClass;
+						Args.Add(TEXT("TargetClass"), PreviewTargetClass == GetAnimBlueprint()->GeneratedClass ? LOCTEXT("ClassSelf", "Self") : FText::FromName(PreviewTargetClass->GetFName()));
+					}
+				}
+			}
 
-		if (TitleType == ENodeTitleType::ListView)
-	{
-		return FText::Format(LOCTEXT("TitleListFormatOutputPose", "{NodeTitle}: {Layer} - {TargetClass}"), Args);
-	}
-	else
-	{
-		return FText::Format(LOCTEXT("TitleFormatOutputPose", "{NodeTitle}: {Layer}\n{TargetClass}"), Args);
-	}
+			if(bIsSelf)
+			{
+				return FText::Format(LOCTEXT("TitleOtherFormatSelf", "{Layer}\n{NodeType}"), Args);
+			}
+			else
+			{
+				return FText::Format(LOCTEXT("TitleOtherFormat", "{TargetClass} - {Layer}\n{NodeType}"), Args);
+			}
+		}
 	}
 }
 
@@ -701,4 +726,69 @@ void UAnimGraphNode_LinkedAnimLayer::HandleInstanceChanged()
 	NodeTitleChangedEvent.Broadcast();
 }
 
+void UAnimGraphNode_LinkedAnimLayer::SetupFromLayerId(FName InLayer)
+{
+	Node.Layer = InLayer;
+	
+	// Get the interface for this layer. If null, then we are using a 'self' layer.
+	Node.Interface = GetInterfaceForLayer();
+
+	// Update the Guid for conforming
+	InterfaceGuid = GetGuidForLayer();
+
+	if(Node.Interface.Get() == nullptr)
+	{
+		// Self layers cannot have override implementations
+		Node.InstanceClass = nullptr;
+	}
+}
+
+void UAnimGraphNode_LinkedAnimLayer::GetMenuActions(FBlueprintActionDatabaseRegistrar& ActionRegistrar) const
+{
+	// Anim graph node base class will allow us to spawn an 'empty' node
+	UAnimGraphNode_Base::GetMenuActions(ActionRegistrar);
+	
+	auto MakeAnimBlueprintAction = [](TSubclassOf<UEdGraphNode> const NodeClass, const FName& InLayerId)
+	{
+		auto SetNodeLayerId = [](UEdGraphNode* NewNode, bool bIsTemplateNode, FName InLayerId)
+		{
+			UAnimGraphNode_LinkedAnimLayer* LinkedAnimGraphNode = CastChecked<UAnimGraphNode_LinkedAnimLayer>(NewNode);
+			LinkedAnimGraphNode->SetupFromLayerId(InLayerId);
+		};
+		
+		UBlueprintNodeSpawner* NodeSpawner = UBlueprintNodeSpawner::Create(NodeClass);
+		check(NodeSpawner != nullptr);
+		NodeSpawner->CustomizeNodeDelegate = UBlueprintNodeSpawner::FCustomizeNodeDelegate::CreateStatic(SetNodeLayerId, InLayerId);
+		NodeSpawner->DefaultMenuSignature.Category = LOCTEXT("LinkedAnimLayerCategory", "Linked Anim Layers");
+		NodeSpawner->DefaultMenuSignature.MenuName = NodeSpawner->DefaultMenuSignature.Tooltip = FText::Format(LOCTEXT("LinkedAnimGraphMenuFormat", "{0} - Linked Anim Layer"), FText::FromName(InLayerId));
+		return NodeSpawner;
+	};
+
+	if (const UObject* RegistrarTarget = ActionRegistrar.GetActionKeyFilter())
+	{
+		if (const UAnimBlueprint* TargetAnimBlueprint = Cast<UAnimBlueprint>(RegistrarTarget))
+		{
+			UClass* TargetClass = *TargetAnimBlueprint->SkeletonGeneratedClass;
+			if(TargetClass)
+			{
+				UClass* NodeClass = GetClass();
+				IAnimClassInterface* AnimClassInterface = IAnimClassInterface::GetFromClass(TargetClass);
+				for(const FAnimBlueprintFunction& AnimBlueprintFunction : AnimClassInterface->GetAnimBlueprintFunctions())
+				{
+					if(AnimBlueprintFunction.Name != UEdGraphSchema_K2::GN_AnimGraph)
+					{
+						if(UFunction* Function = TargetClass->FindFunctionByName(AnimBlueprintFunction.Name))
+						{
+							if (UBlueprintNodeSpawner* NodeSpawner = MakeAnimBlueprintAction(NodeClass, AnimBlueprintFunction.Name))
+							{
+								ActionRegistrar.AddBlueprintAction(Function, NodeSpawner);
+							}
+						}
+					}
+				}
+			}
+		}
+	}
+}
+
 #undef LOCTEXT_NAMESPACE
diff --git a/Engine/Source/Editor/AnimGraph/Private/AnimationGraphSchema.cpp b/Engine/Source/Editor/AnimGraph/Private/AnimationGraphSchema.cpp
--- a/Engine/Source/Editor/AnimGraph/Private/AnimationGraphSchema.cpp
+++ b/Engine/Source/Editor/AnimGraph/Private/AnimationGraphSchema.cpp
@@ -848,7 +848,7 @@ FVector2D UAnimationGraphSchema::GetPositionForNewLinkedInputPoseNode(UEdGraph&
 		InGraph.GetNodesOfClass<UAnimGraphNode_Base>(AllNodes);
 
 		// No nodes, so insert to the top-left of all existing nodes.
-		FBox2D AllNodesBounds;
+		FBox2D AllNodesBounds(ForceInit);
 		for(UAnimGraphNode_Base* Node : AllNodes)
 		{
 			FBox2D NodeBounds(
diff --git a/Engine/Source/Editor/AnimGraph/Public/AnimGraphNode_LinkedAnimGraph.h b/Engine/Source/Editor/AnimGraph/Public/AnimGraphNode_LinkedAnimGraph.h
--- a/Engine/Source/Editor/AnimGraph/Public/AnimGraphNode_LinkedAnimGraph.h
+++ b/Engine/Source/Editor/AnimGraph/Public/AnimGraphNode_LinkedAnimGraph.h
@@ -23,6 +23,10 @@ public:
 	// Begin UEdGraphNode
 	virtual void PostPasteNode();
 
+	// Begin UK2Node
+	virtual void GetMenuActions(FBlueprintActionDatabaseRegistrar& ActionRegistrar) const override;
+	virtual bool IsActionFilteredOut(class FBlueprintActionFilter const& Filter) override;
+	
 	// Begin UAnimGraphNode_CustomProperty
 	virtual FAnimNode_CustomProperty* GetCustomPropertyNode() override { return &Node;  }
 	virtual const FAnimNode_CustomProperty* GetCustomPropertyNode() const override { return &Node; }
@@ -30,6 +34,10 @@ public:
 	// Begin UAnimGraphNode_LinkedAnimGraphBase
 	virtual FAnimNode_LinkedAnimGraph* GetLinkedAnimGraphNode() override { return &Node; }
 	virtual const FAnimNode_LinkedAnimGraph* GetLinkedAnimGraphNode() const override { return &Node; }
+
+private:
+	// Setup the node from a specified anim BP
+	void SetupFromAsset(const FAssetData& InAssetData, bool bInIsTemplateNode);
 };
 
 UE_DEPRECATED(4.24, "UAnimGraphNode_SubInstance has been renamed to UAnimGraphNode_LinkedAnimGraph")
diff --git a/Engine/Source/Editor/AnimGraph/Public/AnimGraphNode_LinkedAnimGraphBase.h b/Engine/Source/Editor/AnimGraph/Public/AnimGraphNode_LinkedAnimGraphBase.h
--- a/Engine/Source/Editor/AnimGraph/Public/AnimGraphNode_LinkedAnimGraphBase.h
+++ b/Engine/Source/Editor/AnimGraph/Public/AnimGraphNode_LinkedAnimGraphBase.h
@@ -6,7 +6,9 @@
 #include "UObject/ObjectMacros.h"
 #include "Misc/Guid.h"
 #include "AnimGraphNode_CustomProperty.h"
+#include "EdGraphSchema_K2_Actions.h"
 #include "IClassVariableCreator.h"
+#include "K2Node_EventNodeInterface.h"
 
 #include "AnimGraphNode_LinkedAnimGraphBase.generated.h"
 
@@ -17,14 +19,16 @@ class SToolTip;
 struct FAnimNode_LinkedAnimGraph;
 
 UCLASS(MinimalAPI, Abstract)
-class UAnimGraphNode_LinkedAnimGraphBase : public UAnimGraphNode_CustomProperty
+class UAnimGraphNode_LinkedAnimGraphBase : public UAnimGraphNode_CustomProperty, public IK2Node_EventNodeInterface
 {
 	GENERATED_BODY()
 
 public:
 	//~ Begin UEdGraphNode Interface.
 	virtual FLinearColor GetNodeTitleColor() const override;
+	virtual FSlateIcon GetIconAndTint(FLinearColor& OutColor) const override;
 	virtual FText GetTooltipText() const override;
+	virtual FText GetMenuCategory() const override;
 	virtual FText GetNodeTitle(ENodeTitleType::Type TitleType) const override;
 	virtual void ValidateAnimNodeDuringCompilation(USkeleton* ForSkeleton, FCompilerResultsLog& MessageLog) override;
 	virtual void ReallocatePinsDuringReconstruction(TArray<UEdGraphPin*>& OldPins) override;
@@ -46,6 +50,9 @@ public:
 	virtual FAnimNode_LinkedAnimGraph* GetLinkedAnimGraphNode() PURE_VIRTUAL(UAnimGraphNode_LinkedAnimGraphBase::GetLinkedAnimGraphNode, return nullptr;);
 	virtual const FAnimNode_LinkedAnimGraph* GetLinkedAnimGraphNode() const PURE_VIRTUAL(UAnimGraphNode_LinkedAnimGraphBase::GetLinkedAnimGraphNode, return nullptr;);
 
+	// IK2Node_EventNodeInterface interface
+	virtual TSharedPtr<FEdGraphSchemaAction> GetEventNodeAction(const FText& ActionCategory) override;
+
 protected:
 	friend class UAnimBlueprintExtension_LinkedAnimGraph;
 
@@ -73,6 +80,9 @@ protected:
 	// Instance blueprint was changed by user
 	void OnSetInstanceBlueprint(const FAssetData& AssetData, IDetailLayoutBuilder* InDetailBuilder);
 	// ----- END UI CALLBACKS ----- //
+
+	// Skeleton name used for filtering unloaded classes 
+	FString SkeletonName;
 };
 
 UE_DEPRECATED(4.24, "UAnimGraphNode_SubInstanceBase has been renamed to UAnimGraphNode_LinkedAnimGraphBase")
diff --git a/Engine/Source/Editor/AnimGraph/Public/AnimGraphNode_LinkedAnimLayer.h b/Engine/Source/Editor/AnimGraph/Public/AnimGraphNode_LinkedAnimLayer.h
--- a/Engine/Source/Editor/AnimGraph/Public/AnimGraphNode_LinkedAnimLayer.h
+++ b/Engine/Source/Editor/AnimGraph/Public/AnimGraphNode_LinkedAnimLayer.h
@@ -34,6 +34,7 @@ public:
 
 	//~ Begin UEdGraphNode Interface.
 	virtual FText GetTooltipText() const override;
+	virtual FSlateIcon GetIconAndTint(FLinearColor& OutColor) const override;
 	virtual FText GetNodeTitle(ENodeTitleType::Type TitleType) const override;
 	virtual void ValidateAnimNodeDuringCompilation(USkeleton* ForSkeleton, FCompilerResultsLog& MessageLog) override;
 	virtual UObject* GetJumpTargetForDoubleClick() const override;
@@ -42,6 +43,9 @@ public:
 	virtual void ReconstructNode() override;
 	//~ End UEdGraphNode Interface.
 
+	// UK2Node interface
+	virtual void GetMenuActions(FBlueprintActionDatabaseRegistrar& ActionRegistrar) const override;
+	
 	// UAnimGraphNode_Base interface
 	virtual void CustomizeDetails(IDetailLayoutBuilder& DetailBuilder) override;
 
@@ -83,6 +87,9 @@ protected:
 	// Get the preview node, if any, when instanced in an animation blueprint and debugged
 	FAnimNode_LinkedAnimLayer* GetPreviewNode() const;
 
+	// Helper function to setup a newly spawned node
+	void SetupFromLayerId(FName InLayerId);
+	
 	// Handle used to hook into object being debugged changing
 	FDelegateHandle SetObjectBeingDebuggedHandle;
 };
diff --git a/Engine/Source/Runtime/Engine/Classes/Animation/Skeleton.h b/Engine/Source/Runtime/Engine/Classes/Animation/Skeleton.h
--- a/Engine/Source/Runtime/Engine/Classes/Animation/Skeleton.h
+++ b/Engine/Source/Runtime/Engine/Classes/Animation/Skeleton.h
@@ -802,7 +802,7 @@ public:
 	/**
 	 * Check if this skeleton is compatible with a given other asset, if that is a skeleton.
 	 */
-	ENGINE_API bool IsCompatibleSkeletonByAssetData(const FAssetData& AssetData) const;
+	ENGINE_API bool IsCompatibleSkeletonByAssetData(const FAssetData& AssetData, const TCHAR* InTag = TEXT("Skeleton")) const;
 
 	/**
 	 * Check if this skeleton is compatible with another skeleton asset that is identified by the string returned by AssetData(SkeletonPtr).GetExportTextName().
diff --git a/Engine/Source/Runtime/Engine/Private/Animation/Skeleton.cpp b/Engine/Source/Runtime/Engine/Private/Animation/Skeleton.cpp
--- a/Engine/Source/Runtime/Engine/Private/Animation/Skeleton.cpp
+++ b/Engine/Source/Runtime/Engine/Private/Animation/Skeleton.cpp
@@ -158,9 +158,9 @@ bool USkeleton::IsCompatibleSkeletonByAssetString(const FString& SkeletonAssetSt
 	return false;
 }
 
-bool USkeleton::IsCompatibleSkeletonByAssetData(const FAssetData& AssetData) const
+bool USkeleton::IsCompatibleSkeletonByAssetData(const FAssetData& AssetData, const TCHAR* InTag) const
 {
-	return IsCompatibleSkeletonByAssetString(AssetData.GetTagValueRef<FString>(TEXT("Skeleton")));
+	return IsCompatibleSkeletonByAssetString(AssetData.GetTagValueRef<FString>(InTag));
 }
 
 FSkeletonRemapping::FSkeletonRemapping(const USkeleton* InSourceSkeleton, const USkeleton* InTargetSkeleton)
