Fixed property access binding copies being skipped if a node uses BP

Fast path copies were unaffected as they defaulted to their generated BP logic, despite the fast path not being generated.

diff --git a/Engine/Source/Editor/AnimGraph/Private/AnimBlueprintExtension_Base.cpp b/Engine/Source/Editor/AnimGraph/Private/AnimBlueprintExtension_Base.cpp
--- a/Engine/Source/Editor/AnimGraph/Private/AnimBlueprintExtension_Base.cpp
+++ b/Engine/Source/Editor/AnimGraph/Private/AnimBlueprintExtension_Base.cpp
@@ -525,7 +525,7 @@ void UAnimBlueprintExtension_Base::CreateEvaluationHandler(IAnimBlueprintCompila
 							// Copy the data (link up to the source nodes)
 							TargetPin->CopyPersistentDataFromOldPin(*DestPin);
 							InCompilationContext.GetMessageLog().NotifyIntermediatePinCreation(TargetPin, DestPin);
-
+							
 							break;
 						}
 					}
@@ -747,25 +747,23 @@ void UAnimBlueprintExtension_Base::FEvaluationHandlerRecord::PatchFunctionNameAn
 {
 	Handler.CopyRecords.Empty();
 
-	if (IsFastPath())
+	for (const TPair<FName, FAnimNodeSinglePropertyHandler>& ServicedPropPair : ServicedProperties)
 	{
-		for (const TPair<FName, FAnimNodeSinglePropertyHandler>& ServicedPropPair : ServicedProperties)
-		{
-			const FName& PropertyName = ServicedPropPair.Key;
-			const FAnimNodeSinglePropertyHandler& PropertyHandler = ServicedPropPair.Value;
+		const FName& PropertyName = ServicedPropPair.Key;
+		const FAnimNodeSinglePropertyHandler& PropertyHandler = ServicedPropPair.Value;
 
-			for (const FPropertyCopyRecord& PropertyCopyRecord : PropertyHandler.CopyRecords)
+		for (const FPropertyCopyRecord& PropertyCopyRecord : PropertyHandler.CopyRecords)
+		{
+			// Only unbatched copies can be processed on a per-node basis
+			// Skip invalid copy indices as these are usually the result of BP errors/warnings
+			if(PropertyCopyRecord.LibraryCompiledHandle.IsValid() && PropertyCopyRecord.LibraryCompiledHandle.GetBatchId() == (int32)EAnimPropertyAccessCallSite::WorkerThread_Unbatched)
 			{
-				// Only unbatched copies can be processed on a per-node basis
-				// Skip invalid copy indices as these are usually the result of BP errors/warnings
-				if(PropertyCopyRecord.LibraryCompiledHandle.IsValid() && PropertyCopyRecord.LibraryCompiledHandle.GetBatchId() == (int32)EAnimPropertyAccessCallSite::WorkerThread_Unbatched)
-				{
-					Handler.CopyRecords.Emplace(PropertyCopyRecord.LibraryCompiledHandle.GetId(), PropertyCopyRecord.Operation);
-				}
+				Handler.CopyRecords.Emplace(PropertyCopyRecord.LibraryCompiledHandle.GetId(), PropertyCopyRecord.Operation);
 			}
 		}
 	}
-	else
+
+	if (!IsFastPath())
 	{
 		// not all of our pins use copy records so we will need to call our exposed value handler
 		Handler.BoundFunction = HandlerFunctionName;
