Fix for structure accesses from property access nodes not linking their expanded nodes correctly

#rn Jurre.deBaare

diff --git a/Engine/Source/Editor/AnimGraph/Private/AnimBlueprintExtension_PropertyAccess.cpp b/Engine/Source/Editor/AnimGraph/Private/AnimBlueprintExtension_PropertyAccess.cpp
--- a/Engine/Source/Editor/AnimGraph/Private/AnimBlueprintExtension_PropertyAccess.cpp
+++ b/Engine/Source/Editor/AnimGraph/Private/AnimBlueprintExtension_PropertyAccess.cpp
@@ -261,6 +261,8 @@ void UAnimBlueprintExtension_PropertyAccess::ExpandPropertyAccess(FKismetCompile
 				// Create a break struct/split pin node
 				const UEdGraphSchema_K2* K2Schema = CastChecked<UEdGraphSchema_K2>(InParentGraph->GetSchema());
 				UK2Node* SplitPinNode = K2Schema->CreateSplitPinNode(CurrentPin, UEdGraphSchema_K2::FCreateSplitPinNodeParams(&InCompilerContext, InParentGraph));
+				UEdGraphPin* InputPin = SplitPinNode->FindPinChecked(ScriptStruct->GetFName(), EGPD_Input);
+				CurrentPin->MakeLinkTo(InputPin);
 
 				// Current pin is the property of the struct
 				CurrentPin = SplitPinNode->FindPinChecked(InPropertyName);
