Fix 'expose as pin' menu item not appearing

Menu extender wasnt being correctly combined with a last minute change to fix a CIS issue.

diff --git a/Engine/Source/Editor/AnimGraph/Private/AnimGraphNode_Base.cpp b/Engine/Source/Editor/AnimGraph/Private/AnimGraphNode_Base.cpp
--- a/Engine/Source/Editor/AnimGraph/Private/AnimGraphNode_Base.cpp
+++ b/Engine/Source/Editor/AnimGraph/Private/AnimGraphNode_Base.cpp
@@ -1163,7 +1163,7 @@ TSharedRef<SWidget> UAnimGraphNode_Base::MakePropertyBindingWidget(const FAnimPr
 		MenuArgs.OnGetPropertyAccessContextId = FOnGetPropertyAccessContextId::CreateLambda(OnGetPropertyAccessContextId);
 		
 		// Add the binding menu extenders
-		TArray<TSharedPtr<FExtender>> Extenders;
+		TArray<TSharedPtr<FExtender>> Extenders( { Args.MenuExtender } );
 		for(IPropertyAccessBlueprintBinding* Binding : IModularFeatures::Get().GetModularFeatureImplementations<IPropertyAccessBlueprintBinding>("PropertyAccessBlueprintBinding"))
 		{
 			TSharedPtr<FExtender> BindingExtender = Binding->MakeBindingMenuExtender(BindingContext, MenuArgs);
