Move anim node pin binding to pin context menu

Binding workflow is now common with control rig
Bindings still show up on pins when bound, but the binding combobox is not displayed until bound

diff --git a/Engine/Source/Editor/AnimGraph/Private/AnimGraphNode_Base.cpp b/Engine/Source/Editor/AnimGraph/Private/AnimGraphNode_Base.cpp
--- a/Engine/Source/Editor/AnimGraph/Private/AnimGraphNode_Base.cpp
+++ b/Engine/Source/Editor/AnimGraph/Private/AnimGraphNode_Base.cpp
@@ -587,6 +587,30 @@ bool UAnimGraphNode_Base::IsPinExposedAndBound(const FString& InPinName, const E
 	return Pin != nullptr && Pin->LinkedTo.Num() == 0 && PropertyBindings.Find(Pin->GetFName()) != nullptr;
 }
 
+bool UAnimGraphNode_Base::IsPinBindable(const UEdGraphPin* InPin) const
+{
+	if(const FProperty* PinProperty = GetPinProperty(InPin))
+	{
+		const int32 OptionalPinIndex = ShowPinForProperties.IndexOfByPredicate([PinProperty](const FOptionalPinFromProperty& InOptionalPin)
+		{
+			return PinProperty->GetFName() == InOptionalPin.PropertyName;
+		});
+
+		return OptionalPinIndex != INDEX_NONE;
+	}
+
+	return false;
+}
+
+FProperty* UAnimGraphNode_Base::GetPinProperty(const UEdGraphPin* InPin) const
+{
+	// Compare FName without number to make sure we catch array properties that are split into multiple pins
+	FName ComparisonName = InPin->GetFName();
+	ComparisonName.SetNumber(0);
+	
+	return GetFNodeType()->FindPropertyByName(ComparisonName);
+}
+
 void UAnimGraphNode_Base::PinConnectionListChanged(UEdGraphPin* Pin)
 {
 	if(Pin->LinkedTo.Num() > 0)
diff --git a/Engine/Source/Editor/AnimGraph/Private/AnimationGraphSchema.cpp b/Engine/Source/Editor/AnimGraph/Private/AnimationGraphSchema.cpp
--- a/Engine/Source/Editor/AnimGraph/Private/AnimationGraphSchema.cpp
+++ b/Engine/Source/Editor/AnimGraph/Private/AnimationGraphSchema.cpp
@@ -602,6 +602,22 @@ void UAnimationGraphSchema::GetContextMenuActions(UToolMenu* Menu, UGraphNodeCon
 			FToolMenuSection& Section = Menu->AddSection("AnimGraphSchemaNodeActions", LOCTEXT("AnimNodeActionsMenuHeader", "Anim Node Actions"));
 			Section.AddMenuEntry(FAnimGraphCommands::Get().TogglePoseWatch);
 		}
+
+		if(Context->Pin && !IsPosePin(Context->Pin->PinType) && AnimGraphNode->IsPinBindable(Context->Pin))
+		{
+			FToolMenuSection& Section = Menu->AddSection("EdGraphSchemaPinActions");
+			
+			FProperty* PinProperty = AnimGraphNode->GetPinProperty(Context->Pin);
+			check(PinProperty);
+
+			const int32 OptionalPinIndex = AnimGraphNode->ShowPinForProperties.IndexOfByPredicate([PinProperty](const FOptionalPinFromProperty& InOptionalPin)
+			{
+				return PinProperty->GetFName() == InOptionalPin.PropertyName;
+			});
+			
+			const UAnimGraphNode_Base::FAnimPropertyBindingWidgetArgs Args({ const_cast<UAnimGraphNode_Base*>(AnimGraphNode) }, PinProperty, Context->Pin->GetFName(), OptionalPinIndex);
+			Section.AddEntry(FToolMenuEntry::InitWidget("BindingWidget", UAnimGraphNode_Base::MakePropertyBindingWidget(Args), LOCTEXT("BindingWidgetLabel", "Binding"), true));
+		}
 	}
 }
 
diff --git a/Engine/Source/Editor/AnimGraph/Public/AnimGraphNode_Base.h b/Engine/Source/Editor/AnimGraph/Public/AnimGraphNode_Base.h
--- a/Engine/Source/Editor/AnimGraph/Public/AnimGraphNode_Base.h
+++ b/Engine/Source/Editor/AnimGraph/Public/AnimGraphNode_Base.h
@@ -428,6 +428,12 @@ class ANIMGRAPH_API UAnimGraphNode_Base : public UK2Node
 	
 	// Make a property binding widget to edit the bindings of the passed-in nodes
 	static TSharedRef<SWidget> MakePropertyBindingWidget(const FAnimPropertyBindingWidgetArgs& InArgs);
+
+	// Get the property corresponding to a pin. For array element pins returns the outer array property. Returns null if a property cannot be found.
+	FProperty* GetPinProperty(const UEdGraphPin* InPin) const;
+	
+	// Check whether the named pin is bindable
+	bool IsPinBindable(const UEdGraphPin* InPin) const;
 	
 protected:
 	friend class FAnimBlueprintCompilerContext;
diff --git a/Engine/Source/Editor/AnimationBlueprintEditor/Private/AnimationNodes/SAnimationGraphNode.cpp b/Engine/Source/Editor/AnimationBlueprintEditor/Private/AnimationNodes/SAnimationGraphNode.cpp
--- a/Engine/Source/Editor/AnimationBlueprintEditor/Private/AnimationNodes/SAnimationGraphNode.cpp
+++ b/Engine/Source/Editor/AnimationBlueprintEditor/Private/AnimationNodes/SAnimationGraphNode.cpp
@@ -327,7 +327,19 @@ void SAnimationGraphNode::ReconfigurePinWidgetsForPropertyBindings(UAnimGraphNod
 					// Add binding widget
 					PinWidget->GetLabelAndValue()->AddSlot()
 					[
-						UAnimGraphNode_Base::MakePropertyBindingWidget(Args)
+						SNew(SBox)
+						.Visibility_Lambda([PinName, InAnimGraphNode]()
+						{
+							if (FAnimGraphNodePropertyBinding* BindingPtr = InAnimGraphNode->PropertyBindings.Find(PinName))
+							{
+								return EVisibility::Visible;
+							}
+
+							return EVisibility::Collapsed;
+						})
+						[
+							UAnimGraphNode_Base::MakePropertyBindingWidget(Args)
+						]
 					];
 				}
 			}
