Unify titles of asset players to use a common format

Use correct title for asset players when asset is bound to property access

diff --git a/Engine/Source/Editor/AnimGraph/Private/AnimGraphNode_AssetPlayerBase.cpp b/Engine/Source/Editor/AnimGraph/Private/AnimGraphNode_AssetPlayerBase.cpp
--- a/Engine/Source/Editor/AnimGraph/Private/AnimGraphNode_AssetPlayerBase.cpp
+++ b/Engine/Source/Editor/AnimGraph/Private/AnimGraphNode_AssetPlayerBase.cpp
@@ -317,4 +317,117 @@ bool UAnimGraphNode_AssetPlayerBase::IsActionFilteredOut(class FBlueprintActionF
 	return bIsFilteredOut;
 }
 
+FText UAnimGraphNode_AssetPlayerBase::GetNodeTitleHelper(ENodeTitleType::Type InTitleType, UEdGraphPin* InAssetPin, const FText& InAssetDesc, const TFunction<FText(UAnimationAsset*)> InPostFixFunctionRef) const
+{
+	UAnimationAsset* Asset = GetAnimationAsset();
+	if (Asset == nullptr)
+	{
+		// Check for bindings
+		bool bHasBinding = false;
+		if(InAssetPin != nullptr)
+		{
+			if (PropertyBindings.Find(InAssetPin->GetFName()) != nullptr)
+			{
+				bHasBinding = true;
+			}
+		}
+
+		// Also check for links
+		if (bHasBinding || (InAssetPin && InAssetPin->LinkedTo.Num() > 0))
+		{
+			return InAssetDesc;
+		}
+		// check for a default value on the pin
+		else if (InAssetPin && InAssetPin->DefaultObject != nullptr)
+		{
+			return GetNodeTitleForAsset(InTitleType, CastChecked<UAnimationAsset>(InAssetPin->DefaultObject), InAssetDesc, InPostFixFunctionRef);
+		}
+		else
+		{
+			return InAssetDesc;
+		}
+	}
+	else
+	{
+		return GetNodeTitleForAsset(InTitleType, Asset, InAssetDesc, InPostFixFunctionRef);
+	}
+}
+
+FText UAnimGraphNode_AssetPlayerBase::GetNodeTitleForAsset(ENodeTitleType::Type InTitleType, UAnimationAsset* InAsset, const FText& InAssetDesc, const TFunction<FText(UAnimationAsset*)> InPostFixFunctionRef) const
+{
+	UAnimationAsset* Asset = GetAnimationAsset();
+	check(Asset);
+	const FText AssetName = FText::FromString(Asset->GetName());
+
+	if (InTitleType == ENodeTitleType::ListView || InTitleType == ENodeTitleType::MenuTitle)
+	{
+		FFormatNamedArguments Args;
+		Args.Add(TEXT("AssetName"), AssetName);
+		Args.Add(TEXT("AssetDesc"), InAssetDesc);
+
+		if(InPostFixFunctionRef != nullptr)
+		{
+			const FText PostFix = InPostFixFunctionRef(InAsset);
+			if(!PostFix.IsEmpty())
+			{
+				Args.Add(TEXT("PostFix"), PostFix);
+				static const FTextFormat FormatWithPostFix(LOCTEXT("AssetPlayerTitlewithPostFix", "{AssetDesc} {PostFix} '{AssetName}'"));
+				return FText::Format(FormatWithPostFix, Args);
+			}
+		}
+
+		static const FTextFormat FormatWithoutPostFix(LOCTEXT("AssetPlayerTitle", "{AssetDesc} '{AssetName}'"));
+		return FText::Format(FormatWithoutPostFix, Args);
+	}
+	else
+	{
+		FFormatNamedArguments TitleArgs;
+		TitleArgs.Add(TEXT("AssetName"), AssetName);
+		TitleArgs.Add(TEXT("AssetDesc"), InAssetDesc);
+		FText Title = FText::Format(LOCTEXT("AssetPlayerFullTitle", "{AssetName}\n{AssetDesc}"), TitleArgs);
+
+		if (InTitleType == ENodeTitleType::FullTitle)
+		{
+			FStructProperty* NodeProperty = GetFNodeProperty();
+			if(NodeProperty->Struct->IsChildOf(FAnimNode_AssetPlayerBase::StaticStruct()))
+			{
+				const FAnimNode_AssetPlayerBase* Node = NodeProperty->ContainerPtrToValuePtr<FAnimNode_AssetPlayerBase>(this);
+				
+				FFormatNamedArguments Args;
+				Args.Add(TEXT("Title"), Title);
+
+				if(Node->GetGroupMethod() == EAnimSyncMethod::SyncGroup)
+				{
+					Args.Add(TEXT("SyncGroupName"), FText::FromName(Node->GetGroupName()));
+					static const FTextFormat FormatAssetPlayerNodeSyncGroupSubtitle(LOCTEXT("AssetPlayerNodeSyncGroupSubtitle", "{Title}\nSync group {SyncGroupName}"));
+					Title = FText::Format(FormatAssetPlayerNodeSyncGroupSubtitle, Args);
+				}
+				else if(Node->GetGroupMethod() == EAnimSyncMethod::Graph)
+				{
+					static const FTextFormat FormatAssetPlayerNodeGraphSyncGroupSubtitle(LOCTEXT("AssetPlayerNodeGraphSyncGroupSubtitle", "{Title}\nGraph sync group"));
+					Title = FText::Format(FormatAssetPlayerNodeGraphSyncGroupSubtitle, Args);
+
+					UObject* ObjectBeingDebugged = GetAnimBlueprint()->GetObjectBeingDebugged();
+					UAnimBlueprintGeneratedClass* GeneratedClass = GetAnimBlueprint()->GetAnimBlueprintGeneratedClass();
+					if (ObjectBeingDebugged && GeneratedClass)
+					{
+						int32 NodeIndex = GeneratedClass->GetNodeIndexFromGuid(NodeGuid);
+						if(NodeIndex != INDEX_NONE)
+						{
+							if(const FName* SyncGroupNamePtr = GeneratedClass->GetAnimBlueprintDebugData().NodeSyncsThisFrame.Find(NodeIndex))
+							{
+								Args.Add(TEXT("SyncGroupName"), FText::FromName(*SyncGroupNamePtr));
+								static const FTextFormat FormatAssetPlayerNodeDynamicGraphSyncGroupSubtitle(LOCTEXT("AssetPlayerNodeDynamicGraphSyncGroupSubtitle", "{Title}\nGraph sync group {SyncGroupName}"));
+								Title = FText::Format(FormatAssetPlayerNodeDynamicGraphSyncGroupSubtitle, Args);
+							}
+						}
+					}
+				}
+			}
+		}
+
+		return Title;
+	}
+}
+
 #undef LOCTEXT_NAMESPACE
\ No newline at end of file
diff --git a/Engine/Source/Editor/AnimGraph/Private/AnimGraphNode_BlendSpaceEvaluator.cpp b/Engine/Source/Editor/AnimGraph/Private/AnimGraphNode_BlendSpaceEvaluator.cpp
--- a/Engine/Source/Editor/AnimGraph/Private/AnimGraphNode_BlendSpaceEvaluator.cpp
+++ b/Engine/Source/Editor/AnimGraph/Private/AnimGraphNode_BlendSpaceEvaluator.cpp
@@ -20,70 +20,10 @@ UAnimGraphNode_BlendSpaceEvaluator::UAnimGraphNode_BlendSpaceEvaluator(const FOb
 {
 }
 
-FText UAnimGraphNode_BlendSpaceEvaluator::GetNodeTitleForBlendSpace(ENodeTitleType::Type TitleType, UBlendSpace* InBlendSpace) const
-{
-	const FText BlendSpaceName = FText::FromString(InBlendSpace->GetName());
-
-	if (TitleType == ENodeTitleType::ListView || TitleType == ENodeTitleType::MenuTitle)
-	{
-		FFormatNamedArguments Args;
-		Args.Add(TEXT("BlendSpaceName"), BlendSpaceName);
-
-		// FText::Format() is slow, so we cache this to save on performance
-		CachedNodeTitles.SetCachedTitle(TitleType, FText::Format(LOCTEXT("BlendSpaceEvaluatorListTitle", "Blendspace Evaluator '{BlendSpaceName}'"), Args), this);
-	}
-	else
-	{
-		FFormatNamedArguments TitleArgs;
-		TitleArgs.Add(TEXT("BlendSpaceName"), BlendSpaceName);
-		FText Title = FText::Format(LOCTEXT("BlendSpaceEvaluatorFullTitle", "{BlendSpaceName}\nBlendspace Evaluator"), TitleArgs);
-
-		if ((TitleType == ENodeTitleType::FullTitle) && (Node.GetGroupName() != NAME_None))
-		{
-			FFormatNamedArguments Args;
-			Args.Add(TEXT("Title"), Title);
-			Args.Add(TEXT("SyncGroupName"), FText::FromName(Node.GetGroupName()));
-			Title = FText::Format(LOCTEXT("BlendSpaceNodeGroupSubtitle", "{Title}\nSync group {SyncGroupName}"), Args);
-		}
-		// FText::Format() is slow, so we cache this to save on performance
-		CachedNodeTitles.SetCachedTitle(TitleType, Title, this);
-	}
-
-	return CachedNodeTitles[TitleType];
-}
-
 FText UAnimGraphNode_BlendSpaceEvaluator::GetNodeTitle(ENodeTitleType::Type TitleType) const
 {
-	if (Node.GetBlendSpace() == nullptr)
-	{
-		// we may have a valid variable connected or default pin value
-		UEdGraphPin* BlendSpacePin = FindPin(GET_MEMBER_NAME_STRING_CHECKED(FAnimNode_BlendSpacePlayer, BlendSpace));
-		if (BlendSpacePin && BlendSpacePin->LinkedTo.Num() > 0)
-		{
-			return LOCTEXT("BlendSpaceEvaluator_Variable_Title", "Blendspace Evaluator");
-		}
-		else if (BlendSpacePin && BlendSpacePin->DefaultObject != nullptr)
-		{
-			return GetNodeTitleForBlendSpace(TitleType, CastChecked<UBlendSpace>(BlendSpacePin->DefaultObject));
-		}
-		else
-		{
-			if (TitleType == ENodeTitleType::ListView || TitleType == ENodeTitleType::MenuTitle)
-			{
-				return LOCTEXT("BlendSpaceEvaluator_NONE_ListTitle", "Blendspace Evaluator '(None)'");
-			}
-			else
-			{
-				return LOCTEXT("BlendSpaceEvaluator_NONE_Title", "(None)\nBlendspace Evaluator");
-			}
-		}
-	}
-	// @TODO: the bone can be altered in the property editor, so we have to 
-	//        choose to mark this dirty when that happens for this to properly work
-	else //if (!CachedNodeTitles.IsTitleCached(TitleType, this))
-	{
-		return GetNodeTitleForBlendSpace(TitleType, Node.GetBlendSpace());
-	}
+	UEdGraphPin* BlendSpacePin = FindPin(GET_MEMBER_NAME_STRING_CHECKED(FAnimNode_BlendSpaceEvaluator, BlendSpace));
+	return GetNodeTitleHelper(TitleType, BlendSpacePin, LOCTEXT("PlayerDesc", "Blendspace Evaluator"));
 }
 
 void UAnimGraphNode_BlendSpaceEvaluator::GetMenuActions(FBlueprintActionDatabaseRegistrar& InActionRegistrar) const
diff --git a/Engine/Source/Editor/AnimGraph/Private/AnimGraphNode_BlendSpacePlayer.cpp b/Engine/Source/Editor/AnimGraph/Private/AnimGraphNode_BlendSpacePlayer.cpp
--- a/Engine/Source/Editor/AnimGraph/Private/AnimGraphNode_BlendSpacePlayer.cpp
+++ b/Engine/Source/Editor/AnimGraph/Private/AnimGraphNode_BlendSpacePlayer.cpp
@@ -25,92 +25,10 @@ UAnimGraphNode_BlendSpacePlayer::UAnimGraphNode_BlendSpacePlayer(const FObjectIn
 {
 }
 
-FText UAnimGraphNode_BlendSpacePlayer::GetNodeTitleForBlendSpace(ENodeTitleType::Type TitleType, UBlendSpace* InBlendSpace) const
-{
-	const FText BlendSpaceName = FText::FromString(InBlendSpace->GetName());
-
-	if (TitleType == ENodeTitleType::ListView || TitleType == ENodeTitleType::MenuTitle)
-	{
-		FFormatNamedArguments Args;
-		Args.Add(TEXT("BlendSpaceName"), BlendSpaceName);
-		// FText::Format() is slow, so we cache this to save on performance
-		CachedNodeTitles.SetCachedTitle(TitleType, FText::Format(LOCTEXT("BlendspacePlayer", "Blendspace Player '{BlendSpaceName}'"), Args), this);
-	}
-	else
-	{
-		FFormatNamedArguments TitleArgs;
-		TitleArgs.Add(TEXT("BlendSpaceName"), BlendSpaceName);
-		FText Title = FText::Format(LOCTEXT("BlendSpacePlayerFullTitle", "{BlendSpaceName}\nBlendspace Player"), TitleArgs);
-
-		if (TitleType == ENodeTitleType::FullTitle)
-		{
-			FFormatNamedArguments Args;
-			Args.Add(TEXT("Title"), Title);
-
-			if(Node.GetGroupMethod() == EAnimSyncMethod::SyncGroup)
-			{
-				Args.Add(TEXT("SyncGroupName"), FText::FromName(Node.GetGroupName()));
-				Title = FText::Format(LOCTEXT("BlendSpaceNodeGroupSubtitle", "{Title}\nSync group {SyncGroupName}"), Args);
-			}
-			else if(Node.GetGroupMethod() == EAnimSyncMethod::Graph)
-			{
-				Title = FText::Format(LOCTEXT("BlendSpaceNodeGroupSubtitle", "{Title}\nGraph sync group"), Args);
-
-				UObject* ObjectBeingDebugged = GetAnimBlueprint()->GetObjectBeingDebugged();
-				UAnimBlueprintGeneratedClass* GeneratedClass = GetAnimBlueprint()->GetAnimBlueprintGeneratedClass();
-				if (ObjectBeingDebugged && GeneratedClass)
-				{
-					int32 NodeIndex = GeneratedClass->GetNodeIndexFromGuid(NodeGuid);
-					if(NodeIndex != INDEX_NONE)
-					{
-						if(const FName* SyncGroupNamePtr = GeneratedClass->GetAnimBlueprintDebugData().NodeSyncsThisFrame.Find(NodeIndex))
-						{
-							Args.Add(TEXT("SyncGroupName"), FText::FromName(*SyncGroupNamePtr));
-							Title = FText::Format(LOCTEXT("BlendSpaceNodeGraphGroupSubtitle", "{Title}\nGraph sync group {SyncGroupName}"), Args);
-						}
-					}
-				}
-			}
-		}
-		// FText::Format() is slow, so we cache this to save on performance
-		CachedNodeTitles.SetCachedTitle(TitleType, Title, this);
-	}
-
-	return CachedNodeTitles[TitleType];
-}
-
 FText UAnimGraphNode_BlendSpacePlayer::GetNodeTitle(ENodeTitleType::Type TitleType) const
-{	
-	if (Node.GetBlendSpace() == nullptr)
-	{
-		// we may have a valid variable connected or default pin value
-		UEdGraphPin* BlendSpacePin = FindPin(GET_MEMBER_NAME_STRING_CHECKED(FAnimNode_BlendSpacePlayer, BlendSpace));
-		if (BlendSpacePin && BlendSpacePin->LinkedTo.Num() > 0)
-		{
-			return LOCTEXT("BlendspacePlayer_Variable_Title", "Blendspace Player");
-		}
-		else if (BlendSpacePin && BlendSpacePin->DefaultObject != nullptr)
-		{
-			return GetNodeTitleForBlendSpace(TitleType, CastChecked<UBlendSpace>(BlendSpacePin->DefaultObject));
-		}
-		else
-		{
-			if (TitleType == ENodeTitleType::ListView || TitleType == ENodeTitleType::MenuTitle)
-			{
-				return LOCTEXT("BlendspacePlayer_NONE_ListTitle", "Blendspace Player '(None)'");
-			}
-			else
-			{
-				return LOCTEXT("BlendspacePlayer_NONE_Title", "(None)\nBlendspace Player");
-			}
-		}
-	}
-	// @TODO: the bone can be altered in the property editor, so we have to 
-	//        choose to mark this dirty when that happens for this to properly work
-	else //if (!CachedNodeTitles.IsTitleCached(TitleType, this))
-	{
-		return GetNodeTitleForBlendSpace(TitleType, Node.GetBlendSpace());
-	}
+{
+	UEdGraphPin* BlendSpacePin = FindPin(GET_MEMBER_NAME_STRING_CHECKED(FAnimNode_BlendSpacePlayer, BlendSpace));
+	return GetNodeTitleHelper(TitleType, BlendSpacePin, LOCTEXT("PlayerDesc", "Blendspace Player"));
 }
 
 void UAnimGraphNode_BlendSpacePlayer::ValidateAnimNodeDuringCompilation(class USkeleton* ForSkeleton, class FCompilerResultsLog& MessageLog)
diff --git a/Engine/Source/Editor/AnimGraph/Private/AnimGraphNode_SequenceEvaluator.cpp b/Engine/Source/Editor/AnimGraph/Private/AnimGraphNode_SequenceEvaluator.cpp
--- a/Engine/Source/Editor/AnimGraph/Private/AnimGraphNode_SequenceEvaluator.cpp
+++ b/Engine/Source/Editor/AnimGraph/Private/AnimGraphNode_SequenceEvaluator.cpp
@@ -57,55 +57,10 @@ FText UAnimGraphNode_SequenceEvaluator::GetMenuCategory() const
 	return FEditorCategoryUtils::GetCommonCategory(FCommonEditorCategory::Animation);
 }
 
-FText UAnimGraphNode_SequenceEvaluator::GetNodeTitleForSequence(ENodeTitleType::Type TitleType, UAnimSequenceBase* InSequence) const
-{
-	const FText SequenceName = FText::FromString(InSequence->GetName());
-
-	FFormatNamedArguments Args;
-	Args.Add(TEXT("SequenceName"), SequenceName);
-
-	// FText::Format() is slow, so we cache this to save on performance
-	if (InSequence->IsValidAdditive())
-	{
-		CachedNodeTitle.SetCachedText(FText::Format(LOCTEXT("EvaluateSequence_Additive", "Evaluate {SequenceName} (additive)"), Args), this);
-	}
-	else
-	{
-		CachedNodeTitle.SetCachedText(FText::Format(LOCTEXT("EvaluateSequence", "Evaluate {SequenceName}"), Args), this);
-	}
-
-	return CachedNodeTitle;
-}
-
 FText UAnimGraphNode_SequenceEvaluator::GetNodeTitle(ENodeTitleType::Type TitleType) const
 {
-	UAnimSequenceBase* Sequence = Node.GetSequence();
-	if (Sequence == nullptr)
-	{
-		// we may have a valid variable connected or default pin value
-		UEdGraphPin* SequencePin = FindPin(GET_MEMBER_NAME_STRING_CHECKED(FAnimNode_SequenceEvaluator, Sequence));
-		if (SequencePin && SequencePin->LinkedTo.Num() > 0)
-		{
-			return LOCTEXT("EvaluateSequence_TitleVariable", "Evaluate Animation Sequence");
-		}
-		else if (SequencePin && SequencePin->DefaultObject != nullptr)
-		{
-			return GetNodeTitleForSequence(TitleType, CastChecked<UAnimSequenceBase>(SequencePin->DefaultObject));
-		}
-		else
-		{
-			return LOCTEXT("EvaluateSequence_TitleNONE", "Evaluate (None)");
-		}
-	}
-	// @TODO: don't know enough about this node type to comfortably assert that
-	//        the CacheName won't change after the node has spawned... until
-	//        then, we'll leave this optimization off
-	else //if (CachedNodeTitle.IsOutOfDate(this))
-	{
-		GetNodeTitleForSequence(TitleType, Sequence);
-	}
-
-	return CachedNodeTitle;
+	UEdGraphPin* SequencePin = FindPin(GET_MEMBER_NAME_STRING_CHECKED(FAnimNode_SequenceEvaluator, Sequence));
+	return GetNodeTitleHelper(TitleType, SequencePin, LOCTEXT("PlayerDesc", "Sequence Evaluator"));
 }
 
 FSlateIcon UAnimGraphNode_SequenceEvaluator::GetIconAndTint(FLinearColor& OutColor) const
diff --git a/Engine/Source/Editor/AnimGraph/Private/AnimGraphNode_SequencePlayer.cpp b/Engine/Source/Editor/AnimGraph/Private/AnimGraphNode_SequencePlayer.cpp
--- a/Engine/Source/Editor/AnimGraph/Private/AnimGraphNode_SequencePlayer.cpp
+++ b/Engine/Source/Editor/AnimGraph/Private/AnimGraphNode_SequencePlayer.cpp
@@ -68,111 +68,16 @@ FSlateIcon UAnimGraphNode_SequencePlayer::GetIconAndTint(FLinearColor& OutColor)
 	return FSlateIcon("EditorStyle", "ClassIcon.AnimSequence");
 }
 
-FText UAnimGraphNode_SequencePlayer::GetNodeTitleForSequence(ENodeTitleType::Type TitleType, UAnimSequenceBase* InSequence) const
-{
-	const bool bAdditive = InSequence->IsValidAdditive();
-	const FText BasicTitle = GetTitleGivenAssetInfo(FText::FromName(InSequence->GetFName()), bAdditive);
-
-	if(Node.GetGroupMethod() == EAnimSyncMethod::SyncGroup)
-	{
-		const FText SyncGroupName = FText::FromName(Node.GetGroupName());
-
-		FFormatNamedArguments Args;
-		Args.Add(TEXT("Title"), BasicTitle);
-		Args.Add(TEXT("SyncGroup"), SyncGroupName);
-
-		if (TitleType == ENodeTitleType::FullTitle)
-		{
-			return FText::Format(LOCTEXT("SequenceNodeGroupWithSubtitleFull", "{Title}\nSync group {SyncGroup}"), Args);
-		}
-		else
-		{
-			return FText::Format(LOCTEXT("SequenceNodeGroupWithSubtitleList", "{Title} (Sync group {SyncGroup})"), Args);
-		}
-	}
-	else if(Node.GetGroupMethod() == EAnimSyncMethod::Graph)
-	{
-		FFormatNamedArguments Args;
-		Args.Add(TEXT("Title"), BasicTitle);
-
-		UObject* ObjectBeingDebugged = GetAnimBlueprint()->GetObjectBeingDebugged();
-		UAnimBlueprintGeneratedClass* GeneratedClass = GetAnimBlueprint()->GetAnimBlueprintGeneratedClass();
-		if (ObjectBeingDebugged && GeneratedClass)
-		{
-			int32 NodeIndex = GeneratedClass->GetNodeIndexFromGuid(NodeGuid);
-			if(NodeIndex != INDEX_NONE)
-			{
-				if(const FName* SyncGroupNamePtr = GeneratedClass->GetAnimBlueprintDebugData().NodeSyncsThisFrame.Find(NodeIndex))
-				{
-					Args.Add(TEXT("SyncGroup"), FText::FromName(*SyncGroupNamePtr));
-
-					if (TitleType == ENodeTitleType::FullTitle)
-					{
-						return FText::Format(LOCTEXT("SequenceNodeGraphSyncDebuggedWithSubtitleFull", "{Title}\nGraph sync group {SyncGroup}"), Args);
-					}
-					else
-					{
-						return FText::Format(LOCTEXT("SequenceNodeGraphSyncDebuggedWithSubtitleList", "{Title} (Graph sync group {SyncGroup})"), Args);
-					}
-				}
-			}
-		}
-
-		// Indicate that the sync is graph-based if we have no debug data
-		if (TitleType == ENodeTitleType::FullTitle)
-		{
-			return FText::Format(LOCTEXT("SequenceNodeGraphSyncWithSubtitleFull", "{Title}\nGraph sync group"), Args);
-		}
-		else
-		{
-			return FText::Format(LOCTEXT("SequenceNodeGraphSyncWithSubtitleList", "{Title} (Graph sync group)"), Args);
-		}
-	}
-	else
-	{
-		return BasicTitle;
-	}
-}
-
 FText UAnimGraphNode_SequencePlayer::GetNodeTitle(ENodeTitleType::Type TitleType) const
 {
-	UAnimSequenceBase* Sequence = Node.GetSequence();
-	if (Sequence == nullptr)
-	{
-		// we may have a valid variable connected or default pin value
-		UEdGraphPin* SequencePin = FindPin(GET_MEMBER_NAME_STRING_CHECKED(FAnimNode_SequencePlayer, Sequence));
-		if (SequencePin && SequencePin->LinkedTo.Num() > 0)
-		{
-			return LOCTEXT("SequenceNodeTitleVariable", "Play Animation Sequence");
-		}
-		else if (SequencePin && SequencePin->DefaultObject != nullptr)
-		{
-			return GetNodeTitleForSequence(TitleType, CastChecked<UAnimSequenceBase>(SequencePin->DefaultObject));
-		}
-		else
+	UEdGraphPin* SequencePin = FindPin(GET_MEMBER_NAME_STRING_CHECKED(FAnimNode_SequencePlayer, Sequence));
+	return GetNodeTitleHelper(TitleType, SequencePin, LOCTEXT("PlayerDesc", "Sequence Player"),
+		[](UAnimationAsset* InAsset)
 		{
-			return LOCTEXT("SequenceNullTitle", "Play (None)");
-		}
-	}
-	else
-	{
-		return GetNodeTitleForSequence(TitleType, Sequence);
-	}
-}
-
-FText UAnimGraphNode_SequencePlayer::GetTitleGivenAssetInfo(const FText& AssetName, bool bKnownToBeAdditive)
-{
-	FFormatNamedArguments Args;
-	Args.Add(TEXT("AssetName"), AssetName);
-
-	if (bKnownToBeAdditive)
-	{
-		return FText::Format(LOCTEXT("SequenceNodeTitleAdditive", "Play {AssetName} (additive)"), Args);
-	}
-	else
-	{
-		return FText::Format(LOCTEXT("SequenceNodeTitle", "Play {AssetName}"), Args);
-	}
+			UAnimSequenceBase* SequenceBase = CastChecked<UAnimSequenceBase>(InAsset);
+			const bool bAdditive = SequenceBase->IsValidAdditive();
+			return bAdditive ? LOCTEXT("AdditivePostFix", "(additive)") : FText::GetEmpty();
+		});
 }
 
 FText UAnimGraphNode_SequencePlayer::GetMenuCategory() const
diff --git a/Engine/Source/Editor/AnimGraph/Public/AnimGraphNode_AssetPlayerBase.h b/Engine/Source/Editor/AnimGraph/Public/AnimGraphNode_AssetPlayerBase.h
--- a/Engine/Source/Editor/AnimGraph/Public/AnimGraphNode_AssetPlayerBase.h
+++ b/Engine/Source/Editor/AnimGraph/Public/AnimGraphNode_AssetPlayerBase.h
@@ -6,6 +6,7 @@
 #include "UObject/ObjectMacros.h"
 #include "Animation/AnimationAsset.h"
 #include "AnimGraphNode_Base.h"
+#include "EdGraph/EdGraphNodeUtils.h"
 #include "AnimGraphNode_AssetPlayerBase.generated.h"
 
 /** Get the default anim node class for playing a particular asset */
@@ -44,8 +45,12 @@ public:
 
 	// Helper function to gather menu actions from specific asset types supported by this node
 	static void GetMenuActionsHelper(FBlueprintActionDatabaseRegistrar& InActionRegistrar, TSubclassOf<UAnimGraphNode_Base> InNodeClass, const TArray<TSubclassOf<UObject>>& InAssetTypes, const TArray<TSubclassOf<UObject>>& InExcludedAssetTypes, const TFunctionRef<FText(const FAssetData&)>& InMenuNameFunction, const TFunctionRef<FText(const FAssetData&)>& InMenuTooltipFunction, const TFunction<void(UEdGraphNode*, bool, const FAssetData)>& InSetupNewNodeFunction);
-	
+
 protected:
+	// Helper functions to build a title for an asset player node
+	FText GetNodeTitleHelper(ENodeTitleType::Type InTitleType, UEdGraphPin* InAssetPin, const FText& InAssetDesc, const TFunction<FText(UAnimationAsset*)> InPostFixFunctionRef = nullptr) const;
+	FText GetNodeTitleForAsset(ENodeTitleType::Type InTitleType, UAnimationAsset* InAsset, const FText& InAssetDesc, const TFunction<FText(UAnimationAsset*)> InPostFixFunctionRef = nullptr) const;
+	
 	// Default setup function that can be used with GetMenuActionsHalper
 	static void SetupNewNode(UEdGraphNode* InNewNode, bool bInIsTemplateNode, const FAssetData InAssetData);
 
diff --git a/Engine/Source/Editor/AnimGraph/Public/AnimGraphNode_BlendSpacePlayer.h b/Engine/Source/Editor/AnimGraph/Public/AnimGraphNode_BlendSpacePlayer.h
--- a/Engine/Source/Editor/AnimGraph/Public/AnimGraphNode_BlendSpacePlayer.h
+++ b/Engine/Source/Editor/AnimGraph/Public/AnimGraphNode_BlendSpacePlayer.h
@@ -46,12 +46,4 @@ class UAnimGraphNode_BlendSpacePlayer : public UAnimGraphNode_BlendSpaceBase
 	// UAnimGraphNode_AssetPlayerBase interface
 	virtual void SetAnimationAsset(UAnimationAsset* Asset) override;
 	// End of UAnimGraphNode_AssetPlayerBase interface
-
-private:
-	/** Helper function for GetNodeTitle */
-	FText GetNodeTitleForBlendSpace(ENodeTitleType::Type TitleType, UBlendSpace* InBlendSpace) const;
-
-private:
-	/** Constructing FText strings can be costly, so we cache the node's title */
-	FNodeTitleTextTable CachedNodeTitles;
 };
diff --git a/Engine/Source/Editor/AnimGraph/Public/AnimGraphNode_SequenceEvaluator.h b/Engine/Source/Editor/AnimGraph/Public/AnimGraphNode_SequenceEvaluator.h
--- a/Engine/Source/Editor/AnimGraph/Public/AnimGraphNode_SequenceEvaluator.h
+++ b/Engine/Source/Editor/AnimGraph/Public/AnimGraphNode_SequenceEvaluator.h
@@ -48,11 +48,4 @@ class UAnimGraphNode_SequenceEvaluator : public UAnimGraphNode_AssetPlayerBase
 	// UAnimGraphNode_AssetPlayerBase interface
 	virtual void SetAnimationAsset(UAnimationAsset* Asset) override;
 	// End of UAnimGraphNode_AssetPlayerBase interface
-
-private:
-	FText GetNodeTitleForSequence(ENodeTitleType::Type TitleType, UAnimSequenceBase* InSequence) const;
-
-private:
-	/** Constructing FText strings can be costly, so we cache the node's title */
-	FNodeTextCache CachedNodeTitle;
 };
diff --git a/Engine/Source/Editor/AnimGraph/Public/AnimGraphNode_SequencePlayer.h b/Engine/Source/Editor/AnimGraph/Public/AnimGraphNode_SequencePlayer.h
--- a/Engine/Source/Editor/AnimGraph/Public/AnimGraphNode_SequencePlayer.h
+++ b/Engine/Source/Editor/AnimGraph/Public/AnimGraphNode_SequencePlayer.h
@@ -56,11 +56,4 @@ class UAnimGraphNode_SequencePlayer : public UAnimGraphNode_AssetPlayerBase
 	// UAnimGraphNode_AssetPlayerBase interface
 	virtual void SetAnimationAsset(UAnimationAsset* Asset) override;
 	// End of UAnimGraphNode_AssetPlayerBase interface
-
-private:
-	static FText GetTitleGivenAssetInfo(const FText& AssetName, bool bKnownToBeAdditive);
-	FText GetNodeTitleForSequence(ENodeTitleType::Type TitleType, UAnimSequenceBase* InSequence) const;
-
-	/** Constructing FText strings can be costly, so we cache the node's title */
-	FNodeTitleTextTable CachedNodeTitles;
 };
