Fixed a bug where state naming was not validating against sub-graphs

diff --git a/Engine/Source/Editor/AnimGraph/Private/AnimGraphNode_StateMachineBase.cpp b/Engine/Source/Editor/AnimGraph/Private/AnimGraphNode_StateMachineBase.cpp
--- a/Engine/Source/Editor/AnimGraph/Private/AnimGraphNode_StateMachineBase.cpp
+++ b/Engine/Source/Editor/AnimGraph/Private/AnimGraphNode_StateMachineBase.cpp
@@ -52,6 +52,25 @@ public:
 				Names.Add(Node->GetStateMachineName());
 			}
 		}
+
+		// Include the name of animation layers
+		UBlueprint* Blueprint = FBlueprintEditorUtils::FindBlueprintForGraphChecked(StateMachine);
+
+		if (Blueprint)
+		{
+			UClass* TargetClass = *Blueprint->SkeletonGeneratedClass;
+			if (TargetClass)
+			{
+				IAnimClassInterface* AnimClassInterface = IAnimClassInterface::GetFromClass(TargetClass);
+				for (const FAnimBlueprintFunction& AnimBlueprintFunction : AnimClassInterface->GetAnimBlueprintFunctions())
+				{
+					if (AnimBlueprintFunction.Name != UEdGraphSchema_K2::GN_AnimGraph)
+					{
+						Names.Add(AnimBlueprintFunction.Name.ToString());
+					}
+				}
+			}
+		}
 	}
 };
 
diff --git a/Engine/Source/Editor/AnimGraph/Private/AnimStateNodeBase.cpp b/Engine/Source/Editor/AnimGraph/Private/AnimStateNodeBase.cpp
--- a/Engine/Source/Editor/AnimGraph/Private/AnimStateNodeBase.cpp
+++ b/Engine/Source/Editor/AnimGraph/Private/AnimStateNodeBase.cpp
@@ -34,6 +34,25 @@ public:
 				Names.Add(Node->GetStateName());
 			}
 		}
+
+		// Include the name of animation layers
+		UBlueprint* Blueprint = FBlueprintEditorUtils::FindBlueprintForGraphChecked(StateMachine);
+		
+		if (Blueprint)
+		{
+			UClass* TargetClass = *Blueprint->SkeletonGeneratedClass;
+			if (TargetClass)
+			{
+				IAnimClassInterface* AnimClassInterface = IAnimClassInterface::GetFromClass(TargetClass);
+				for (const FAnimBlueprintFunction& AnimBlueprintFunction : AnimClassInterface->GetAnimBlueprintFunctions())
+				{
+					if (AnimBlueprintFunction.Name != UEdGraphSchema_K2::GN_AnimGraph)
+					{
+						Names.Add(AnimBlueprintFunction.Name.ToString());
+					}
+				}
+			}
+		}
 	}
 };
 
