Fix Non Unity compile errors caused by changes to AnimStateNodeBase

diff --git a/Engine/Source/Editor/AnimGraph/Private/AnimStateNodeBase.cpp b/Engine/Source/Editor/AnimGraph/Private/AnimStateNodeBase.cpp
--- a/Engine/Source/Editor/AnimGraph/Private/AnimStateNodeBase.cpp
+++ b/Engine/Source/Editor/AnimGraph/Private/AnimStateNodeBase.cpp
@@ -7,6 +7,7 @@
 #include "AnimStateNodeBase.h"
 #include "UObject/FrameworkObjectVersion.h"
 #include "Animation/AnimBlueprint.h"
+#include "Animation/AnimClassInterface.h"
 #include "Kismet2/BlueprintEditorUtils.h"
 #include "AnimationStateMachineGraph.h"
 #include "AnimationStateMachineSchema.h"
