Added the ability to drag layers to create a linked anim layer node

diff --git a/Engine/Source/Editor/AnimGraph/Private/AnimationGraphSchema.cpp b/Engine/Source/Editor/AnimGraph/Private/AnimationGraphSchema.cpp
--- a/Engine/Source/Editor/AnimGraph/Private/AnimationGraphSchema.cpp
+++ b/Engine/Source/Editor/AnimGraph/Private/AnimationGraphSchema.cpp
@@ -38,9 +38,112 @@
 #include "AnimGraphNode_LinkedAnimLayer.h"
 #include "AnimGraphNode_RigidBody.h"
 #include "AnimationBlendSpaceSampleGraph.h"
+#include "GraphEditorDragDropAction.h"
 
 #define LOCTEXT_NAMESPACE "AnimationGraphSchema"
 
+/////////////////////////////////////////////////////
+// FAnimationLayerDragDropAction
+/** DragDropAction class for drag and dropping animation layers */
+class ANIMGRAPH_API FAnimationLayerDragDropAction : public FGraphSchemaActionDragDropAction
+{
+public:
+	DRAG_DROP_OPERATOR_TYPE(FAnimationLayerDragDropAction, FGraphSchemaActionDragDropAction)
+
+	virtual FReply DroppedOnPanel(const TSharedRef< class SWidget >& Panel, FVector2D ScreenPosition, FVector2D GraphPosition, UEdGraph& Graph) override;
+	virtual FReply DroppedOnNode(FVector2D ScreenPosition, FVector2D GraphPosition) override;
+	virtual FReply DroppedOnPin(FVector2D ScreenPosition, FVector2D GraphPosition) override;
+	virtual FReply DroppedOnAction(TSharedRef<FEdGraphSchemaAction> Action) override;
+	virtual FReply DroppedOnCategory(FText Category) override;
+	virtual void HoverTargetChanged() override;
+
+protected:
+
+	/** Constructor */
+	FAnimationLayerDragDropAction();
+
+	static TSharedRef<FAnimationLayerDragDropAction> New(TSharedPtr<FEdGraphSchemaAction> InAction, FName InFuncName, UAnimBlueprint* InRigBlueprint, UAnimationGraph* InRigGraph);
+
+	UAnimBlueprint* SourceAnimBlueprint;
+	UAnimationGraph* SourceAnimLayerGraph;
+	FName SourceFuncName;
+
+	friend class UAnimationGraphSchema;
+};
+
+
+FAnimationLayerDragDropAction::FAnimationLayerDragDropAction()
+    : FGraphSchemaActionDragDropAction()
+    , SourceAnimBlueprint(nullptr)
+    , SourceAnimLayerGraph(nullptr)
+    , SourceFuncName(NAME_None)
+{
+}
+
+FReply FAnimationLayerDragDropAction::DroppedOnPanel(const TSharedRef< class SWidget >& Panel, FVector2D ScreenPosition, FVector2D GraphPosition, UEdGraph& Graph)
+{
+	if (UAnimationGraph* TargetRigGraph = Cast<UAnimationGraph>(&Graph))
+	{
+		if (UAnimBlueprint* TargetAnimBlueprint = Cast<UAnimBlueprint>(FBlueprintEditorUtils::FindBlueprintForGraph(TargetRigGraph)))
+		{
+			FGraphNodeCreator<UAnimGraphNode_LinkedAnimLayer> LinkedInputLayerNodeCreator(*TargetRigGraph);
+			UAnimGraphNode_LinkedAnimLayer* LinkedAnimLayerNode = LinkedInputLayerNodeCreator.CreateNode();	
+			const FName GraphName = TargetRigGraph->GetFName();
+			LinkedAnimLayerNode->Node.Layer = SourceFuncName; 
+			LinkedInputLayerNodeCreator.Finalize();
+			LinkedAnimLayerNode->NodePosX = GraphPosition.X;
+			LinkedAnimLayerNode->NodePosY = GraphPosition.Y;
+		}
+	}
+	return FReply::Unhandled();
+}
+
+FReply FAnimationLayerDragDropAction::DroppedOnNode(FVector2D ScreenPosition, FVector2D GraphPosition)
+{
+	if (UEdGraphNode* TargetNode = GetHoveredNode())
+	{
+		if (UAnimGraphNode_LinkedAnimLayer* LinkedAnimLayer = Cast<UAnimGraphNode_LinkedAnimLayer>(TargetNode))
+		{
+			LinkedAnimLayer->Node.Layer = SourceFuncName; 
+			return FReply::Handled();
+		}
+	}
+	return FReply::Unhandled();
+}
+
+FReply FAnimationLayerDragDropAction::DroppedOnPin(FVector2D ScreenPosition, FVector2D GraphPosition)
+{
+	return FReply::Unhandled();
+}
+
+FReply FAnimationLayerDragDropAction::DroppedOnAction(TSharedRef<FEdGraphSchemaAction> Action)
+{
+	return FReply::Unhandled();
+}
+
+FReply FAnimationLayerDragDropAction::DroppedOnCategory(FText Category)
+{
+	return FReply::Unhandled();
+}
+
+void FAnimationLayerDragDropAction::HoverTargetChanged()
+{
+	FGraphSchemaActionDragDropAction::HoverTargetChanged();
+	bDropTargetValid = true;
+}
+
+
+TSharedRef<FAnimationLayerDragDropAction> FAnimationLayerDragDropAction::New(TSharedPtr<FEdGraphSchemaAction> InAction, FName InFuncName, UAnimBlueprint* InAnimBlueprint, UAnimationGraph* InAnimationLayerGraph)
+{
+	TSharedRef<FAnimationLayerDragDropAction> Action = MakeShareable(new FAnimationLayerDragDropAction);
+	Action->SourceAction = InAction;
+	Action->SourceAnimBlueprint = InAnimBlueprint;
+	Action->SourceAnimLayerGraph = InAnimationLayerGraph;
+	Action->SourceFuncName = InFuncName; 
+	Action->Construct();
+	return Action;
+}
+
 /////////////////////////////////////////////////////
 // UAnimationGraphSchema
 
@@ -335,6 +438,48 @@ void UAnimationGraphSchema::CreateFunctionGraphTerminators(UEdGraph& Graph, UCla
 	}
 }
 
+
+bool UAnimationGraphSchema::CanGraphBeDropped(TSharedPtr<FEdGraphSchemaAction> InAction) const
+{
+	if (!InAction.IsValid())
+	{
+		return false;
+	}
+
+	if (InAction->GetTypeId() == FEdGraphSchemaAction_K2Graph::StaticGetTypeId())
+	{
+		FEdGraphSchemaAction_K2Graph* FuncAction = (FEdGraphSchemaAction_K2Graph*)InAction.Get();
+		if (UAnimationGraph* AnimGraph = Cast<UAnimationGraph>((UEdGraph*)FuncAction->EdGraph))
+		{
+			return true;
+		}
+	}
+
+	return false;
+}
+
+FReply UAnimationGraphSchema::BeginGraphDragAction(TSharedPtr<FEdGraphSchemaAction> InAction, const FPointerEvent& MouseEvent) const
+{
+	if (!InAction.IsValid())
+	{
+		return FReply::Unhandled();
+	}
+
+	if (InAction->GetTypeId() == FEdGraphSchemaAction_K2Graph::StaticGetTypeId())
+	{
+		FEdGraphSchemaAction_K2Graph* FuncAction = (FEdGraphSchemaAction_K2Graph*)InAction.Get();
+		if (UAnimationGraph* AnimationLayerGraph = Cast<UAnimationGraph>((UEdGraph*)FuncAction->EdGraph))
+		{
+			if (UAnimBlueprint* TargetAnimBlueprint = Cast<UAnimBlueprint>(FBlueprintEditorUtils::FindBlueprintForGraph(AnimationLayerGraph)))
+			{
+				return FReply::Handled().BeginDragDrop(FAnimationLayerDragDropAction::New(InAction, FuncAction->FuncName, TargetAnimBlueprint, AnimationLayerGraph));
+			}
+		}
+	}
+	
+	return FReply::Unhandled();
+}
+
 bool UAnimationGraphSchema::SearchForAutocastFunction(const FEdGraphPinType& OutputPinType, const FEdGraphPinType& InputPinType, FName& TargetFunction, /*out*/ UClass*& FunctionOwner) const
 {
 	if (IsComponentSpacePosePin(OutputPinType) && IsLocalSpacePosePin(InputPinType))
diff --git a/Engine/Source/Editor/AnimGraph/Public/AnimationGraphSchema.h b/Engine/Source/Editor/AnimGraph/Public/AnimationGraphSchema.h
--- a/Engine/Source/Editor/AnimGraph/Public/AnimationGraphSchema.h
+++ b/Engine/Source/Editor/AnimGraph/Public/AnimationGraphSchema.h
@@ -20,7 +20,6 @@ UCLASS(MinimalAPI)
 class UAnimationGraphSchema : public UEdGraphSchema_K2
 {
 	GENERATED_UCLASS_BODY()
-
 	// Common PinNames
 	UPROPERTY()
 	FString PN_SequenceName;    // PC_Object+PSC_Sequence
@@ -65,6 +64,8 @@ class UAnimationGraphSchema : public UEdGraphSchema_K2
 	virtual bool ShouldAlwaysPurgeOnModification() const override { return true; }
 	virtual bool TryCreateConnection(UEdGraphPin* A, UEdGraphPin* B) const override;
 	virtual void GetGraphDisplayInformation(const UEdGraph& Graph, /*out*/ FGraphDisplayInfo& DisplayInfo) const override;
+	virtual bool CanGraphBeDropped(TSharedPtr<FEdGraphSchemaAction> InAction) const override;
+	virtual FReply BeginGraphDragAction(TSharedPtr<FEdGraphSchemaAction> InAction, const FPointerEvent& MouseEvent = FPointerEvent()) const override;
 	//~ End UEdGraphSchema Interface.
 
 	//~ Begin UEdGraphSchema_K2 Interface
@@ -73,6 +74,7 @@ class UAnimationGraphSchema : public UEdGraphSchema_K2
 	virtual bool ArePinsCompatible(const UEdGraphPin* PinA, const UEdGraphPin* PinB, const UClass* CallingContext = NULL, bool bIgnoreArray = false) const override;
 	virtual bool DoesSupportAnimNotifyActions() const override;
 	virtual void CreateFunctionGraphTerminators(UEdGraph& Graph, UClass* Class) const override;
+
 	//~ End UEdGraphSchema_K2 Interface
 
 	/** Spawn the correct node in the Animation Graph using the given AnimationAsset at the supplied location */
