Fix asset player names and 'empty' players not appearing in context menus correctly

Also fix asset references persisting when applying a property access binding

diff --git a/Engine/Source/Editor/AnimGraph/Private/AnimGraphNode_AimOffsetLookAt.cpp b/Engine/Source/Editor/AnimGraph/Private/AnimGraphNode_AimOffsetLookAt.cpp
--- a/Engine/Source/Editor/AnimGraph/Private/AnimGraphNode_AimOffsetLookAt.cpp
+++ b/Engine/Source/Editor/AnimGraph/Private/AnimGraphNode_AimOffsetLookAt.cpp
@@ -16,7 +16,7 @@
 /////////////////////////////////////////////////////
 // UAnimGraphNode_RotationOffsetBlendSpace
 
-#define LOCTEXT_NAMESPACE "A3Nodes"
+#define LOCTEXT_NAMESPACE "UAnimGraphNode_AimOffsetLookAt"
 
 UAnimGraphNode_AimOffsetLookAt::UAnimGraphNode_AimOffsetLookAt(const FObjectInitializer& ObjectInitializer)
 	: Super(ObjectInitializer)
diff --git a/Engine/Source/Editor/AnimGraph/Private/AnimGraphNode_AssetPlayerBase.cpp b/Engine/Source/Editor/AnimGraph/Private/AnimGraphNode_AssetPlayerBase.cpp
--- a/Engine/Source/Editor/AnimGraph/Private/AnimGraphNode_AssetPlayerBase.cpp
+++ b/Engine/Source/Editor/AnimGraph/Private/AnimGraphNode_AssetPlayerBase.cpp
@@ -221,11 +221,19 @@ bool SupportNodeClassForAsset(const UClass* AssetClass, UClass* NodeClass)
 void UAnimGraphNode_AssetPlayerBase::SetupNewNode(UEdGraphNode* InNewNode, bool bInIsTemplateNode, const FAssetData InAssetData)
 {
 	UAnimGraphNode_AssetPlayerBase* GraphNode = CastChecked<UAnimGraphNode_AssetPlayerBase>(InNewNode);
-	InAssetData.GetTagValue("Skeleton", GraphNode->UnloadedSkeletonName);
 	
-	if(!bInIsTemplateNode)
+	if(InAssetData.IsValid())
 	{
-		GraphNode->SetAnimationAsset(CastChecked<UAnimationAsset>(InAssetData.GetAsset()));
+		InAssetData.GetTagValue("Skeleton", GraphNode->UnloadedSkeletonName);
+		if(GraphNode->UnloadedSkeletonName == TEXT("None"))
+		{
+			GraphNode->UnloadedSkeletonName.Empty();
+		}
+		
+		if(!bInIsTemplateNode)
+		{
+			GraphNode->SetAnimationAsset(CastChecked<UAnimationAsset>(InAssetData.GetAsset()));
+		}
 	}
 }
 
@@ -271,16 +279,19 @@ void UAnimGraphNode_AssetPlayerBase::GetMenuActionsHelper(FBlueprintActionDataba
 		}	
 		Filter.bRecursiveClasses = true;
 		
-		TArray<FAssetData> AnimBlueprints;
-		AssetRegistryModule.Get().GetAssets(Filter, AnimBlueprints);
+		TArray<FAssetData> Assets;
+		AssetRegistryModule.Get().GetAssets(Filter, Assets);
 
-		for (const FAssetData& AssetData : AnimBlueprints)
+		for (const FAssetData& AssetData : Assets)
 		{
 			if(AssetData.IsUAsset())
 			{
 				MakeAction(AssetData);
 			}
 		}
+
+		// Add 'empty' asset node
+		MakeAction(FAssetData());
 	}
 	else if (bIsObjectOfAssetType)
 	{
@@ -291,28 +302,38 @@ void UAnimGraphNode_AssetPlayerBase::GetMenuActionsHelper(FBlueprintActionDataba
 bool UAnimGraphNode_AssetPlayerBase::IsActionFilteredOut(class FBlueprintActionFilter const& Filter)
 {
 	bool bIsFilteredOut = false;
+	FBlueprintActionContext const& FilterContext = Filter.Context;
 
-	if(!UnloadedSkeletonName.IsEmpty())
+	for (UBlueprint* Blueprint : FilterContext.Blueprints)
 	{
-		FBlueprintActionContext const& FilterContext = Filter.Context;
-
-		for (UBlueprint* Blueprint : FilterContext.Blueprints)
+		UAnimBlueprint* AnimBlueprint = Cast<UAnimBlueprint>(Blueprint);
+		if (AnimBlueprint)
 		{
-			if (UAnimBlueprint* AnimBlueprint = Cast<UAnimBlueprint>(Blueprint))
+			UAnimationAsset* Asset = GetAnimationAsset();
+			if(Asset)
 			{
-				if(!AnimBlueprint->TargetSkeleton->IsCompatibleSkeletonByAssetString(UnloadedSkeletonName))
+				if (AnimBlueprint->TargetSkeleton == nullptr || !AnimBlueprint->TargetSkeleton->IsCompatible(Asset->GetSkeleton()))
 				{
+					// Asset does not use a compatible skeleton with the Blueprint, cannot use
 					bIsFilteredOut = true;
 					break;
 				}
 			}
-			else
+			else if(!UnloadedSkeletonName.IsEmpty())
 			{
-				// Not an animation Blueprint, cannot use
-				bIsFilteredOut = true;
-				break;
+				if(AnimBlueprint->TargetSkeleton == nullptr || !AnimBlueprint->TargetSkeleton->IsCompatibleSkeletonByAssetString(UnloadedSkeletonName))
+				{
+					bIsFilteredOut = true;
+					break;
+				}
 			}
 		}
+		else
+		{
+			// Not an animation Blueprint, cannot use
+			bIsFilteredOut = true;
+			break;
+		}
 	}
 	return bIsFilteredOut;
 }
diff --git a/Engine/Source/Editor/AnimGraph/Private/AnimGraphNode_Base.cpp b/Engine/Source/Editor/AnimGraph/Private/AnimGraphNode_Base.cpp
--- a/Engine/Source/Editor/AnimGraph/Private/AnimGraphNode_Base.cpp
+++ b/Engine/Source/Editor/AnimGraph/Private/AnimGraphNode_Base.cpp
@@ -295,6 +295,16 @@ FStructProperty* UAnimGraphNode_Base::GetFNodeProperty() const
 	return NULL;
 }
 
+FAnimNode_Base* UAnimGraphNode_Base::GetFNode()
+{
+	if(FStructProperty* Property = GetFNodeProperty())
+	{
+		return Property->ContainerPtrToValuePtr<FAnimNode_Base>(this);
+	}
+
+	return nullptr;
+}
+
 FString UAnimGraphNode_Base::GetNodeCategory() const
 {
 	return TEXT("Misc.");
@@ -701,6 +711,13 @@ TSharedRef<SWidget> UAnimGraphNode_Base::MakePropertyBindingWidget(const FAnimPr
 			{
 				AnimGraphNode->Modify();
 
+				// Reset to default so that references are not preserved
+				if(BindingProperty)
+				{
+					void* PropertyAddress = BindingProperty->ContainerPtrToValuePtr<void>(AnimGraphNode->GetFNode());
+					BindingProperty->InitializeValue(PropertyAddress);
+				}
+				
 				// Pins are exposed if we have a binding or not - and after running this we do.
 				AnimGraphNode->SetPinVisibility(true, InArgs.OptionalPinIndex);
 
diff --git a/Engine/Source/Editor/AnimGraph/Private/AnimGraphNode_BlendSpaceEvaluator.cpp b/Engine/Source/Editor/AnimGraph/Private/AnimGraphNode_BlendSpaceEvaluator.cpp
--- a/Engine/Source/Editor/AnimGraph/Private/AnimGraphNode_BlendSpaceEvaluator.cpp
+++ b/Engine/Source/Editor/AnimGraph/Private/AnimGraphNode_BlendSpaceEvaluator.cpp
@@ -13,7 +13,7 @@
 /////////////////////////////////////////////////////
 // UAnimGraphNode_BlendSpaceEvaluator
 
-#define LOCTEXT_NAMESPACE "A3Nodes"
+#define LOCTEXT_NAMESPACE "UAnimGraphNode_BlendSpaceEvaluator"
 
 UAnimGraphNode_BlendSpaceEvaluator::UAnimGraphNode_BlendSpaceEvaluator(const FObjectInitializer& ObjectInitializer)
 	: Super(ObjectInitializer)
@@ -35,11 +35,25 @@ void UAnimGraphNode_BlendSpaceEvaluator::GetMenuActions(FBlueprintActionDatabase
 		{ },
 		[](const FAssetData& InAssetData)
 		{
-			return FText::Format(LOCTEXT("MenuDescFormat", "Blendspace Evaluator '{0}'"), FText::FromName(InAssetData.AssetName));
+			if(InAssetData.IsValid())
+			{
+				return FText::Format(LOCTEXT("MenuDescFormat", "Blendspace Evaluator '{0}'"), FText::FromName(InAssetData.AssetName));
+			}
+			else
+			{
+				return LOCTEXT("MenuDesc", "Blendspace Evaluator");
+			}
 		},
 		[](const FAssetData& InAssetData)
 		{
-			return FText::Format(LOCTEXT("MenuDescTooltipFormat", "Blendspace Evaluator\n'{0}'"), FText::FromName(InAssetData.ObjectPath));
+			if(InAssetData.IsValid())
+			{
+				return FText::Format(LOCTEXT("MenuDescTooltipFormat", "Blendspace Evaluator\n'{0}'"), FText::FromName(InAssetData.ObjectPath));
+			}
+			else
+			{
+				return LOCTEXT("MenuDescTooltip", "Blendspace Evaluator");
+			}
 		},
 		[](UEdGraphNode* InNewNode, bool bInIsTemplateNode, const FAssetData InAssetData)
 		{
@@ -99,7 +113,7 @@ void UAnimGraphNode_BlendSpaceEvaluator::GetNodeContextMenuActions(UToolMenu* Me
 	{
 		// add an option to convert to single frame
 		{
-			FToolMenuSection& Section = Menu->AddSection("AnimGraphNodeBlendSpacePlayer", NSLOCTEXT("A3Nodes", "BlendSpaceHeading", "Blend Space"));
+			FToolMenuSection& Section = Menu->AddSection("AnimGraphNodeBlendSpacePlayer", LOCTEXT("BlendSpaceHeading", "Blend Space"));
 			Section.AddMenuEntry(FAnimGraphCommands::Get().OpenRelatedAsset);
 			Section.AddMenuEntry(FAnimGraphCommands::Get().ConvertToBSPlayer);
 		}
diff --git a/Engine/Source/Editor/AnimGraph/Private/AnimGraphNode_BlendSpaceGraph.cpp b/Engine/Source/Editor/AnimGraph/Private/AnimGraphNode_BlendSpaceGraph.cpp
--- a/Engine/Source/Editor/AnimGraph/Private/AnimGraphNode_BlendSpaceGraph.cpp
+++ b/Engine/Source/Editor/AnimGraph/Private/AnimGraphNode_BlendSpaceGraph.cpp
@@ -63,11 +63,25 @@ void UAnimGraphNode_BlendSpaceGraph::GetMenuActions(FBlueprintActionDatabaseRegi
 		{ UAimOffsetBlendSpace::StaticClass(), UAimOffsetBlendSpace1D::StaticClass() },
 		[](const FAssetData& InAssetData)
 		{
-			return FText::Format(LOCTEXT("MenuDescFormat", "Blendspace '{0}'"), FText::FromName(InAssetData.AssetName));
+			if(InAssetData.IsValid())
+			{
+				return FText::Format(LOCTEXT("MenuDescFormat", "Blendspace '{0}'"), FText::FromName(InAssetData.AssetName));
+			}
+			else
+			{
+				return LOCTEXT("MenuDesc", "Blendspace");
+			}
 		},
 		[](const FAssetData& InAssetData)
 		{
-			return FText::Format(LOCTEXT("MenuDescTooltipFormat", "Blendspace\n'{0}'"), FText::FromName(InAssetData.ObjectPath));
+			if(InAssetData.IsValid())
+			{
+				return FText::Format(LOCTEXT("MenuDescTooltipFormat", "Blendspace\n'{0}'"), FText::FromName(InAssetData.ObjectPath));
+			}
+			else
+			{
+				return LOCTEXT("MenuDesc", "Blendspace");
+			}
 		},
 		[](UEdGraphNode* InNewNode, bool bInIsTemplateNode, const FAssetData InAssetData)
 		{
diff --git a/Engine/Source/Editor/AnimGraph/Private/AnimGraphNode_BlendSpaceGraphBase.cpp b/Engine/Source/Editor/AnimGraph/Private/AnimGraphNode_BlendSpaceGraphBase.cpp
--- a/Engine/Source/Editor/AnimGraph/Private/AnimGraphNode_BlendSpaceGraphBase.cpp
+++ b/Engine/Source/Editor/AnimGraph/Private/AnimGraphNode_BlendSpaceGraphBase.cpp
@@ -145,44 +145,51 @@ void UAnimGraphNode_BlendSpaceGraphBase::OnCopyTermDefaultsToDefaultObject(IAnim
 
 void UAnimGraphNode_BlendSpaceGraphBase::SetupFromAsset(const FAssetData& InAssetData, bool bInIsTemplateNode)
 {
-	InAssetData.GetTagValue("Skeleton", SkeletonName);
-	
-	if(!bInIsTemplateNode)
+	if(InAssetData.IsValid())
 	{
-		UBlendSpace* Asset = CastChecked<UBlendSpace>(InAssetData.GetAsset());
-		
-		UAnimBlueprint* AnimBlueprint = GetAnimBlueprint();
-		const UAnimationGraphSchema* AnimationGraphSchema = GetDefault<UAnimationGraphSchema>();
-
-		BlendSpaceGraph = CastChecked<UBlendSpaceGraph>(FBlueprintEditorUtils::CreateNewGraph(this, Asset->GetFName(), UBlendSpaceGraph::StaticClass(), UEdGraphSchema::StaticClass()));
-		BlendSpaceGraph->BlendSpace = BlendSpace = DuplicateObject(Asset, BlendSpaceGraph);
-		BlendSpaceGraph->BlendSpace->ClearFlags(RF_Public | RF_Standalone | RF_Transient);
-		BlendSpaceGraph->BlendSpace->SetFlags(RF_Transactional);
-		BlendSpaceGraph->BlendSpace->ResetSkeleton(nullptr);
-		BlendSpaceGraph->BlendSpace->EmptyMetaData();
-		BlendSpaceGraph->BlendSpace->RemoveUserDataOfClass(UAssetUserData::StaticClass());
-		BlendSpaceGraph->bAllowDeletion = false;
-		BlendSpaceGraph->bAllowRenaming = true;
-
-		// Add the new graph as a child of our parent graph
-		UEdGraph* ParentGraph = GetGraph();
-	
-		if(ParentGraph->SubGraphs.Find(BlendSpaceGraph) == INDEX_NONE)
+		InAssetData.GetTagValue("Skeleton", SkeletonName);
+		if(SkeletonName == TEXT("None"))
 		{
-			ParentGraph->Modify();
-			ParentGraph->SubGraphs.Add(BlendSpaceGraph);
+			SkeletonName.Empty();
 		}
-
-		for(int32 SampleIndex = 0; SampleIndex < BlendSpace->SampleData.Num(); ++SampleIndex)
+		
+		if(!bInIsTemplateNode)
 		{
-			FBlendSample& Sample = BlendSpace->SampleData[SampleIndex];
+			UBlendSpace* Asset = CastChecked<UBlendSpace>(InAssetData.GetAsset());
+			
+			UAnimBlueprint* AnimBlueprint = GetAnimBlueprint();
+			const UAnimationGraphSchema* AnimationGraphSchema = GetDefault<UAnimationGraphSchema>();
+
+			BlendSpaceGraph = CastChecked<UBlendSpaceGraph>(FBlueprintEditorUtils::CreateNewGraph(this, Asset->GetFName(), UBlendSpaceGraph::StaticClass(), UEdGraphSchema::StaticClass()));
+			BlendSpaceGraph->BlendSpace = BlendSpace = DuplicateObject(Asset, BlendSpaceGraph);
+			BlendSpaceGraph->BlendSpace->ClearFlags(RF_Public | RF_Standalone | RF_Transient);
+			BlendSpaceGraph->BlendSpace->SetFlags(RF_Transactional);
+			BlendSpaceGraph->BlendSpace->ResetSkeleton(nullptr);
+			BlendSpaceGraph->BlendSpace->EmptyMetaData();
+			BlendSpaceGraph->BlendSpace->RemoveUserDataOfClass(UAssetUserData::StaticClass());
+			BlendSpaceGraph->bAllowDeletion = false;
+			BlendSpaceGraph->bAllowRenaming = true;
+
+			// Add the new graph as a child of our parent graph
+			UEdGraph* ParentGraph = GetGraph();
+		
+			if(ParentGraph->SubGraphs.Find(BlendSpaceGraph) == INDEX_NONE)
+			{
+				ParentGraph->Modify();
+				ParentGraph->SubGraphs.Add(BlendSpaceGraph);
+			}
 
-			FName SampleName = Sample.Animation != nullptr ? Sample.Animation->GetFName() : FName("Sample", SampleIndex);
+			for(int32 SampleIndex = 0; SampleIndex < BlendSpace->SampleData.Num(); ++SampleIndex)
+			{
+				FBlendSample& Sample = BlendSpace->SampleData[SampleIndex];
+
+				FName SampleName = Sample.Animation != nullptr ? Sample.Animation->GetFName() : FName("Sample", SampleIndex);
 
-			UAnimationBlendSpaceSampleGraph* SampleGraph = AddGraph(SampleName, Sample.Animation);
+				UAnimationBlendSpaceSampleGraph* SampleGraph = AddGraph(SampleName, Sample.Animation);
 
-			// Clear the animation now we have created the point
-			Sample.Animation = nullptr;
+				// Clear the animation now we have created the point
+				Sample.Animation = nullptr;
+			}
 		}
 	}
 }
@@ -533,28 +540,29 @@ bool UAnimGraphNode_BlendSpaceGraphBase::IsActionFilteredOut(class FBlueprintAct
 {
 	bool bIsFilteredOut = false;
 
-	if(!SkeletonName.IsEmpty())
-	{
-		FBlueprintActionContext const& FilterContext = Filter.Context;
+	FBlueprintActionContext const& FilterContext = Filter.Context;
 
-		for (UBlueprint* Blueprint : FilterContext.Blueprints)
+	for (UBlueprint* Blueprint : FilterContext.Blueprints)
+	{
+		if (UAnimBlueprint* AnimBlueprint = Cast<UAnimBlueprint>(Blueprint))
 		{
-			if (UAnimBlueprint* AnimBlueprint = Cast<UAnimBlueprint>(Blueprint))
+			if(!SkeletonName.IsEmpty())
 			{
-				if(!AnimBlueprint->TargetSkeleton->IsCompatibleSkeletonByAssetString(SkeletonName))
+				if(AnimBlueprint->TargetSkeleton == nullptr || !AnimBlueprint->TargetSkeleton->IsCompatibleSkeletonByAssetString(SkeletonName))
 				{
 					bIsFilteredOut = true;
 					break;
 				}
 			}
-			else
-			{
-				// Not an animation Blueprint, cannot use
-				bIsFilteredOut = true;
-				break;
-			}
+		}
+		else
+		{
+			// Not an animation Blueprint, cannot use
+			bIsFilteredOut = true;
+			break;
 		}
 	}
+
 	return bIsFilteredOut;
 }
 
diff --git a/Engine/Source/Editor/AnimGraph/Private/AnimGraphNode_BlendSpacePlayer.cpp b/Engine/Source/Editor/AnimGraph/Private/AnimGraphNode_BlendSpacePlayer.cpp
--- a/Engine/Source/Editor/AnimGraph/Private/AnimGraphNode_BlendSpacePlayer.cpp
+++ b/Engine/Source/Editor/AnimGraph/Private/AnimGraphNode_BlendSpacePlayer.cpp
@@ -18,7 +18,7 @@
 /////////////////////////////////////////////////////
 // UAnimGraphNode_BlendSpacePlayer
 
-#define LOCTEXT_NAMESPACE "A3Nodes"
+#define LOCTEXT_NAMESPACE "UAnimGraphNode_BlendSpacePlayer"
 
 UAnimGraphNode_BlendSpacePlayer::UAnimGraphNode_BlendSpacePlayer(const FObjectInitializer& ObjectInitializer)
 	: Super(ObjectInitializer)
@@ -83,7 +83,7 @@ void UAnimGraphNode_BlendSpacePlayer::GetNodeContextMenuActions(UToolMenu* Menu,
 	{
 		// add an option to convert to single frame
 		{
-			FToolMenuSection& Section = Menu->AddSection("AnimGraphNodeBlendSpaceEvaluator", NSLOCTEXT("A3Nodes", "BlendSpaceHeading", "Blend Space"));
+			FToolMenuSection& Section = Menu->AddSection("AnimGraphNodeBlendSpaceEvaluator", LOCTEXT("BlendSpaceHeading", "Blend Space"));
 			Section.AddMenuEntry(FAnimGraphCommands::Get().OpenRelatedAsset);
 			Section.AddMenuEntry(FAnimGraphCommands::Get().ConvertToBSEvaluator);
 			Section.AddMenuEntry(FAnimGraphCommands::Get().ConvertToBSGraph);
@@ -100,11 +100,25 @@ void UAnimGraphNode_BlendSpacePlayer::GetMenuActions(FBlueprintActionDatabaseReg
 		{ },
 		[](const FAssetData& InAssetData)
 		{
-			return FText::Format(LOCTEXT("MenuDescFormat", "Blendspace Player '{0}'"), FText::FromName(InAssetData.AssetName));
+			if(InAssetData.IsValid())
+			{
+				return FText::Format(LOCTEXT("MenuDescFormat", "Blendspace Player '{0}'"), FText::FromName(InAssetData.AssetName));
+			}
+			else
+			{
+				return LOCTEXT("MenuDesc", "Blendspace Player");
+			}
 		},
 		[](const FAssetData& InAssetData)
 		{
-			return FText::Format(LOCTEXT("MenuDescTooltipFormat", "Blendspace Player\n'{0}'"), FText::FromName(InAssetData.ObjectPath));
+			if(InAssetData.IsValid())
+			{
+				return FText::Format(LOCTEXT("MenuDescTooltipFormat", "Blendspace Player\n'{0}'"), FText::FromName(InAssetData.ObjectPath));
+			}
+			else
+			{
+				return LOCTEXT("MenuDescTooltip", "Blendspace Player");
+			}
 		},
 		[](UEdGraphNode* InNewNode, bool bInIsTemplateNode, const FAssetData InAssetData)
 		{
diff --git a/Engine/Source/Editor/AnimGraph/Private/AnimGraphNode_Mirror.cpp b/Engine/Source/Editor/AnimGraph/Private/AnimGraphNode_Mirror.cpp
--- a/Engine/Source/Editor/AnimGraph/Private/AnimGraphNode_Mirror.cpp
+++ b/Engine/Source/Editor/AnimGraph/Private/AnimGraphNode_Mirror.cpp
@@ -243,19 +243,19 @@ bool UAnimGraphNode_Mirror::IsActionFilteredOut(class FBlueprintActionFilter con
 	{
 		if (UAnimBlueprint* AnimBlueprint = Cast<UAnimBlueprint>(Blueprint))
 		{
-			if (Node.GetMirrorDataTable())
+			UMirrorDataTable* MirrorDataTable = Node.GetMirrorDataTable();
+			if (MirrorDataTable)
 			{
-				if (Node.GetMirrorDataTable()->Skeleton != AnimBlueprint->TargetSkeleton)
+				if(AnimBlueprint->TargetSkeleton == nullptr || !AnimBlueprint->TargetSkeleton->IsCompatible(MirrorDataTable->Skeleton))
 				{
 					// Mirror Data Table does not use the same skeleton as the Blueprint, cannot use
 					bIsFilteredOut = true;
 					break;
 				}
 			}
-			else
+			else if(!UnloadedSkeletonName.IsEmpty())
 			{
-				FAssetData SkeletonData(AnimBlueprint->TargetSkeleton);
-				if (UnloadedSkeletonName != SkeletonData.GetExportTextName())
+				if(AnimBlueprint->TargetSkeleton == nullptr || !AnimBlueprint->TargetSkeleton->IsCompatibleSkeletonByAssetString(UnloadedSkeletonName))
 				{
 					bIsFilteredOut = true;
 					break;
diff --git a/Engine/Source/Editor/AnimGraph/Private/AnimGraphNode_PoseBlendNode.cpp b/Engine/Source/Editor/AnimGraph/Private/AnimGraphNode_PoseBlendNode.cpp
--- a/Engine/Source/Editor/AnimGraph/Private/AnimGraphNode_PoseBlendNode.cpp
+++ b/Engine/Source/Editor/AnimGraph/Private/AnimGraphNode_PoseBlendNode.cpp
@@ -97,95 +97,39 @@ FText UAnimGraphNode_PoseBlendNode::GetNodeTitle(ENodeTitleType::Type TitleType)
 	}
 }
 
-void UAnimGraphNode_PoseBlendNode::GetMenuActions(FBlueprintActionDatabaseRegistrar& ActionRegistrar) const
+void UAnimGraphNode_PoseBlendNode::GetMenuActions(FBlueprintActionDatabaseRegistrar& InActionRegistrar) const
 {
-	auto LoadedAssetSetup = [](UEdGraphNode* NewNode, bool /*bIsTemplateNode*/, TWeakObjectPtr<UPoseAsset> PoseAssetPtr)
-	{
-		UAnimGraphNode_PoseBlendNode* PoseBlendNodeNode = CastChecked<UAnimGraphNode_PoseBlendNode>(NewNode);
-		PoseBlendNodeNode->Node.PoseAsset = PoseAssetPtr.Get();
-	};
-
-	auto UnloadedAssetSetup = [](UEdGraphNode* NewNode, bool bIsTemplateNode, const FAssetData AssetData)
-	{
-		UAnimGraphNode_PoseBlendNode* PoseBlendNodeNode = CastChecked<UAnimGraphNode_PoseBlendNode>(NewNode);
-		if (bIsTemplateNode)
-		{
-			AssetData.GetTagValue("Skeleton", PoseBlendNodeNode->UnloadedSkeletonName);
-		}
-		else
-		{
-			UPoseAsset* PoseAsset = Cast<UPoseAsset>(AssetData.GetAsset());
-			check(PoseAsset != nullptr);
-			PoseBlendNodeNode->Node.PoseAsset = PoseAsset;
-		}
-	};
-
-	const UObject* QueryObject = ActionRegistrar.GetActionKeyFilter();
-	if (QueryObject == nullptr)
-	{
-		FAssetRegistryModule& AssetRegistryModule = FModuleManager::LoadModuleChecked<FAssetRegistryModule>(TEXT("AssetRegistry"));
-		// define a filter to help in pulling UPoseAsset asset data from the registry
-		FARFilter Filter;
-		Filter.ClassNames.Add(UPoseAsset::StaticClass()->GetFName());
-		Filter.bRecursiveClasses = true;
-		// Find matching assets and add an entry for each one
-		TArray<FAssetData> PoseAssetList;
-		AssetRegistryModule.Get().GetAssets(Filter, /*out*/PoseAssetList);
-
-		for (auto AssetIt = PoseAssetList.CreateConstIterator(); AssetIt; ++AssetIt)
+	GetMenuActionsHelper(
+		InActionRegistrar,
+		GetClass(),
+		{ UPoseAsset::StaticClass() },
+		{ },
+		[](const FAssetData& InAssetData)
 		{
-			const FAssetData& Asset = *AssetIt;
-
-			UBlueprintNodeSpawner* NodeSpawner = UBlueprintNodeSpawner::Create(GetClass());
-			if (Asset.IsAssetLoaded())
+			if(InAssetData.IsValid())
 			{
-				TWeakObjectPtr<UPoseAsset> PoseAsset = Cast<UPoseAsset>(Asset.GetAsset());
-				NodeSpawner->CustomizeNodeDelegate = UBlueprintNodeSpawner::FCustomizeNodeDelegate::CreateStatic(LoadedAssetSetup, PoseAsset);
-				NodeSpawner->DefaultMenuSignature.MenuName = GetTitleGivenAssetInfo(FText::FromName(PoseAsset->GetFName()));
+				return GetTitleGivenAssetInfo(FText::FromName(InAssetData.AssetName));
 			}
 			else
 			{
-				NodeSpawner->CustomizeNodeDelegate = UBlueprintNodeSpawner::FCustomizeNodeDelegate::CreateStatic(UnloadedAssetSetup, Asset);
-				NodeSpawner->DefaultMenuSignature.MenuName = GetTitleGivenAssetInfo(FText::FromName(Asset.AssetName));
+				return LOCTEXT("PlayerDesc", "Evaluate Pose");
 			}
-			ActionRegistrar.AddBlueprintAction(Asset, NodeSpawner);
-		}
-	}
-	else if (const UPoseAsset* PoseAsset = Cast<UPoseAsset>(QueryObject))
-	{
-		UBlueprintNodeSpawner* NodeSpawner = UBlueprintNodeSpawner::Create(GetClass());
-
-		TWeakObjectPtr<UPoseAsset> PoseAssetPtr = MakeWeakObjectPtr(const_cast<UPoseAsset*>(PoseAsset));
-		NodeSpawner->CustomizeNodeDelegate = UBlueprintNodeSpawner::FCustomizeNodeDelegate::CreateStatic(LoadedAssetSetup, PoseAssetPtr);
-		NodeSpawner->DefaultMenuSignature.MenuName = GetTitleGivenAssetInfo(FText::FromName(PoseAsset->GetFName()));
-
-		ActionRegistrar.AddBlueprintAction(QueryObject, NodeSpawner);
-	}
-	else if (QueryObject == GetClass())
-	{
-		FAssetRegistryModule& AssetRegistryModule = FModuleManager::LoadModuleChecked<FAssetRegistryModule>(TEXT("AssetRegistry"));
-		// define a filter to help in pulling UPoseAsset asset data from the registry
-		FARFilter Filter;
-		Filter.ClassNames.Add(UPoseAsset::StaticClass()->GetFName());
-		Filter.bRecursiveClasses = true;
-		// Find matching assets and add an entry for each one
-		TArray<FAssetData> PoseAssetList;
-		AssetRegistryModule.Get().GetAssets(Filter, /*out*/PoseAssetList);
-
-		for (auto AssetIt = PoseAssetList.CreateConstIterator(); AssetIt; ++AssetIt)
+		},
+		[](const FAssetData& InAssetData)
 		{
-			const FAssetData& Asset = *AssetIt;
-			if (Asset.IsAssetLoaded())
+			if(InAssetData.IsValid())
 			{
-				continue;
+				return GetTitleGivenAssetInfo(FText::FromName(InAssetData.AssetName));
 			}
-
-			UBlueprintNodeSpawner* NodeSpawner = UBlueprintNodeSpawner::Create(GetClass());
-			NodeSpawner->CustomizeNodeDelegate = UBlueprintNodeSpawner::FCustomizeNodeDelegate::CreateStatic(UnloadedAssetSetup, Asset);
-			NodeSpawner->DefaultMenuSignature.MenuName = GetTitleGivenAssetInfo(FText::FromName(Asset.AssetName));
-			ActionRegistrar.AddBlueprintAction(Asset, NodeSpawner);
-		}
-	}
+			else
+			{
+				return LOCTEXT("PlayerDescTooltip", "Evaluate Pose Asset");
+			}
+		},
+		[](UEdGraphNode* InNewNode, bool bInIsTemplateNode, const FAssetData InAssetData)
+		{
+			UAnimGraphNode_AssetPlayerBase::SetupNewNode(InNewNode, bInIsTemplateNode, InAssetData);
+		});	
 }
 
 FText UAnimGraphNode_PoseBlendNode::GetTitleGivenAssetInfo(const FText& AssetName)
@@ -196,44 +140,6 @@ FText UAnimGraphNode_PoseBlendNode::GetTitleGivenAssetInfo(const FText& AssetNam
 	return FText::Format(LOCTEXT("PoseAssetNodeTitle", "Evaluate Pose {AssetName}"), Args);
 }
 
-bool UAnimGraphNode_PoseBlendNode::IsActionFilteredOut(class FBlueprintActionFilter const& Filter)
-{
-	bool bIsFilteredOut = false;
-	FBlueprintActionContext const& FilterContext = Filter.Context;
-
-	for (UBlueprint* Blueprint : FilterContext.Blueprints)
-	{
-		if (UAnimBlueprint* AnimBlueprint = Cast<UAnimBlueprint>(Blueprint))
-		{
-			if (Node.PoseAsset)
-			{
-				if (Node.PoseAsset->GetSkeleton() != AnimBlueprint->TargetSkeleton)
-				{
-					// PoseAsset does not use the same skeleton as the Blueprint, cannot use
-					bIsFilteredOut = true;
-					break;
-				}
-			}
-			else
-			{
-				FAssetData SkeletonData(AnimBlueprint->TargetSkeleton);
-				if (UnloadedSkeletonName != SkeletonData.GetExportTextName())
-				{
-					bIsFilteredOut = true;
-					break;
-				}
-			}
-		}
-		else
-		{
-			// Not an animation Blueprint, cannot use
-			bIsFilteredOut = true;
-			break;
-		}
-	}
-	return bIsFilteredOut;
-}
-
 FText UAnimGraphNode_PoseBlendNode::GetMenuCategory() const
 {
 	return LOCTEXT("PoseAssetCategory_Label", "Poses");
diff --git a/Engine/Source/Editor/AnimGraph/Private/AnimGraphNode_RotationOffsetBlendSpace.cpp b/Engine/Source/Editor/AnimGraph/Private/AnimGraphNode_RotationOffsetBlendSpace.cpp
--- a/Engine/Source/Editor/AnimGraph/Private/AnimGraphNode_RotationOffsetBlendSpace.cpp
+++ b/Engine/Source/Editor/AnimGraph/Private/AnimGraphNode_RotationOffsetBlendSpace.cpp
@@ -21,7 +21,7 @@
 /////////////////////////////////////////////////////
 // UAnimGraphNode_RotationOffsetBlendSpace
 
-#define LOCTEXT_NAMESPACE "A3Nodes"
+#define LOCTEXT_NAMESPACE "UAnimGraphNode_RotationOffsetBlendSpace"
 
 UAnimGraphNode_RotationOffsetBlendSpace::UAnimGraphNode_RotationOffsetBlendSpace(const FObjectInitializer& ObjectInitializer)
 	: Super(ObjectInitializer)
@@ -79,11 +79,25 @@ void UAnimGraphNode_RotationOffsetBlendSpace::GetMenuActions(FBlueprintActionDat
 		{ },
 		[](const FAssetData& InAssetData)
 		{
-			return FText::Format(LOCTEXT("MenuDescFormat", "AimOffset Player '{0}'"), FText::FromName(InAssetData.AssetName));
+			if(InAssetData.IsValid())
+			{
+				return FText::Format(LOCTEXT("MenuDescFormat", "AimOffset Player '{0}'"), FText::FromName(InAssetData.AssetName));
+			}
+			else
+			{
+				return LOCTEXT("MenuDesc", "AimOffset Player");
+			}
 		},
 		[](const FAssetData& InAssetData)
 		{
-			return FText::Format(LOCTEXT("MenuDescTooltipFormat", "AimOffset Player\n'{0}'"), FText::FromName(InAssetData.ObjectPath));
+			if(InAssetData.IsValid())
+			{
+				return FText::Format(LOCTEXT("MenuDescTooltipFormat", "AimOffset Player\n'{0}'"), FText::FromName(InAssetData.ObjectPath));
+			}
+			else
+			{
+				return LOCTEXT("MenuDescTooltip", "AimOffset Player");
+			}
 		},
 		[](UEdGraphNode* InNewNode, bool bInIsTemplateNode, const FAssetData InAssetData)
 		{
@@ -156,7 +170,7 @@ void UAnimGraphNode_RotationOffsetBlendSpace::GetNodeContextMenuActions(UToolMen
 	{
 		// add an option to convert to single frame
 		{
-			FToolMenuSection& Section = Menu->AddSection("AnimGraphNodeBlendSpacePlayer", NSLOCTEXT("A3Nodes", "BlendSpaceHeading", "Blend Space"));
+			FToolMenuSection& Section = Menu->AddSection("AnimGraphNodeBlendSpacePlayer", LOCTEXT("BlendSpaceHeading", "Blend Space"));
 			Section.AddMenuEntry(FAnimGraphCommands::Get().OpenRelatedAsset);
 			Section.AddMenuEntry(FAnimGraphCommands::Get().ConvertToAimOffsetLookAt);
 			Section.AddMenuEntry(FAnimGraphCommands::Get().ConvertToAimOffsetGraph);
diff --git a/Engine/Source/Editor/AnimGraph/Private/AnimGraphNode_RotationOffsetBlendSpaceGraph.cpp b/Engine/Source/Editor/AnimGraph/Private/AnimGraphNode_RotationOffsetBlendSpaceGraph.cpp
--- a/Engine/Source/Editor/AnimGraph/Private/AnimGraphNode_RotationOffsetBlendSpaceGraph.cpp
+++ b/Engine/Source/Editor/AnimGraph/Private/AnimGraphNode_RotationOffsetBlendSpaceGraph.cpp
@@ -65,11 +65,25 @@ void UAnimGraphNode_RotationOffsetBlendSpaceGraph::GetMenuActions(FBlueprintActi
 		{ },
 		[](const FAssetData& InAssetData)
 		{
-			return FText::Format(LOCTEXT("MenuDescFormat", "AimOffset '{0}'"), FText::FromName(InAssetData.AssetName));
+			if(InAssetData.IsValid())
+			{
+				return FText::Format(LOCTEXT("MenuDescFormat", "AimOffset '{0}'"), FText::FromName(InAssetData.AssetName));
+			}
+			else
+			{
+				return LOCTEXT("MenuDesc", "AimOffset");
+			}
 		},
 		[](const FAssetData& InAssetData)
 		{
-			return FText::Format(LOCTEXT("MenuDescTooltipFormat", "AimOffset\n'{0}'"), FText::FromName(InAssetData.ObjectPath));
+			if(InAssetData.IsValid())
+			{
+				return FText::Format(LOCTEXT("MenuDescTooltipFormat", "AimOffset\n'{0}'"), FText::FromName(InAssetData.ObjectPath));
+			}
+			else
+			{
+				return LOCTEXT("MenuDescTooltip", "AimOffset");
+			}
 		},
 		[](UEdGraphNode* InNewNode, bool bInIsTemplateNode, const FAssetData InAssetData)
 		{
diff --git a/Engine/Source/Editor/AnimGraph/Private/AnimGraphNode_SequenceEvaluator.cpp b/Engine/Source/Editor/AnimGraph/Private/AnimGraphNode_SequenceEvaluator.cpp
--- a/Engine/Source/Editor/AnimGraph/Private/AnimGraphNode_SequenceEvaluator.cpp
+++ b/Engine/Source/Editor/AnimGraph/Private/AnimGraphNode_SequenceEvaluator.cpp
@@ -19,7 +19,7 @@
 /////////////////////////////////////////////////////
 // UAnimGraphNode_SequenceEvaluator
 
-#define LOCTEXT_NAMESPACE "A3Nodes"
+#define LOCTEXT_NAMESPACE "UAnimGraphNode_SequenceEvaluator"
 
 UAnimGraphNode_SequenceEvaluator::UAnimGraphNode_SequenceEvaluator(const FObjectInitializer& ObjectInitializer)
 	: Super(ObjectInitializer)
@@ -77,26 +77,40 @@ void UAnimGraphNode_SequenceEvaluator::GetMenuActions(FBlueprintActionDatabaseRe
 		{ },
 		[](const FAssetData& InAssetData)
 		{
-			const FString TagValue = InAssetData.GetTagValueRef<FString>(GET_MEMBER_NAME_CHECKED(UAnimSequence, AdditiveAnimType));
-			if(const bool bKnownToBeAdditive = (!TagValue.IsEmpty() && !TagValue.Equals(TEXT("AAT_None"))))
+			if(InAssetData.IsValid())
 			{
-				return FText::Format(LOCTEXT("MenuDescFormat", "Evaluate '{0}' (additive)"), FText::FromName(InAssetData.AssetName));
+				const FString TagValue = InAssetData.GetTagValueRef<FString>(GET_MEMBER_NAME_CHECKED(UAnimSequence, AdditiveAnimType));
+				if(const bool bKnownToBeAdditive = (!TagValue.IsEmpty() && !TagValue.Equals(TEXT("AAT_None"))))
+				{
+					return FText::Format(LOCTEXT("MenuDescFormatAdditive", "Evaluate '{0}' (additive)"), FText::FromName(InAssetData.AssetName));
+				}
+				else
+				{
+					return FText::Format(LOCTEXT("MenuDescFormat", "Evaluate '{0}'"), FText::FromName(InAssetData.AssetName));
+				}
 			}
 			else
 			{
-				return FText::Format(LOCTEXT("MenuDescFormat", "Evaluate '{0}'"), FText::FromName(InAssetData.AssetName));
+				return LOCTEXT("MenuDesc", "Sequence Evaluator");
 			}
 		},
 		[](const FAssetData& InAssetData)
 		{
-			const FString TagValue = InAssetData.GetTagValueRef<FString>(GET_MEMBER_NAME_CHECKED(UAnimSequence, AdditiveAnimType));
-			if(const bool bKnownToBeAdditive = (!TagValue.IsEmpty() && !TagValue.Equals(TEXT("AAT_None"))))
+			if(InAssetData.IsValid())
 			{
-				return FText::Format(LOCTEXT("MenuDescTooltipFormat", "Evaluate (additive)\n'{0}'"), FText::FromName(InAssetData.ObjectPath));
+				const FString TagValue = InAssetData.GetTagValueRef<FString>(GET_MEMBER_NAME_CHECKED(UAnimSequence, AdditiveAnimType));
+				if(const bool bKnownToBeAdditive = (!TagValue.IsEmpty() && !TagValue.Equals(TEXT("AAT_None"))))
+				{
+					return FText::Format(LOCTEXT("MenuDescTooltipFormat", "Evaluate (additive)\n'{0}'"), FText::FromName(InAssetData.ObjectPath));
+				}
+				else
+				{
+					return FText::Format(LOCTEXT("MenuDescTooltipFormat", "Evaluate\n'{0}'"), FText::FromName(InAssetData.ObjectPath));
+				}
 			}
 			else
 			{
-				return FText::Format(LOCTEXT("MenuDescTooltipFormat", "Evaluate\n'{0}'"), FText::FromName(InAssetData.ObjectPath));
+				return LOCTEXT("MenuDescTooltip", "Sequence Evaluator");
 			}
 		},
 		[](UEdGraphNode* InNewNode, bool bInIsTemplateNode, const FAssetData InAssetData)
@@ -171,7 +185,7 @@ void UAnimGraphNode_SequenceEvaluator::GetNodeContextMenuActions(UToolMenu* Menu
 	{
 		// add an option to convert to a regular sequence player
 		{
-			FToolMenuSection& Section = Menu->AddSection("AnimGraphNodeSequenceEvaluator", NSLOCTEXT("A3Nodes", "SequenceEvaluatorHeading", "Sequence Evaluator"));
+			FToolMenuSection& Section = Menu->AddSection("AnimGraphNodeSequenceEvaluator", LOCTEXT("SequenceEvaluatorHeading", "Sequence Evaluator"));
 			Section.AddMenuEntry(FAnimGraphCommands::Get().OpenRelatedAsset);
 			Section.AddMenuEntry(FAnimGraphCommands::Get().ConvertToSeqPlayer);
 		}
diff --git a/Engine/Source/Editor/AnimGraph/Private/AnimGraphNode_SequencePlayer.cpp b/Engine/Source/Editor/AnimGraph/Private/AnimGraphNode_SequencePlayer.cpp
--- a/Engine/Source/Editor/AnimGraph/Private/AnimGraphNode_SequencePlayer.cpp
+++ b/Engine/Source/Editor/AnimGraph/Private/AnimGraphNode_SequencePlayer.cpp
@@ -21,7 +21,7 @@
 #include "UObject/UE5MainStreamObjectVersion.h"
 #include "IAnimBlueprintNodeOverrideAssetsContext.h"
 
-#define LOCTEXT_NAMESPACE "A3Nodes"
+#define LOCTEXT_NAMESPACE "UAnimGraphNode_SequencePlayer"
 
 /////////////////////////////////////////////////////
 // UAnimGraphNode_SequencePlayer
@@ -94,26 +94,40 @@ void UAnimGraphNode_SequencePlayer::GetMenuActions(FBlueprintActionDatabaseRegis
 		{ },
 		[](const FAssetData& InAssetData)
 		{
-			const FString TagValue = InAssetData.GetTagValueRef<FString>(GET_MEMBER_NAME_CHECKED(UAnimSequence, AdditiveAnimType));
-			if(const bool bKnownToBeAdditive = (!TagValue.IsEmpty() && !TagValue.Equals(TEXT("AAT_None"))))
+			if(InAssetData.IsValid())
 			{
-				return FText::Format(LOCTEXT("MenuDescFormat", "Play '{0}' (additive)"), FText::FromName(InAssetData.AssetName));
+				const FString TagValue = InAssetData.GetTagValueRef<FString>(GET_MEMBER_NAME_CHECKED(UAnimSequence, AdditiveAnimType));
+				if(const bool bKnownToBeAdditive = (!TagValue.IsEmpty() && !TagValue.Equals(TEXT("AAT_None"))))
+				{
+					return FText::Format(LOCTEXT("MenuDescFormat", "Play '{0}' (additive)"), FText::FromName(InAssetData.AssetName));
+				}
+				else
+				{
+					return FText::Format(LOCTEXT("MenuDescFormat", "Play '{0}'"), FText::FromName(InAssetData.AssetName));
+				}
 			}
 			else
 			{
-				return FText::Format(LOCTEXT("MenuDescFormat", "Play '{0}'"), FText::FromName(InAssetData.AssetName));
+				return LOCTEXT("PlayerDesc", "Sequence Player");
 			}
 		},
 		[](const FAssetData& InAssetData)
 		{
-			const FString TagValue = InAssetData.GetTagValueRef<FString>(GET_MEMBER_NAME_CHECKED(UAnimSequence, AdditiveAnimType));
-			if(const bool bKnownToBeAdditive = (!TagValue.IsEmpty() && !TagValue.Equals(TEXT("AAT_None"))))
+			if(InAssetData.IsValid())
 			{
-				return FText::Format(LOCTEXT("MenuDescTooltipFormat", "Play (additive)\n'{0}'"), FText::FromName(InAssetData.ObjectPath));
+				const FString TagValue = InAssetData.GetTagValueRef<FString>(GET_MEMBER_NAME_CHECKED(UAnimSequence, AdditiveAnimType));
+				if(const bool bKnownToBeAdditive = (!TagValue.IsEmpty() && !TagValue.Equals(TEXT("AAT_None"))))
+				{
+					return FText::Format(LOCTEXT("MenuDescTooltipFormat", "Play (additive)\n'{0}'"), FText::FromName(InAssetData.ObjectPath));
+				}
+				else
+				{
+					return FText::Format(LOCTEXT("MenuDescTooltipFormat", "Play\n'{0}'"), FText::FromName(InAssetData.ObjectPath));
+				}
 			}
 			else
 			{
-				return FText::Format(LOCTEXT("MenuDescTooltipFormat", "Play\n'{0}'"), FText::FromName(InAssetData.ObjectPath));
+				return LOCTEXT("PlayerDescTooltip", "Sequence Player");
 			}
 		},
 		[](UEdGraphNode* InNewNode, bool bInIsTemplateNode, const FAssetData InAssetData)
@@ -122,45 +136,6 @@ void UAnimGraphNode_SequencePlayer::GetMenuActions(FBlueprintActionDatabaseRegis
 		});	
 }
 
-bool UAnimGraphNode_SequencePlayer::IsActionFilteredOut(class FBlueprintActionFilter const& Filter)
-{
-	bool bIsFilteredOut = false;
-	FBlueprintActionContext const& FilterContext = Filter.Context;
-
-	for (UBlueprint* Blueprint : FilterContext.Blueprints)
-	{
-		UAnimBlueprint* AnimBlueprint = Cast<UAnimBlueprint>(Blueprint);
-		if (AnimBlueprint && AnimBlueprint->TargetSkeleton)
-		{
-			UAnimSequenceBase* Sequence = Node.GetSequence();
-			if(Sequence)
-			{
-				if (!AnimBlueprint->TargetSkeleton->IsCompatible(Sequence->GetSkeleton()))
-				{
-					// Asset does not use a compatible skeleton with the Blueprint, cannot use
-					bIsFilteredOut = true;
-					break;
-				}
-			}
-			else 
-			{
-				if (!AnimBlueprint->TargetSkeleton->IsCompatibleSkeletonByAssetString(UnloadedSkeletonName))
-				{
-					bIsFilteredOut = true;
-					break;
-				}
-			}
-		}
-		else
-		{
-			// Not an animation Blueprint or has no target skeleton, cannot use
-			bIsFilteredOut = true;
-			break;
-		}
-	}
-	return bIsFilteredOut;
-}
-
 EAnimAssetHandlerType UAnimGraphNode_SequencePlayer::SupportsAssetClass(const UClass* AssetClass) const
 {
 	if (AssetClass->IsChildOf(UAnimSequence::StaticClass()) || AssetClass->IsChildOf(UAnimComposite::StaticClass()))
@@ -233,7 +208,7 @@ void UAnimGraphNode_SequencePlayer::GetNodeContextMenuActions(UToolMenu* Menu, U
 	{
 		// add an option to convert to single frame
 		{
-			FToolMenuSection& Section = Menu->AddSection("AnimGraphNodeSequencePlayer", NSLOCTEXT("A3Nodes", "SequencePlayerHeading", "Sequence Player"));
+			FToolMenuSection& Section = Menu->AddSection("AnimGraphNodeSequencePlayer", LOCTEXT("SequencePlayerHeading", "Sequence Player"));
 			Section.AddMenuEntry(FAnimGraphCommands::Get().OpenRelatedAsset);
 			Section.AddMenuEntry(FAnimGraphCommands::Get().ConvertToSeqEvaluator);
 		}
diff --git a/Engine/Source/Editor/AnimGraph/Public/AnimGraphNode_Base.h b/Engine/Source/Editor/AnimGraph/Public/AnimGraphNode_Base.h
--- a/Engine/Source/Editor/AnimGraph/Public/AnimGraphNode_Base.h
+++ b/Engine/Source/Editor/AnimGraph/Public/AnimGraphNode_Base.h
@@ -448,6 +448,9 @@ protected:
 	// Gets the animation FNode property represented by this ed graph node
 	FStructProperty* GetFNodeProperty() const;
 
+	// Get the runtime anim node that we template
+	FAnimNode_Base* GetFNode();
+	
 	// Get the extension types that this node type holds on the anim blueprint. Some extension types are always requested by the system
 	virtual void GetRequiredExtensions(TArray<TSubclassOf<UAnimBlueprintExtension>>& OutExtensions) const {}
 	
diff --git a/Engine/Source/Editor/AnimGraph/Public/AnimGraphNode_PoseBlendNode.h b/Engine/Source/Editor/AnimGraph/Public/AnimGraphNode_PoseBlendNode.h
--- a/Engine/Source/Editor/AnimGraph/Public/AnimGraphNode_PoseBlendNode.h
+++ b/Engine/Source/Editor/AnimGraph/Public/AnimGraphNode_PoseBlendNode.h
@@ -30,7 +30,6 @@ class UAnimGraphNode_PoseBlendNode : public UAnimGraphNode_PoseHandler
 	virtual void GetAllAnimationSequencesReferred(TArray<UAnimationAsset*>& AnimationAssets) const override;
 	virtual void ReplaceReferredAnimations(const TMap<UAnimationAsset*, UAnimationAsset*>& AnimAssetReplacementMap) override;
 	virtual void GetMenuActions(FBlueprintActionDatabaseRegistrar& ActionRegistrar) const override;
-	virtual bool IsActionFilteredOut(class FBlueprintActionFilter const& Filter) override;
 	virtual EAnimAssetHandlerType SupportsAssetClass(const UClass* AssetClass) const override;
 	// End of UAnimGraphNode_Base
 
diff --git a/Engine/Source/Editor/AnimGraph/Public/AnimGraphNode_SequencePlayer.h b/Engine/Source/Editor/AnimGraph/Public/AnimGraphNode_SequencePlayer.h
--- a/Engine/Source/Editor/AnimGraph/Public/AnimGraphNode_SequencePlayer.h
+++ b/Engine/Source/Editor/AnimGraph/Public/AnimGraphNode_SequencePlayer.h
@@ -43,7 +43,6 @@ class UAnimGraphNode_SequencePlayer : public UAnimGraphNode_AssetPlayerBase
 	virtual void GetAllAnimationSequencesReferred(TArray<UAnimationAsset*>& AnimationAssets) const override;
 	virtual void ReplaceReferredAnimations(const TMap<UAnimationAsset*, UAnimationAsset*>& AnimAssetReplacementMap) override;
 	virtual void GetMenuActions(FBlueprintActionDatabaseRegistrar& ActionRegistrar) const override;
-	virtual bool IsActionFilteredOut(class FBlueprintActionFilter const& Filter) override;
 	virtual EAnimAssetHandlerType SupportsAssetClass(const UClass* AssetClass) const override;
 	virtual void OnOverrideAssets(IAnimBlueprintNodeOverrideAssetsContext& InContext) const override;
 	virtual void GetOutputLinkAttributes(FNodeAttributeArray& OutAttributes) const override;
