Fixed being unable to open newly-spawned empty blend sapce graphs

diff --git a/Engine/Source/Editor/AnimGraph/Private/AnimGraphNode_AssetPlayerBase.cpp b/Engine/Source/Editor/AnimGraph/Private/AnimGraphNode_AssetPlayerBase.cpp
--- a/Engine/Source/Editor/AnimGraph/Private/AnimGraphNode_AssetPlayerBase.cpp
+++ b/Engine/Source/Editor/AnimGraph/Private/AnimGraphNode_AssetPlayerBase.cpp
@@ -237,22 +237,39 @@ void UAnimGraphNode_AssetPlayerBase::SetupNewNode(UEdGraphNode* InNewNode, bool
 	}
 }
 
-void UAnimGraphNode_AssetPlayerBase::GetMenuActionsHelper(FBlueprintActionDatabaseRegistrar& InActionRegistrar, TSubclassOf<UAnimGraphNode_Base> InNodeClass, const TArray<TSubclassOf<UObject>>& InAssetTypes, const TArray<TSubclassOf<UObject>>& InExcludedAssetTypes, const TFunctionRef<FText(const FAssetData&)>& InMenuNameFunction, const TFunctionRef<FText(const FAssetData&)>& InMenuTooltipFunction, const TFunction<void(UEdGraphNode*, bool, const FAssetData)>& InSetupNewNodeFunction)
+void UAnimGraphNode_AssetPlayerBase::GetMenuActionsHelper(FBlueprintActionDatabaseRegistrar& InActionRegistrar, TSubclassOf<UAnimGraphNode_Base> InNodeClass, const TArray<TSubclassOf<UObject>>& InAssetTypes, const TArray<TSubclassOf<UObject>>& InExcludedAssetTypes, const TFunctionRef<FText(const FAssetData&, UClass*)>& InMenuNameFunction, const TFunctionRef<FText(const FAssetData&, UClass*)>& InMenuTooltipFunction, const TFunction<void(UEdGraphNode*, bool, const FAssetData)>& InSetupNewNodeFromAssetFunction, const TFunction<void(UEdGraphNode*, bool, TSubclassOf<UObject>)>& InSetupNewNodeFromClassFunction)
 {
-	auto MakeAction = [&InActionRegistrar, &InMenuNameFunction, &InMenuTooltipFunction, InSetupNewNodeFunction, InNodeClass](const FAssetData& InAssetData)
+	auto MakeActionFromAsset = [&InActionRegistrar, &InMenuNameFunction, &InMenuTooltipFunction, InSetupNewNodeFromAssetFunction, InNodeClass](const FAssetData& InAssetData)
 	{
-		auto AssetSetup = [InSetupNewNodeFunction, InAssetData](UEdGraphNode* InNewNode, bool bInIsTemplateNode)
+		auto AssetSetup = [InSetupNewNodeFromAssetFunction, InAssetData](UEdGraphNode* InNewNode, bool bInIsTemplateNode)
 		{
-			InSetupNewNodeFunction(InNewNode, bInIsTemplateNode, InAssetData);
+			InSetupNewNodeFromAssetFunction(InNewNode, bInIsTemplateNode, InAssetData);
 		};
 
 		UBlueprintNodeSpawner* NodeSpawner = UBlueprintNodeSpawner::Create(InNodeClass.Get());
 		NodeSpawner->CustomizeNodeDelegate = UBlueprintNodeSpawner::FCustomizeNodeDelegate::CreateLambda(AssetSetup);
-		NodeSpawner->DefaultMenuSignature.MenuName = InMenuNameFunction(InAssetData);
-		NodeSpawner->DefaultMenuSignature.Tooltip = InMenuTooltipFunction(InAssetData);
+		NodeSpawner->DefaultMenuSignature.MenuName = InMenuNameFunction(InAssetData, InAssetData.GetClass());
+		NodeSpawner->DefaultMenuSignature.Tooltip = InMenuTooltipFunction(InAssetData, InAssetData.GetClass());
 		InActionRegistrar.AddBlueprintAction(InAssetData, NodeSpawner);
-	};	
+	};
 
+	auto MakeActionFromClass = [&InActionRegistrar, &InMenuNameFunction, &InMenuTooltipFunction, InSetupNewNodeFromClassFunction, InNodeClass](const TSubclassOf<UObject>& InAssetClass)
+	{
+		if(InSetupNewNodeFromClassFunction != nullptr)
+		{
+			auto AssetSetup = [InSetupNewNodeFromClassFunction, InAssetClass](UEdGraphNode* InNewNode, bool bInIsTemplateNode)
+			{
+				InSetupNewNodeFromClassFunction(InNewNode, bInIsTemplateNode, InAssetClass);
+			};
+
+			UBlueprintNodeSpawner* NodeSpawner = UBlueprintNodeSpawner::Create(InNodeClass.Get());
+			NodeSpawner->CustomizeNodeDelegate = UBlueprintNodeSpawner::FCustomizeNodeDelegate::CreateLambda(AssetSetup);
+			NodeSpawner->DefaultMenuSignature.MenuName = InMenuNameFunction(FAssetData(), InAssetClass.Get());
+			NodeSpawner->DefaultMenuSignature.Tooltip = InMenuTooltipFunction(FAssetData(), InAssetClass.Get());
+			InActionRegistrar.AddBlueprintAction(InAssetClass, NodeSpawner);
+		}
+	};
+	
 	const UObject* QueryObject = InActionRegistrar.GetActionKeyFilter();
 	bool bIsObjectOfAssetType = false;
 	for(const TSubclassOf<UObject>& AssetType : InAssetTypes)
@@ -286,16 +303,27 @@ void UAnimGraphNode_AssetPlayerBase::GetMenuActionsHelper(FBlueprintActionDataba
 		{
 			if(AssetData.IsUAsset())
 			{
-				MakeAction(AssetData);
+				MakeActionFromAsset(AssetData);
 			}
 		}
 
-		// Add 'empty' asset node
-		MakeAction(FAssetData());
+		if(InSetupNewNodeFromClassFunction != nullptr)
+		{
+			// Add 'class' nodes
+			for(const TSubclassOf<UObject>& AssetType : InAssetTypes)
+			{
+				MakeActionFromClass(AssetType);
+			}
+		}
+		else
+		{
+			// Add 'empty' asset node
+			MakeActionFromAsset(FAssetData());
+		}
 	}
 	else if (bIsObjectOfAssetType)
 	{
-		MakeAction(FAssetData(QueryObject));
+		MakeActionFromAsset(FAssetData(QueryObject));
 	}
 }
 
diff --git a/Engine/Source/Editor/AnimGraph/Private/AnimGraphNode_BlendSpaceEvaluator.cpp b/Engine/Source/Editor/AnimGraph/Private/AnimGraphNode_BlendSpaceEvaluator.cpp
--- a/Engine/Source/Editor/AnimGraph/Private/AnimGraphNode_BlendSpaceEvaluator.cpp
+++ b/Engine/Source/Editor/AnimGraph/Private/AnimGraphNode_BlendSpaceEvaluator.cpp
@@ -33,7 +33,7 @@ void UAnimGraphNode_BlendSpaceEvaluator::GetMenuActions(FBlueprintActionDatabase
 		GetClass(),
 		{ UBlendSpace::StaticClass() },
 		{ },
-		[](const FAssetData& InAssetData)
+		[](const FAssetData& InAssetData, UClass* InClass)
 		{
 			if(InAssetData.IsValid())
 			{
@@ -44,7 +44,7 @@ void UAnimGraphNode_BlendSpaceEvaluator::GetMenuActions(FBlueprintActionDatabase
 				return LOCTEXT("MenuDesc", "Blendspace Evaluator");
 			}
 		},
-		[](const FAssetData& InAssetData)
+		[](const FAssetData& InAssetData, UClass* InClass)
 		{
 			if(InAssetData.IsValid())
 			{
diff --git a/Engine/Source/Editor/AnimGraph/Private/AnimGraphNode_BlendSpaceGraph.cpp b/Engine/Source/Editor/AnimGraph/Private/AnimGraphNode_BlendSpaceGraph.cpp
--- a/Engine/Source/Editor/AnimGraph/Private/AnimGraphNode_BlendSpaceGraph.cpp
+++ b/Engine/Source/Editor/AnimGraph/Private/AnimGraphNode_BlendSpaceGraph.cpp
@@ -61,23 +61,31 @@ void UAnimGraphNode_BlendSpaceGraph::GetMenuActions(FBlueprintActionDatabaseRegi
 		GetClass(),
 		{ UBlendSpace::StaticClass(), UBlendSpace1D::StaticClass() },
 		{ UAimOffsetBlendSpace::StaticClass(), UAimOffsetBlendSpace1D::StaticClass() },
-		[](const FAssetData& InAssetData)
+		[](const FAssetData& InAssetData, UClass* InClass)
 		{
 			if(InAssetData.IsValid())
 			{
 				return FText::Format(LOCTEXT("MenuDescFormat", "Blendspace '{0}'"), FText::FromName(InAssetData.AssetName));
 			}
+			else if(InClass != nullptr)
+			{
+				return InClass->GetDisplayNameText();
+			}
 			else
 			{
 				return LOCTEXT("MenuDesc", "Blendspace");
 			}
 		},
-		[](const FAssetData& InAssetData)
+		[](const FAssetData& InAssetData, UClass* InClass)
 		{
 			if(InAssetData.IsValid())
 			{
 				return FText::Format(LOCTEXT("MenuDescTooltipFormat", "Blendspace\n'{0}'"), FText::FromName(InAssetData.ObjectPath));
 			}
+			else if(InClass != nullptr)
+			{
+				return InClass->GetDisplayNameText();
+			}	
 			else
 			{
 				return LOCTEXT("MenuDesc", "Blendspace");
@@ -87,6 +95,11 @@ void UAnimGraphNode_BlendSpaceGraph::GetMenuActions(FBlueprintActionDatabaseRegi
 		{
 			UAnimGraphNode_BlendSpaceGraph* GraphNode = CastChecked<UAnimGraphNode_BlendSpaceGraph>(InNewNode);
 			GraphNode->SetupFromAsset(InAssetData, bInIsTemplateNode);
+		},
+		[](UEdGraphNode* InNewNode, bool bInIsTemplateNode, const TSubclassOf<UObject> InClass)
+		{
+			UAnimGraphNode_BlendSpaceGraph* GraphNode = CastChecked<UAnimGraphNode_BlendSpaceGraph>(InNewNode);
+			GraphNode->SetupFromClass(InClass.Get(), bInIsTemplateNode);
 		});
 }
 
diff --git a/Engine/Source/Editor/AnimGraph/Private/AnimGraphNode_BlendSpacePlayer.cpp b/Engine/Source/Editor/AnimGraph/Private/AnimGraphNode_BlendSpacePlayer.cpp
--- a/Engine/Source/Editor/AnimGraph/Private/AnimGraphNode_BlendSpacePlayer.cpp
+++ b/Engine/Source/Editor/AnimGraph/Private/AnimGraphNode_BlendSpacePlayer.cpp
@@ -65,7 +65,7 @@ void UAnimGraphNode_BlendSpacePlayer::GetMenuActions(FBlueprintActionDatabaseReg
 		GetClass(),
 		{ UBlendSpace::StaticClass() },
 		{ },
-		[](const FAssetData& InAssetData)
+		[](const FAssetData& InAssetData, UClass* InClass)
 		{
 			if(InAssetData.IsValid())
 			{
@@ -76,7 +76,7 @@ void UAnimGraphNode_BlendSpacePlayer::GetMenuActions(FBlueprintActionDatabaseReg
 				return LOCTEXT("MenuDesc", "Blendspace Player");
 			}
 		},
-		[](const FAssetData& InAssetData)
+		[](const FAssetData& InAssetData, UClass* InClass)
 		{
 			if(InAssetData.IsValid())
 			{
diff --git a/Engine/Source/Editor/AnimGraph/Private/AnimGraphNode_LinkedAnimGraph.cpp b/Engine/Source/Editor/AnimGraph/Private/AnimGraphNode_LinkedAnimGraph.cpp
--- a/Engine/Source/Editor/AnimGraph/Private/AnimGraphNode_LinkedAnimGraph.cpp
+++ b/Engine/Source/Editor/AnimGraph/Private/AnimGraphNode_LinkedAnimGraph.cpp
@@ -76,7 +76,7 @@ void UAnimGraphNode_LinkedAnimGraph::GetMenuActions(FBlueprintActionDatabaseRegi
 		GetClass(),
 		{ UAnimBlueprint::StaticClass()},
 		{ },
-		[](const FAssetData& InAssetData)
+		[](const FAssetData& InAssetData, UClass* InClass)
 		{
 			if(InAssetData.IsValid())
 			{
@@ -87,7 +87,7 @@ void UAnimGraphNode_LinkedAnimGraph::GetMenuActions(FBlueprintActionDatabaseRegi
 				return LOCTEXT("MenuDesc", "Linked Anim Graph");
 			}
 		},
-		[](const FAssetData& InAssetData)
+		[](const FAssetData& InAssetData, UClass* InClass)
 		{
 			if(InAssetData.IsValid())
 			{
diff --git a/Engine/Source/Editor/AnimGraph/Private/AnimGraphNode_PoseBlendNode.cpp b/Engine/Source/Editor/AnimGraph/Private/AnimGraphNode_PoseBlendNode.cpp
--- a/Engine/Source/Editor/AnimGraph/Private/AnimGraphNode_PoseBlendNode.cpp
+++ b/Engine/Source/Editor/AnimGraph/Private/AnimGraphNode_PoseBlendNode.cpp
@@ -104,7 +104,7 @@ void UAnimGraphNode_PoseBlendNode::GetMenuActions(FBlueprintActionDatabaseRegist
 		GetClass(),
 		{ UPoseAsset::StaticClass() },
 		{ },
-		[](const FAssetData& InAssetData)
+		[](const FAssetData& InAssetData, UClass* InClass)
 		{
 			if(InAssetData.IsValid())
 			{
@@ -115,7 +115,7 @@ void UAnimGraphNode_PoseBlendNode::GetMenuActions(FBlueprintActionDatabaseRegist
 				return LOCTEXT("PlayerDesc", "Evaluate Pose");
 			}
 		},
-		[](const FAssetData& InAssetData)
+		[](const FAssetData& InAssetData, UClass* InClass)
 		{
 			if(InAssetData.IsValid())
 			{
diff --git a/Engine/Source/Editor/AnimGraph/Private/AnimGraphNode_RotationOffsetBlendSpace.cpp b/Engine/Source/Editor/AnimGraph/Private/AnimGraphNode_RotationOffsetBlendSpace.cpp
--- a/Engine/Source/Editor/AnimGraph/Private/AnimGraphNode_RotationOffsetBlendSpace.cpp
+++ b/Engine/Source/Editor/AnimGraph/Private/AnimGraphNode_RotationOffsetBlendSpace.cpp
@@ -77,7 +77,7 @@ void UAnimGraphNode_RotationOffsetBlendSpace::GetMenuActions(FBlueprintActionDat
 		GetClass(),
 		{ UAimOffsetBlendSpace::StaticClass(), UAimOffsetBlendSpace1D::StaticClass() },
 		{ },
-		[](const FAssetData& InAssetData)
+		[](const FAssetData& InAssetData, UClass* InClass)
 		{
 			if(InAssetData.IsValid())
 			{
@@ -88,7 +88,7 @@ void UAnimGraphNode_RotationOffsetBlendSpace::GetMenuActions(FBlueprintActionDat
 				return LOCTEXT("MenuDesc", "AimOffset Player");
 			}
 		},
-		[](const FAssetData& InAssetData)
+		[](const FAssetData& InAssetData, UClass* InClass)
 		{
 			if(InAssetData.IsValid())
 			{
diff --git a/Engine/Source/Editor/AnimGraph/Private/AnimGraphNode_RotationOffsetBlendSpaceGraph.cpp b/Engine/Source/Editor/AnimGraph/Private/AnimGraphNode_RotationOffsetBlendSpaceGraph.cpp
--- a/Engine/Source/Editor/AnimGraph/Private/AnimGraphNode_RotationOffsetBlendSpaceGraph.cpp
+++ b/Engine/Source/Editor/AnimGraph/Private/AnimGraphNode_RotationOffsetBlendSpaceGraph.cpp
@@ -63,7 +63,7 @@ void UAnimGraphNode_RotationOffsetBlendSpaceGraph::GetMenuActions(FBlueprintActi
 		GetClass(),
 		{ UAimOffsetBlendSpace::StaticClass(), UAimOffsetBlendSpace1D::StaticClass() },
 		{ },
-		[](const FAssetData& InAssetData)
+		[](const FAssetData& InAssetData, UClass* InClass)
 		{
 			if(InAssetData.IsValid())
 			{
@@ -74,7 +74,7 @@ void UAnimGraphNode_RotationOffsetBlendSpaceGraph::GetMenuActions(FBlueprintActi
 				return LOCTEXT("MenuDesc", "AimOffset");
 			}
 		},
-		[](const FAssetData& InAssetData)
+		[](const FAssetData& InAssetData, UClass* InClass)
 		{
 			if(InAssetData.IsValid())
 			{
diff --git a/Engine/Source/Editor/AnimGraph/Private/AnimGraphNode_SequenceEvaluator.cpp b/Engine/Source/Editor/AnimGraph/Private/AnimGraphNode_SequenceEvaluator.cpp
--- a/Engine/Source/Editor/AnimGraph/Private/AnimGraphNode_SequenceEvaluator.cpp
+++ b/Engine/Source/Editor/AnimGraph/Private/AnimGraphNode_SequenceEvaluator.cpp
@@ -75,7 +75,7 @@ void UAnimGraphNode_SequenceEvaluator::GetMenuActions(FBlueprintActionDatabaseRe
 		GetClass(),
 		{ UAnimSequence::StaticClass() },
 		{ },
-		[](const FAssetData& InAssetData)
+		[](const FAssetData& InAssetData, UClass* InClass)
 		{
 			if(InAssetData.IsValid())
 			{
@@ -94,7 +94,7 @@ void UAnimGraphNode_SequenceEvaluator::GetMenuActions(FBlueprintActionDatabaseRe
 				return LOCTEXT("MenuDesc", "Sequence Evaluator");
 			}
 		},
-		[](const FAssetData& InAssetData)
+		[](const FAssetData& InAssetData, UClass* InClass)
 		{
 			if(InAssetData.IsValid())
 			{
diff --git a/Engine/Source/Editor/AnimGraph/Private/AnimGraphNode_SequencePlayer.cpp b/Engine/Source/Editor/AnimGraph/Private/AnimGraphNode_SequencePlayer.cpp
--- a/Engine/Source/Editor/AnimGraph/Private/AnimGraphNode_SequencePlayer.cpp
+++ b/Engine/Source/Editor/AnimGraph/Private/AnimGraphNode_SequencePlayer.cpp
@@ -92,7 +92,7 @@ void UAnimGraphNode_SequencePlayer::GetMenuActions(FBlueprintActionDatabaseRegis
 		GetClass(),
 		{ UAnimSequence::StaticClass() },
 		{ },
-		[](const FAssetData& InAssetData)
+		[](const FAssetData& InAssetData, UClass* InClass)
 		{
 			if(InAssetData.IsValid())
 			{
@@ -111,7 +111,7 @@ void UAnimGraphNode_SequencePlayer::GetMenuActions(FBlueprintActionDatabaseRegis
 				return LOCTEXT("PlayerDesc", "Sequence Player");
 			}
 		},
-		[](const FAssetData& InAssetData)
+		[](const FAssetData& InAssetData, UClass* InClass)
 		{
 			if(InAssetData.IsValid())
 			{
diff --git a/Engine/Source/Editor/AnimGraph/Public/AnimGraphNode_AssetPlayerBase.h b/Engine/Source/Editor/AnimGraph/Public/AnimGraphNode_AssetPlayerBase.h
--- a/Engine/Source/Editor/AnimGraph/Public/AnimGraphNode_AssetPlayerBase.h
+++ b/Engine/Source/Editor/AnimGraph/Public/AnimGraphNode_AssetPlayerBase.h
@@ -44,7 +44,7 @@ public:
 	virtual void SetAnimationAsset(UAnimationAsset* Asset) { check(false); /*Base function called*/ }
 
 	// Helper function to gather menu actions from specific asset types supported by this node
-	static void GetMenuActionsHelper(FBlueprintActionDatabaseRegistrar& InActionRegistrar, TSubclassOf<UAnimGraphNode_Base> InNodeClass, const TArray<TSubclassOf<UObject>>& InAssetTypes, const TArray<TSubclassOf<UObject>>& InExcludedAssetTypes, const TFunctionRef<FText(const FAssetData&)>& InMenuNameFunction, const TFunctionRef<FText(const FAssetData&)>& InMenuTooltipFunction, const TFunction<void(UEdGraphNode*, bool, const FAssetData)>& InSetupNewNodeFunction);
+	static void GetMenuActionsHelper(FBlueprintActionDatabaseRegistrar& InActionRegistrar, TSubclassOf<UAnimGraphNode_Base> InNodeClass, const TArray<TSubclassOf<UObject>>& InAssetTypes, const TArray<TSubclassOf<UObject>>& InExcludedAssetTypes, const TFunctionRef<FText(const FAssetData&, UClass*)>& InMenuNameFunction, const TFunctionRef<FText(const FAssetData&, UClass*)>& InMenuTooltipFunction, const TFunction<void(UEdGraphNode*, bool, const FAssetData)>& InSetupNewNodeFromAssetFunction, const TFunction<void(UEdGraphNode*, bool, TSubclassOf<UObject>)>& InSetupNewNodeFromClassFunction = nullptr);
 
 	// Helper function to validate player nodes
 	void ValidateAnimNodeDuringCompilationHelper(USkeleton* ForSkeleton, FCompilerResultsLog& MessageLog, UAnimationAsset* InAsset, TSubclassOf<UAnimationAsset> InAssetType, UEdGraphPin* InExposedPin);
