Fix error reporting for dynamic values that are not exposed on pins

diff --git a/Engine/Source/Editor/AnimGraph/Private/AnimGraphNode_AssetPlayerBase.cpp b/Engine/Source/Editor/AnimGraph/Private/AnimGraphNode_AssetPlayerBase.cpp
--- a/Engine/Source/Editor/AnimGraph/Private/AnimGraphNode_AssetPlayerBase.cpp
+++ b/Engine/Source/Editor/AnimGraph/Private/AnimGraphNode_AssetPlayerBase.cpp
@@ -479,7 +479,7 @@ FText UAnimGraphNode_AssetPlayerBase::GetNodeTitleForAsset(ENodeTitleType::Type
 	}
 }
 
-void UAnimGraphNode_AssetPlayerBase::ValidateAnimNodeDuringCompilationHelper(USkeleton* ForSkeleton, FCompilerResultsLog& MessageLog, UAnimationAsset* InAsset, TSubclassOf<UAnimationAsset> InAssetType, UEdGraphPin* InExposedPin)
+void UAnimGraphNode_AssetPlayerBase::ValidateAnimNodeDuringCompilationHelper(USkeleton* ForSkeleton, FCompilerResultsLog& MessageLog, UAnimationAsset* InAsset, TSubclassOf<UAnimationAsset> InAssetType, UEdGraphPin* InExposedPin, FName InPropertyName)
 {
 	UAnimationAsset* AssetToCheck = InAsset;
 	if(InExposedPin != nullptr && AssetToCheck == nullptr)
@@ -500,13 +500,13 @@ void UAnimGraphNode_AssetPlayerBase::ValidateAnimNodeDuringCompilationHelper(USk
 				{
 					bHasBinding = true;
 				}
-
-				if(AlwaysDynamicProperties.Contains(InExposedPin->GetFName()))
-				{
-					bAlwaysDynamic = true;
-				}
 			}
 
+			if(AlwaysDynamicProperties.Contains(InPropertyName))
+			{
+				bAlwaysDynamic = true;
+			}
+			
 			// we may have a connected node or binding
 			if(!bAlwaysDynamic)
 			{
diff --git a/Engine/Source/Editor/AnimGraph/Private/AnimGraphNode_BlendSpaceEvaluator.cpp b/Engine/Source/Editor/AnimGraph/Private/AnimGraphNode_BlendSpaceEvaluator.cpp
--- a/Engine/Source/Editor/AnimGraph/Private/AnimGraphNode_BlendSpaceEvaluator.cpp
+++ b/Engine/Source/Editor/AnimGraph/Private/AnimGraphNode_BlendSpaceEvaluator.cpp
@@ -65,7 +65,7 @@ void UAnimGraphNode_BlendSpaceEvaluator::ValidateAnimNodeDuringCompilation(class
 {
 	Super::ValidateAnimNodeDuringCompilation(ForSkeleton, MessageLog);
 
-	ValidateAnimNodeDuringCompilationHelper(ForSkeleton, MessageLog, Node.GetBlendSpace(), UBlendSpace::StaticClass(), FindPin(GET_MEMBER_NAME_STRING_CHECKED(FAnimNode_BlendSpaceEvaluator, BlendSpace)));
+	ValidateAnimNodeDuringCompilationHelper(ForSkeleton, MessageLog, Node.GetBlendSpace(), UBlendSpace::StaticClass(), FindPin(GET_MEMBER_NAME_STRING_CHECKED(FAnimNode_BlendSpaceEvaluator, BlendSpace)), GET_MEMBER_NAME_CHECKED(FAnimNode_BlendSpaceEvaluator, BlendSpace));
 }
 
 void UAnimGraphNode_BlendSpaceEvaluator::BakeDataDuringCompilation(class FCompilerResultsLog& MessageLog)
diff --git a/Engine/Source/Editor/AnimGraph/Private/AnimGraphNode_BlendSpacePlayer.cpp b/Engine/Source/Editor/AnimGraph/Private/AnimGraphNode_BlendSpacePlayer.cpp
--- a/Engine/Source/Editor/AnimGraph/Private/AnimGraphNode_BlendSpacePlayer.cpp
+++ b/Engine/Source/Editor/AnimGraph/Private/AnimGraphNode_BlendSpacePlayer.cpp
@@ -35,7 +35,7 @@ void UAnimGraphNode_BlendSpacePlayer::ValidateAnimNodeDuringCompilation(class US
 {
 	Super::ValidateAnimNodeDuringCompilation(ForSkeleton, MessageLog);
 
-	ValidateAnimNodeDuringCompilationHelper(ForSkeleton, MessageLog, Node.GetBlendSpace(), UBlendSpace::StaticClass(), FindPin(GET_MEMBER_NAME_STRING_CHECKED(FAnimNode_BlendSpacePlayer, BlendSpace)));
+	ValidateAnimNodeDuringCompilationHelper(ForSkeleton, MessageLog, Node.GetBlendSpace(), UBlendSpace::StaticClass(), FindPin(GET_MEMBER_NAME_STRING_CHECKED(FAnimNode_BlendSpacePlayer, BlendSpace)), GET_MEMBER_NAME_CHECKED(FAnimNode_BlendSpacePlayer, BlendSpace));
 }
 
 void UAnimGraphNode_BlendSpacePlayer::BakeDataDuringCompilation(class FCompilerResultsLog& MessageLog)
diff --git a/Engine/Source/Editor/AnimGraph/Private/AnimGraphNode_PoseHandler.cpp b/Engine/Source/Editor/AnimGraph/Private/AnimGraphNode_PoseHandler.cpp
--- a/Engine/Source/Editor/AnimGraph/Private/AnimGraphNode_PoseHandler.cpp
+++ b/Engine/Source/Editor/AnimGraph/Private/AnimGraphNode_PoseHandler.cpp
@@ -16,7 +16,7 @@ void UAnimGraphNode_PoseHandler::ValidateAnimNodeDuringCompilation(USkeleton* Fo
 {
 	Super::ValidateAnimNodeDuringCompilation(ForSkeleton, MessageLog);
 	
-	ValidateAnimNodeDuringCompilationHelper(ForSkeleton, MessageLog, GetPoseHandlerNode()->PoseAsset, UPoseAsset::StaticClass(), FindPin(GET_MEMBER_NAME_STRING_CHECKED(FAnimNode_PoseHandler, PoseAsset)));
+	ValidateAnimNodeDuringCompilationHelper(ForSkeleton, MessageLog, GetPoseHandlerNode()->PoseAsset, UPoseAsset::StaticClass(), FindPin(GET_MEMBER_NAME_STRING_CHECKED(FAnimNode_PoseHandler, PoseAsset)), GET_MEMBER_NAME_CHECKED(FAnimNode_PoseHandler, PoseAsset));
 }
 
 void UAnimGraphNode_PoseHandler::SetAnimationAsset(UAnimationAsset* Asset)
diff --git a/Engine/Source/Editor/AnimGraph/Private/AnimGraphNode_SequenceEvaluator.cpp b/Engine/Source/Editor/AnimGraph/Private/AnimGraphNode_SequenceEvaluator.cpp
--- a/Engine/Source/Editor/AnimGraph/Private/AnimGraphNode_SequenceEvaluator.cpp
+++ b/Engine/Source/Editor/AnimGraph/Private/AnimGraphNode_SequenceEvaluator.cpp
@@ -143,7 +143,7 @@ void UAnimGraphNode_SequenceEvaluator::ValidateAnimNodeDuringCompilation(class U
 {
 	Super::ValidateAnimNodeDuringCompilation(ForSkeleton, MessageLog);
 
-	ValidateAnimNodeDuringCompilationHelper(ForSkeleton, MessageLog, Node.GetSequence(), UAnimSequenceBase::StaticClass(), FindPin(GET_MEMBER_NAME_STRING_CHECKED(FAnimNode_SequenceEvaluator, Sequence)));
+	ValidateAnimNodeDuringCompilationHelper(ForSkeleton, MessageLog, Node.GetSequence(), UAnimSequenceBase::StaticClass(), FindPin(GET_MEMBER_NAME_STRING_CHECKED(FAnimNode_SequenceEvaluator, Sequence)), GET_MEMBER_NAME_CHECKED(FAnimNode_SequenceEvaluator, Sequence));
 }
 
 void UAnimGraphNode_SequenceEvaluator::GetNodeContextMenuActions(UToolMenu* Menu, UGraphNodeContextMenuContext* Context) const
diff --git a/Engine/Source/Editor/AnimGraph/Private/AnimGraphNode_SequencePlayer.cpp b/Engine/Source/Editor/AnimGraph/Private/AnimGraphNode_SequencePlayer.cpp
--- a/Engine/Source/Editor/AnimGraph/Private/AnimGraphNode_SequencePlayer.cpp
+++ b/Engine/Source/Editor/AnimGraph/Private/AnimGraphNode_SequencePlayer.cpp
@@ -162,7 +162,7 @@ void UAnimGraphNode_SequencePlayer::ValidateAnimNodeDuringCompilation(class USke
 {
 	Super::ValidateAnimNodeDuringCompilation(ForSkeleton, MessageLog);
 
-	ValidateAnimNodeDuringCompilationHelper(ForSkeleton, MessageLog, Node.GetSequence(), UAnimSequenceBase::StaticClass(), FindPin(GET_MEMBER_NAME_STRING_CHECKED(FAnimNode_SequencePlayer, Sequence)));
+	ValidateAnimNodeDuringCompilationHelper(ForSkeleton, MessageLog, Node.GetSequence(), UAnimSequenceBase::StaticClass(), FindPin(GET_MEMBER_NAME_STRING_CHECKED(FAnimNode_SequencePlayer, Sequence)), GET_MEMBER_NAME_CHECKED(FAnimNode_SequencePlayer, Sequence));
 }
 
 void UAnimGraphNode_SequencePlayer::GetNodeContextMenuActions(UToolMenu* Menu, UGraphNodeContextMenuContext* Context) const
diff --git a/Engine/Source/Editor/AnimGraph/Public/AnimGraphNode_AssetPlayerBase.h b/Engine/Source/Editor/AnimGraph/Public/AnimGraphNode_AssetPlayerBase.h
--- a/Engine/Source/Editor/AnimGraph/Public/AnimGraphNode_AssetPlayerBase.h
+++ b/Engine/Source/Editor/AnimGraph/Public/AnimGraphNode_AssetPlayerBase.h
@@ -47,7 +47,7 @@ public:
 	static void GetMenuActionsHelper(FBlueprintActionDatabaseRegistrar& InActionRegistrar, TSubclassOf<UAnimGraphNode_Base> InNodeClass, const TArray<TSubclassOf<UObject>>& InAssetTypes, const TArray<TSubclassOf<UObject>>& InExcludedAssetTypes, const TFunctionRef<FText(const FAssetData&, UClass*)>& InMenuNameFunction, const TFunctionRef<FText(const FAssetData&, UClass*)>& InMenuTooltipFunction, const TFunction<void(UEdGraphNode*, bool, const FAssetData)>& InSetupNewNodeFromAssetFunction, const TFunction<void(UEdGraphNode*, bool, TSubclassOf<UObject>)>& InSetupNewNodeFromClassFunction = nullptr);
 
 	// Helper function to validate player nodes
-	void ValidateAnimNodeDuringCompilationHelper(USkeleton* ForSkeleton, FCompilerResultsLog& MessageLog, UAnimationAsset* InAsset, TSubclassOf<UAnimationAsset> InAssetType, UEdGraphPin* InExposedPin);
+	void ValidateAnimNodeDuringCompilationHelper(USkeleton* ForSkeleton, FCompilerResultsLog& MessageLog, UAnimationAsset* InAsset, TSubclassOf<UAnimationAsset> InAssetType, UEdGraphPin* InExposedPin, FName InPropertyName);
 
 protected:
 	// Helper functions to build a title for an asset player node
