Ported project specific scriptability UFunctions to Engine:
- SetPreviewSkeletalMesh to UAnimationAsset and UAnimBlueprint
- Functionality to retrieve UAnimationGraphs and UAnimationGraphNodes of a specific class from UAnimBlueprint and UAnimationGraph respectively
- Functionality for pose driver/asset scripting

diff --git a/Engine/Source/Editor/AnimGraph/Private/AnimGraphNode_PoseDriver.cpp b/Engine/Source/Editor/AnimGraph/Private/AnimGraphNode_PoseDriver.cpp
--- a/Engine/Source/Editor/AnimGraph/Private/AnimGraphNode_PoseDriver.cpp
+++ b/Engine/Source/Editor/AnimGraph/Private/AnimGraphNode_PoseDriver.cpp
@@ -396,6 +396,80 @@ void UAnimGraphNode_PoseDriver::CopyTargetsFromPoseAsset()
 	}
 }
 
+void UAnimGraphNode_PoseDriver::SetSourceBones(const TArray<FName>& BoneNames)
+{
+	Node.SourceBones.Empty(BoneNames.Num());
+	for (const FName& BoneName : BoneNames)
+	{
+		Node.SourceBones.Add(BoneName);
+	}
+}
+
+void UAnimGraphNode_PoseDriver::GetSourceBoneNames(TArray<FName>& BoneNames)
+{
+	for (const FBoneReference& SourceBone : Node.SourceBones)
+	{
+		BoneNames.Add(SourceBone.BoneName);
+	}
+}
+
+void UAnimGraphNode_PoseDriver::SetDrivingBones(const TArray<FName>& BoneNames)
+{
+	Node.OnlyDriveBones.Empty(BoneNames.Num());
+	for (const FName& BoneName : BoneNames)
+	{
+		Node.OnlyDriveBones.Add(BoneName);
+	}
+}
+
+void UAnimGraphNode_PoseDriver::GetDrivingBoneNames(TArray<FName>& BoneNames)
+{
+	for (const FBoneReference& SourceBone : Node.OnlyDriveBones)
+	{
+		BoneNames.Add(SourceBone.BoneName);
+	}
+}
+
+void UAnimGraphNode_PoseDriver::SetRBFParameters(FRBFParams Parameters)
+{
+	Node.RBFParams = Parameters;
+}
+
+FRBFParams& UAnimGraphNode_PoseDriver::GetRBFParameters()
+{
+	return Node.RBFParams;
+}
+
+void UAnimGraphNode_PoseDriver::SetPoseDriverSource(EPoseDriverSource DriverSource)
+{
+	Node.DriveSource = DriverSource;	
+}
+
+EPoseDriverSource& UAnimGraphNode_PoseDriver::GetPoseDriverSource()
+{
+	return Node.DriveSource;
+}
+
+void UAnimGraphNode_PoseDriver::SetPoseDriverOutput(EPoseDriverOutput DriverOutput)
+{
+	Node.DriveOutput = DriverOutput;
+}
+
+EPoseDriverOutput& UAnimGraphNode_PoseDriver::GetPoseDriverOutput()
+{
+	return Node.DriveOutput;
+}
+
+void UAnimGraphNode_PoseDriver::SetOnlyDriveSelectedBones(bool bOnlyDriveSelectedBones)
+{
+	Node.bOnlyDriveSelectedBones = bOnlyDriveSelectedBones;
+}
+
+bool UAnimGraphNode_PoseDriver::GetOnlyDriveSelectedBones()
+{
+	return Node.bOnlyDriveSelectedBones == 1;
+}
+
 void UAnimGraphNode_PoseDriver::AddNewTarget()
 {
 	FPoseDriverTarget& NewTarget = Node.PoseTargets[Node.PoseTargets.Add(FPoseDriverTarget())];
diff --git a/Engine/Source/Editor/AnimGraph/Private/AnimationGraph.cpp b/Engine/Source/Editor/AnimGraph/Private/AnimationGraph.cpp
--- a/Engine/Source/Editor/AnimGraph/Private/AnimationGraph.cpp
+++ b/Engine/Source/Editor/AnimGraph/Private/AnimationGraph.cpp
@@ -12,5 +12,20 @@ UAnimationGraph::UAnimationGraph(const FObjectInitializer& ObjectInitializer)
 {
 }
 
+void UAnimationGraph::GetGraphNodesOfClass(TSubclassOf<UAnimGraphNode_Base> NodeClass, TArray<UAnimGraphNode_Base*>& GraphNodes, bool bIncludeChildClasses /*= true*/)
+{
+	for (int32 NodeIndex = 0; NodeIndex < Nodes.Num(); ++NodeIndex)
+	{
+		UEdGraphNode* Node = Nodes[NodeIndex];
+		if (UAnimGraphNode_Base* TypedNode = Cast<UAnimGraphNode_Base>(Node))
+		{
+			UClass* TypedNodeClass = TypedNode->GetClass();
+			if (TypedNodeClass == NodeClass || (bIncludeChildClasses && TypedNode->GetClass()->IsChildOf(NodeClass)))
+			{				
+				GraphNodes.Add(TypedNode);
+			}
+		}
+	}
+}
 
 #undef LOCTEXT_NAMESPACE
diff --git a/Engine/Source/Editor/AnimGraph/Public/AnimGraphNode_Base.h b/Engine/Source/Editor/AnimGraph/Public/AnimGraphNode_Base.h
--- a/Engine/Source/Editor/AnimGraph/Public/AnimGraphNode_Base.h
+++ b/Engine/Source/Editor/AnimGraph/Public/AnimGraphNode_Base.h
@@ -190,7 +190,7 @@ struct FAnimGraphNodePropertyBinding
   *
   * Any concrete implementations will be paired with a runtime graph node derived from FAnimNode_Base
   */
-UCLASS(Abstract)
+UCLASS(Abstract, BlueprintType)
 class ANIMGRAPH_API UAnimGraphNode_Base : public UK2Node
 {
 	GENERATED_UCLASS_BODY()
@@ -210,19 +210,19 @@ class ANIMGRAPH_API UAnimGraphNode_Base : public UK2Node
 	EBlueprintUsage BlueprintUsage;
 
 	// Function called when the node is initialized
-	UPROPERTY(EditAnywhere, Category = Functions, meta=(FunctionReference, PrototypeFunction="/Script/AnimGraphRuntime.AnimExecutionContextLibrary.Prototype_ThreadSafeAnimNodeCall"), DisplayName="On Initialize")
+	UPROPERTY(EditAnywhere, Category = Functions, meta=(FunctionReference, PrototypeFunction="/Script/AnimGraphRuntime.AnimNodeLibrary.Prototype_ThreadSafeAnimNodeCall"), DisplayName="On Initialize")
 	FMemberReference InitializeFunction;
 
 	// Function called when the node becomes relevant
-	UPROPERTY(EditAnywhere, Category = Functions, meta=(FunctionReference, PrototypeFunction="/Script/AnimGraphRuntime.AnimExecutionContextLibrary.Prototype_ThreadSafeAnimNodeCall"), DisplayName="On Become Relevant")
+	UPROPERTY(EditAnywhere, Category = Functions, meta=(FunctionReference, PrototypeFunction="/Script/AnimGraphRuntime.AnimNodeLibrary.Prototype_ThreadSafeAnimNodeCall"), DisplayName="On Become Relevant")
 	FMemberReference BecomeRelevantFunction;
 
 	// Function called when the node is updated
-	UPROPERTY(EditAnywhere, Category = Functions, meta=(FunctionReference, PrototypeFunction="/Script/AnimGraphRuntime.AnimExecutionContextLibrary.Prototype_ThreadSafeAnimNodeCall"), DisplayName="On Update")
+	UPROPERTY(EditAnywhere, Category = Functions, meta=(FunctionReference, PrototypeFunction="/Script/AnimGraphRuntime.AnimNodeLibrary.Prototype_ThreadSafeAnimNodeCall"), DisplayName="On Update")
 	FMemberReference UpdateFunction;
 
 	// Function called when the node is evaluated
-	UPROPERTY(EditAnywhere, Category = Functions, meta=(FunctionReference, PrototypeFunction="/Script/AnimGraphRuntime.AnimExecutionContextLibrary.Prototype_ThreadSafeAnimNodeCall"), DisplayName="On Evaluate")
+	UPROPERTY(EditAnywhere, Category = Functions, meta=(FunctionReference, PrototypeFunction="/Script/AnimGraphRuntime.AnimNodeLibrary.Prototype_ThreadSafeAnimNodeCall"), DisplayName="On Evaluate")
 	FMemberReference EvaluateFunction;
 	
 	// UObject interface
diff --git a/Engine/Source/Editor/AnimGraph/Public/AnimGraphNode_PoseDriver.h b/Engine/Source/Editor/AnimGraph/Public/AnimGraphNode_PoseDriver.h
--- a/Engine/Source/Editor/AnimGraph/Public/AnimGraphNode_PoseDriver.h
+++ b/Engine/Source/Editor/AnimGraph/Public/AnimGraphNode_PoseDriver.h
@@ -42,8 +42,60 @@ public:
 	FAnimNode_PoseDriver* GetPreviewPoseDriverNode() const;
 
 	/** Util to replace current contents of PoseTargets with info from assigned PoseAsset */
+	UFUNCTION(BlueprintCallable, Category=PoseDriver)
 	void CopyTargetsFromPoseAsset();
 
+	/** Sets the pose-driver its source bones by name */
+	UFUNCTION(BlueprintCallable, Category=PoseDriver)
+	void SetSourceBones(const TArray<FName>& BoneNames);
+
+	/** Returns the pose-driver its source bones by name */
+	UFUNCTION(BlueprintPure, Category=PoseDriver)
+	void GetSourceBoneNames(TArray<FName>& BoneNames);
+ 
+	/** Set the pose-driver its driven bones by name */
+	UFUNCTION(BlueprintCallable, Category=PoseDriver)
+	void SetDrivingBones(const TArray<FName>& BoneNames);
+
+	/** Returns the pose-driver its driven bones by name */
+	UFUNCTION(BlueprintPure, Category=PoseDriver)
+	void GetDrivingBoneNames(TArray<FName>& BoneNames);
+
+	/**  */
+	UFUNCTION(BlueprintCallable, Category=PoseDriver)
+	void SetRBFParameters(FRBFParams Parameters);
+
+	/**  */
+	UFUNCTION(BlueprintPure, Category=PoseDriver)
+	FRBFParams& GetRBFParameters();
+	
+	/**  */
+	UFUNCTION(BlueprintCallable, Category=PoseDriver)
+	void SetPoseDriverSource(EPoseDriverSource DriverSource);
+
+	/**  */
+	UFUNCTION(BlueprintPure, Category=PoseDriver)
+	EPoseDriverSource& GetPoseDriverSource();
+	
+	/**  */
+	UFUNCTION(BlueprintCallable, Category=PoseDriver)
+	void SetPoseDriverOutput(EPoseDriverOutput DriverOutput);
+
+	/**  */
+	UFUNCTION(BlueprintPure, Category=PoseDriver)
+	EPoseDriverOutput& GetPoseDriverOutput();
+	
+	/**  */
+	UFUNCTION(BlueprintCallable, Category=PoseDriver)
+	void SetOnlyDriveSelectedBones(bool bOnlyDriveSelectedBones);
+
+	/**  */
+	UFUNCTION(BlueprintPure, Category=PoseDriver)
+	bool GetOnlyDriveSelectedBones();
+
+
+	
+
 	/** Automatically modify TargetScale for each PoseTarget, based on distance to nearest neighbor */
 	void AutoSetTargetScales(float& OutMaxDistance);
 
diff --git a/Engine/Source/Editor/AnimGraph/Public/AnimationGraph.h b/Engine/Source/Editor/AnimGraph/Public/AnimationGraph.h
--- a/Engine/Source/Editor/AnimGraph/Public/AnimationGraph.h
+++ b/Engine/Source/Editor/AnimGraph/Public/AnimationGraph.h
@@ -6,6 +6,7 @@
 #include "UObject/ObjectMacros.h"
 #include "EdGraph/EdGraph.h"
 #include "Animation/AnimClassInterface.h"
+#include "AnimGraphNode_Base.h"
 #include "AnimationGraph.generated.h"
 
 class UEdGraphPin;
@@ -13,8 +14,8 @@ class UEdGraphPin;
 /** Delegate fired when a pin's default value is changed */
 DECLARE_MULTICAST_DELEGATE_OneParam(FOnPinDefaultValueChanged, UEdGraphPin* /*InPinThatChanged*/)
 
-UCLASS(MinimalAPI)
-class UAnimationGraph : public UEdGraph
+UCLASS(BlueprintType)
+class ANIMGRAPH_API UAnimationGraph : public UEdGraph
 {
 	GENERATED_UCLASS_BODY()
 
@@ -24,5 +25,9 @@ class UAnimationGraph : public UEdGraph
 	/** Blending options for animation graphs in Linked Animation Blueprints. */
 	UPROPERTY(EditAnywhere, Category = GraphBlending, meta = (ShowOnlyInnerProperties))
 	FAnimGraphBlendOptions BlendOptions;
+
+	/** Returns contained graph nodes of the specified (or child) class */
+	UFUNCTION(BlueprintCallable, Category=AnimationGraph)
+	void GetGraphNodesOfClass(TSubclassOf<UAnimGraphNode_Base> NodeClass, TArray<UAnimGraphNode_Base*>& GraphNodes, bool bIncludeChildClasses = true);
 };
 
diff --git a/Engine/Source/Editor/AnimationBlueprintLibrary/AnimationBlueprintLibrary.build.cs b/Engine/Source/Editor/AnimationBlueprintLibrary/AnimationBlueprintLibrary.build.cs
--- a/Engine/Source/Editor/AnimationBlueprintLibrary/AnimationBlueprintLibrary.build.cs
+++ b/Engine/Source/Editor/AnimationBlueprintLibrary/AnimationBlueprintLibrary.build.cs
@@ -7,13 +7,21 @@ public class AnimationBlueprintLibrary : ModuleRules
 	public AnimationBlueprintLibrary(ReadOnlyTargetRules Target) : base(Target)
 	{
 		PrivateDependencyModuleNames.AddRange(
-			new string[] 
-            {
-                "Core",
-                "CoreUObject",
-                "Engine",
-                "Kismet"              
-            }
+			new string[]
+			{
+				"Core",
+				"CoreUObject",
+				"Engine",
+				"Kismet",
+				"AnimGraph",
+			}
 		);
-    }
+
+		PublicDependencyModuleNames.AddRange(
+			new string[]
+			{
+				"AnimGraph",
+			}
+		);
+	}
 }
diff --git a/Engine/Source/Editor/AnimationBlueprintLibrary/Private/AnimationBlueprintLibrary.cpp b/Engine/Source/Editor/AnimationBlueprintLibrary/Private/AnimationBlueprintLibrary.cpp
--- a/Engine/Source/Editor/AnimationBlueprintLibrary/Private/AnimationBlueprintLibrary.cpp
+++ b/Engine/Source/Editor/AnimationBlueprintLibrary/Private/AnimationBlueprintLibrary.cpp
@@ -2267,6 +2267,44 @@ void UAnimationBlueprintLibrary::FinalizeBoneAnimation(UAnimSequence* AnimationS
 	}
 }
 
+static void RecursiveRetrieveAnimationGraphs(UEdGraph* EdGraph, TArray<UAnimationGraph*>& OutAnimationGraphs)
+{
+	if (UAnimationGraph* AnimGraph = Cast<UAnimationGraph>(EdGraph))
+	{
+		OutAnimationGraphs.Add(AnimGraph);
+	}
+	
+	for (TObjectPtr<class UEdGraph>& SubGraph : EdGraph->SubGraphs)
+	{		
+		RecursiveRetrieveAnimationGraphs(SubGraph, OutAnimationGraphs);
+	}
+}
+
+void UAnimationBlueprintLibrary::GetAnimationGraphs(UAnimBlueprint* AnimationBlueprint, TArray<UAnimationGraph*>& AnimationGraphs)
+{
+	if (AnimationBlueprint)
+	{
+		for (UEdGraph* EdGraph : AnimationBlueprint->FunctionGraphs)
+		{
+			RecursiveRetrieveAnimationGraphs(EdGraph, AnimationGraphs);
+		}
+	}
+	else
+	{
+		UE_LOG(LogAnimationBlueprintLibrary, Warning, TEXT("Invalid Animation Blueprint"));
+	}
+}
+
+void UAnimationBlueprintLibrary::GetNodesOfClass(UAnimBlueprint* AnimationBlueprint, TSubclassOf<UAnimGraphNode_Base> NodeClass, TArray<UAnimGraphNode_Base*>& GraphNodes, bool bIncludeChildClasses /*= true*/)
+{
+	TArray<UAnimationGraph*> AnimationGraphs;
+	GetAnimationGraphs(AnimationBlueprint, AnimationGraphs);
+	for (UAnimationGraph* AnimGraph : AnimationGraphs)
+	{
+		AnimGraph->GetGraphNodesOfClass(NodeClass, GraphNodes, bIncludeChildClasses);
+	}
+}
+
 template void UAnimationBlueprintLibrary::AddCurveKeysInternal<float, FFloatCurve, ERawCurveTrackTypes::RCT_Float>(UAnimSequence* AnimationSequence, FName CurveName, const TArray<float>& Times, const TArray<float>& KeyData);
 template void UAnimationBlueprintLibrary::AddCurveKeysInternal<FVector, FVectorCurve, ERawCurveTrackTypes::RCT_Vector>(UAnimSequence* AnimationSequence, FName CurveName, const TArray<float>& Times, const TArray<FVector>& KeyData);
 template void UAnimationBlueprintLibrary::AddCurveKeysInternal<FTransform, FTransformCurve, ERawCurveTrackTypes::RCT_Transform>(UAnimSequence* AnimationSequence, FName CurveName, const TArray<float>& Times, const TArray<FTransform>& KeyData);
diff --git a/Engine/Source/Editor/AnimationBlueprintLibrary/Public/AnimationBlueprintLibrary.h b/Engine/Source/Editor/AnimationBlueprintLibrary/Public/AnimationBlueprintLibrary.h
--- a/Engine/Source/Editor/AnimationBlueprintLibrary/Public/AnimationBlueprintLibrary.h
+++ b/Engine/Source/Editor/AnimationBlueprintLibrary/Public/AnimationBlueprintLibrary.h
@@ -6,6 +6,7 @@
 #include "Animation/AnimEnums.h"
 #include "Animation/AnimCurveTypes.h"
 #include "Animation/AnimMetaData.h"
+#include "AnimationGraph.h"
 
 #include "AnimationBlueprintLibrary.generated.h"
 
@@ -471,4 +472,12 @@ public:
 	/** Finds the Bone Path from the given Bone to the Root Bone */
 	UFUNCTION(BlueprintPure, Category = "AnimationBlueprintLibrary|Helpers")
 	static void FindBonePathToRoot(const UAnimSequenceBase* AnimationSequenceBase, FName BoneName, TArray<FName>& BonePath);
+
+	/** Returns all Animation Graphs contained by the provided Animation Blueprint */
+	UFUNCTION(BlueprintCallable, Category=Animation, meta=(ScriptMethod))
+	static void GetAnimationGraphs(UAnimBlueprint* AnimationBlueprint, TArray<UAnimationGraph*>& AnimationGraphs);
+
+	/** Returns all Animation Graph Nodes of the provided Node Class contained by the Animation Blueprint */
+	UFUNCTION(BlueprintCallable, Category=Animation, meta=(ScriptMethod))
+	static void GetNodesOfClass(UAnimBlueprint* AnimationBlueprint, TSubclassOf<UAnimGraphNode_Base> NodeClass, TArray<UAnimGraphNode_Base*>& GraphNodes, bool bIncludeChildClasses = true);
 };
diff --git a/Engine/Source/Runtime/AnimGraphRuntime/Public/AnimNodes/AnimNode_PoseDriver.h b/Engine/Source/Runtime/AnimGraphRuntime/Public/AnimNodes/AnimNode_PoseDriver.h
--- a/Engine/Source/Runtime/AnimGraphRuntime/Public/AnimNodes/AnimNode_PoseDriver.h
+++ b/Engine/Source/Runtime/AnimGraphRuntime/Public/AnimNodes/AnimNode_PoseDriver.h
@@ -127,7 +127,7 @@ struct ANIMGRAPHRUNTIME_API FPoseDriverTarget
 };
 
 /** RBF based orientation driver */
-USTRUCT(BlueprintType)
+USTRUCT(BlueprintInternalUseOnly)
 struct ANIMGRAPHRUNTIME_API FAnimNode_PoseDriver : public FAnimNode_PoseHandler
 {
 	GENERATED_BODY()
@@ -165,7 +165,7 @@ struct ANIMGRAPHRUNTIME_API FAnimNode_PoseDriver : public FAnimNode_PoseHandler
 	FBoneReference EvalSpaceBone;
 
 	/** Parameters used by RBF solver */
-	UPROPERTY(EditAnywhere, BlueprintReadWrite, Category = PoseDriver, meta=(ShowOnlyInnerProperties))
+	UPROPERTY(EditAnywhere, Category = PoseDriver, meta=(ShowOnlyInnerProperties))
 	FRBFParams RBFParams;
 
 #if WITH_EDITORONLY_DATA
@@ -182,15 +182,15 @@ struct ANIMGRAPHRUNTIME_API FAnimNode_PoseDriver : public FAnimNode_PoseHandler
 #endif
 
 	/** Which part of the transform is read */
-	UPROPERTY(EditAnywhere, BlueprintReadWrite, Category = PoseDriver)
+	UPROPERTY(EditAnywhere, Category = PoseDriver)
 	EPoseDriverSource DriveSource;
 
 	/** Whether we should drive poses or curves */
-	UPROPERTY(EditAnywhere, BlueprintReadWrite, Category = PoseDriver)
+	UPROPERTY(EditAnywhere, Category = PoseDriver)
 	EPoseDriverOutput DriveOutput;
 
 	/** If we should filter bones to be driven using the DrivenBonesFilter array */
-	UPROPERTY(EditAnywhere, BlueprintReadWrite, Category = PoseDriver)
+	UPROPERTY(EditAnywhere, Category = PoseDriver)
 	uint8 bOnlyDriveSelectedBones : 1;
 
 	/** If true, will recalculate DrivenUID values in PoseTargets array on next eval */
diff --git a/Engine/Source/Runtime/Engine/Classes/Animation/AnimationAsset.h b/Engine/Source/Runtime/Engine/Classes/Animation/AnimationAsset.h
--- a/Engine/Source/Runtime/Engine/Classes/Animation/AnimationAsset.h
+++ b/Engine/Source/Runtime/Engine/Classes/Animation/AnimationAsset.h
@@ -1015,6 +1015,10 @@ public:
 	virtual USkeletalMesh* GetPreviewMesh() const override;
 
 #if WITH_EDITOR
+	/** Sets or updates the preview skeletal mesh */
+	UFUNCTION(BlueprintCallable, Category=Animation)
+	void SetPreviewSkeletalMesh(USkeletalMesh* PreviewMesh) { SetPreviewMesh(PreviewMesh); }
+	
 	/** Replace Skeleton 
 	 * 
 	 * @param NewSkeleton	NewSkeleton to change to 
diff --git a/Engine/Source/Runtime/Engine/Classes/Animation/PoseAsset.h b/Engine/Source/Runtime/Engine/Classes/Animation/PoseAsset.h
--- a/Engine/Source/Runtime/Engine/Classes/Animation/PoseAsset.h
+++ b/Engine/Source/Runtime/Engine/Classes/Animation/PoseAsset.h
@@ -247,9 +247,20 @@ public:
 	ENGINE_API bool ContainsPose(const FName& InPoseName) const;
 
 #if WITH_EDITOR
+	/** Renames a specific pose */
+	UFUNCTION(BlueprintCallable, Category=PoseAsset)
+	void RenamePose(const FName& OriginalPoseName, const FName& NewPoseName);
+	
+	/** Returns the name of all contained poses */
+	UFUNCTION(BlueprintPure, Category=PoseAsset)
+	void GetPoseNames(TArray<FName>& PoseNames) const;
+	
 	ENGINE_API bool AddOrUpdatePoseWithUniqueName(USkeletalMeshComponent* MeshComponent, FSmartName* OutPoseName = nullptr);
 
 	ENGINE_API void CreatePoseFromAnimation(class UAnimSequence* AnimSequence, const TArray<FSmartName>* InPoseNames = nullptr);
+
+	/** Contained poses are re-generated from the provided Animation Sequence*/
+	UFUNCTION(BlueprintCallable, Category=PoseAsset)
 	ENGINE_API void UpdatePoseFromAnimation(class UAnimSequence* AnimSequence);
 
 	// Begin AnimationAsset interface
diff --git a/Engine/Source/Runtime/Engine/Private/Animation/PoseAsset.cpp b/Engine/Source/Runtime/Engine/Private/Animation/PoseAsset.cpp
--- a/Engine/Source/Runtime/Engine/Private/Animation/PoseAsset.cpp
+++ b/Engine/Source/Runtime/Engine/Private/Animation/PoseAsset.cpp
@@ -1110,6 +1110,19 @@ bool UPoseAsset::ContainsPose(const FName& InPoseName) const
 }
 
 #if WITH_EDITOR
+void UPoseAsset::RenamePose(const FName& OriginalPoseName, const FName& NewPoseName)
+{
+	ModifyPoseName(OriginalPoseName, NewPoseName, nullptr);
+}
+
+void UPoseAsset::GetPoseNames(TArray<FName>& PoseNames) const
+{	
+	const int32 NumPoses = GetNumPoses();
+	for (int32 PoseIndex = 0; PoseIndex < NumPoses; ++PoseIndex)
+	{
+		PoseNames.Add(GetPoseNameByIndex(PoseIndex));
+	}
+}
 // whenever you change SourceLocalPoses, or SourceCurves, we should call this to update runtime dataa
 void UPoseAsset::PostProcessData()
 {
