Fixed warning when linked anim graph node pins are connected

diff --git a/Engine/Source/Editor/AnimGraph/Private/AnimBlueprintCompiler.cpp b/Engine/Source/Editor/AnimGraph/Private/AnimBlueprintCompiler.cpp
--- a/Engine/Source/Editor/AnimGraph/Private/AnimBlueprintCompiler.cpp
+++ b/Engine/Source/Editor/AnimGraph/Private/AnimBlueprintCompiler.cpp
@@ -840,7 +840,8 @@ void FAnimBlueprintCompilerContext::CopyTermDefaultsToDefaultObject(UObject* Def
 				const uint8* SourcePtr = SourceNodeProperty->ContainerPtrToValuePtr<uint8>(VisualAnimNode);
 
 				FAnimBlueprintNodeCopyTermDefaultsContext NodeContext(DefaultObject, TargetProperty, DestinationPtr, SourcePtr, LinkIndexCount);
-				VisualAnimNode->CopyTermDefaultsToDefaultObject(CompilerContext, NodeContext, CompiledData);
+				UAnimGraphNode_Base* OriginalAnimNode = Cast<UAnimGraphNode_Base>(MessageLog.FindSourceObject(VisualAnimNode));
+				OriginalAnimNode->CopyTermDefaultsToDefaultObject(CompilerContext, NodeContext, CompiledData);
 
 				LinkIndexMap.Add(VisualAnimNode, LinkIndexCount);
 				NodeBaseAddresses.Add(VisualAnimNode, DestinationPtr);
diff --git a/Engine/Source/Editor/AnimGraph/Private/AnimGraphNode_CustomProperty.cpp b/Engine/Source/Editor/AnimGraph/Private/AnimGraphNode_CustomProperty.cpp
--- a/Engine/Source/Editor/AnimGraph/Private/AnimGraphNode_CustomProperty.cpp
+++ b/Engine/Source/Editor/AnimGraph/Private/AnimGraphNode_CustomProperty.cpp
@@ -76,7 +76,7 @@ void UAnimGraphNode_CustomProperty::OnCopyTermDefaultsToDefaultObject(IAnimBluep
 	// Copy pin default values to generated properties
 	for (UEdGraphPin* Pin : Pins)
 	{
-		if (!Pin->bOrphanedPin && !UAnimationGraphSchema::IsPosePin(Pin->PinType))
+		if (!Pin->bOrphanedPin && Pin->LinkedTo.Num() == 0 && !UAnimationGraphSchema::IsPosePin(Pin->PinType))
 		{
 			FString PrefixedName = GetPinTargetVariableName(Pin);
 
