Display whether anim node properties are dynamic in the binding menu

diff --git a/Engine/Source/Editor/AnimGraph/Private/AnimGraphNode_Base.cpp b/Engine/Source/Editor/AnimGraph/Private/AnimGraphNode_Base.cpp
--- a/Engine/Source/Editor/AnimGraph/Private/AnimGraphNode_Base.cpp
+++ b/Engine/Source/Editor/AnimGraph/Private/AnimGraphNode_Base.cpp
@@ -850,6 +850,7 @@ TSharedRef<SWidget> UAnimGraphNode_Base::MakePropertyBindingWidget(const FAnimPr
 			None,
 			Pin,
 			Binding,
+			Dynamic,
 			MultipleValues,
 		};
 
@@ -860,6 +861,7 @@ TSharedRef<SWidget> UAnimGraphNode_Base::MakePropertyBindingWidget(const FAnimPr
 			const FText MultipleValues = LOCTEXT("MultipleValuesLabel", "Multiple Values");
 			const FText Bind = LOCTEXT("BindLabel", "Bind");
 			const FText ExposedAsPin = LOCTEXT("ExposedAsPinLabel", "Pin");
+			const FText Dynamic = LOCTEXT("DynamicLabel", "Dynamic");
 			FText CurrentValue = Bind;
 
 			auto SetAssignValue = [&CurrentValueType, &CurrentValue, &MultipleValues](const FText& InValue, ECurrentValueType InType)
@@ -893,6 +895,10 @@ TSharedRef<SWidget> UAnimGraphNode_Base::MakePropertyBindingWidget(const FAnimPr
 				{
 					SetAssignValue(BindingPtr->PathAsText, ECurrentValueType::Binding);
 				}
+				else if(AnimGraphNode->AlwaysDynamicProperties.Find(InArgs.PinName))
+				{
+					SetAssignValue(Dynamic, ECurrentValueType::Dynamic);
+				}
 				else
 				{
 					if(AnimGraphNode->ShowPinForProperties[InArgs.OptionalPinIndex].bShowPin)
@@ -916,6 +922,7 @@ TSharedRef<SWidget> UAnimGraphNode_Base::MakePropertyBindingWidget(const FAnimPr
 			const FText MultipleValues = LOCTEXT("MultipleValuesToolTip", "Bindings Have Multiple Values");
 			const FText ExposedAsPin = LOCTEXT("ExposedAsPinToolTip", "Exposed As a Pin on the Node");
 			const FText BindValue = LOCTEXT("BindValueToolTip", "Bind This Value");
+			const FText DynamicValue = LOCTEXT("DynamicValueToolTip", "Dynamic value that can be set externally");
 			FText CurrentValue;
 			
 			auto SetAssignValue = [&CurrentValueType, &CurrentValue, &MultipleValues](const FText& InValue, ECurrentValueType InType)
@@ -965,6 +972,10 @@ TSharedRef<SWidget> UAnimGraphNode_Base::MakePropertyBindingWidget(const FAnimPr
 						}
 					}
 				}
+				else if(AnimGraphNode->AlwaysDynamicProperties.Find(InArgs.PinName))
+				{
+					SetAssignValue(DynamicValue, ECurrentValueType::Dynamic);
+				}
 				else
 				{
 					if(AnimGraphNode->ShowPinForProperties[InArgs.OptionalPinIndex].bShowPin)
