Fixed metadata preventing new bindings being created

diff --git a/Engine/Source/Editor/AnimGraph/Public/AnimGraphNode_Base.h b/Engine/Source/Editor/AnimGraph/Public/AnimGraphNode_Base.h
--- a/Engine/Source/Editor/AnimGraph/Public/AnimGraphNode_Base.h
+++ b/Engine/Source/Editor/AnimGraph/Public/AnimGraphNode_Base.h
@@ -210,19 +210,19 @@ class ANIMGRAPH_API UAnimGraphNode_Base : public UK2Node
 	EBlueprintUsage BlueprintUsage;
 
 	// Function called when the node is initialized
-	UPROPERTY(EditAnywhere, Category = Functions, meta=(FunctionReference, PrototypeFunction="/Script/AnimGraphRuntime.AnimNodeLibrary.Prototype_ThreadSafeAnimNodeCall"), DisplayName="On Initialize")
+	UPROPERTY(EditAnywhere, Category = Functions, meta=(FunctionReference, PrototypeFunction="/Script/AnimGraphRuntime.AnimExecutionContextLibrary.Prototype_ThreadSafeAnimNodeCall"), DisplayName="On Initialize")
 	FMemberReference InitializeFunction;
 
 	// Function called when the node becomes relevant
-	UPROPERTY(EditAnywhere, Category = Functions, meta=(FunctionReference, PrototypeFunction="/Script/AnimGraphRuntime.AnimNodeLibrary.Prototype_ThreadSafeAnimNodeCall"), DisplayName="On Become Relevant")
+	UPROPERTY(EditAnywhere, Category = Functions, meta=(FunctionReference, PrototypeFunction="/Script/AnimGraphRuntime.AnimExecutionContextLibrary.Prototype_ThreadSafeAnimNodeCall"), DisplayName="On Become Relevant")
 	FMemberReference BecomeRelevantFunction;
 
 	// Function called when the node is updated
-	UPROPERTY(EditAnywhere, Category = Functions, meta=(FunctionReference, PrototypeFunction="/Script/AnimGraphRuntime.AnimNodeLibrary.Prototype_ThreadSafeAnimNodeCall"), DisplayName="On Update")
+	UPROPERTY(EditAnywhere, Category = Functions, meta=(FunctionReference, PrototypeFunction="/Script/AnimGraphRuntime.AnimExecutionContextLibrary.Prototype_ThreadSafeAnimNodeCall"), DisplayName="On Update")
 	FMemberReference UpdateFunction;
 
 	// Function called when the node is evaluated
-	UPROPERTY(EditAnywhere, Category = Functions, meta=(FunctionReference, PrototypeFunction="/Script/AnimGraphRuntime.AnimNodeLibrary.Prototype_ThreadSafeAnimNodeCall"), DisplayName="On Evaluate")
+	UPROPERTY(EditAnywhere, Category = Functions, meta=(FunctionReference, PrototypeFunction="/Script/AnimGraphRuntime.AnimExecutionContextLibrary.Prototype_ThreadSafeAnimNodeCall"), DisplayName="On Evaluate")
 	FMemberReference EvaluateFunction;
 	
 	// UObject interface
