Fix crash when compiling-on-load a graph that contains NULL nodes

diff --git a/Engine/Source/Editor/AnimGraph/Private/AnimGraphNode_AssetPlayerBase.cpp b/Engine/Source/Editor/AnimGraph/Private/AnimGraphNode_AssetPlayerBase.cpp
--- a/Engine/Source/Editor/AnimGraph/Private/AnimGraphNode_AssetPlayerBase.cpp
+++ b/Engine/Source/Editor/AnimGraph/Private/AnimGraphNode_AssetPlayerBase.cpp
@@ -122,7 +122,7 @@ void UAnimGraphNode_AssetPlayerBase::OnProcessDuringCompilation(IAnimBlueprintCo
 			// Also make sure we do not process any empty stub graphs
 			if (!GraphName.Contains(ANIM_FUNC_DECORATOR))
 			{
-				if (Graph->Nodes.ContainsByPredicate([this, &OutCompiledData](UEdGraphNode* Node) { return Node->NodeGuid == NodeGuid; }))
+				if (Graph->Nodes.ContainsByPredicate([this, &OutCompiledData](UEdGraphNode* Node) { return Node && Node->NodeGuid == NodeGuid; }))
 				{
 					if (int32* IndexPtr = OutCompiledData.GetAnimBlueprintDebugData().NodeGuidToIndexMap.Find(NodeGuid))
 					{
