Fixed crash drag/dropping a sequence player into a template anim BP

diff --git a/Engine/Source/Editor/AnimGraph/Private/AnimationGraphSchema.cpp b/Engine/Source/Editor/AnimGraph/Private/AnimationGraphSchema.cpp
--- a/Engine/Source/Editor/AnimGraph/Private/AnimationGraphSchema.cpp
+++ b/Engine/Source/Editor/AnimGraph/Private/AnimationGraphSchema.cpp
@@ -551,9 +551,9 @@ void UAnimationGraphSchema::SpawnNodeFromAsset(UAnimationAsset* Asset, const FVe
 
 	UAnimBlueprint* AnimBlueprint = Cast<UAnimBlueprint>(FBlueprintEditorUtils::FindBlueprintForGraph(Graph));
 
-	const bool bSkelMatch = (AnimBlueprint != NULL) && (AnimBlueprint->TargetSkeleton->IsCompatible(Asset->GetSkeleton()));
-	const bool bTypeMatch = (PinIfAvailable == NULL) || UAnimationGraphSchema::IsLocalSpacePosePin(PinIfAvailable->PinType);
-	const bool bDirectionMatch = (PinIfAvailable == NULL) || (PinIfAvailable->Direction == EGPD_Input);
+	const bool bSkelMatch = (AnimBlueprint != nullptr) && (AnimBlueprint->TargetSkeleton != nullptr) && (AnimBlueprint->TargetSkeleton->IsCompatible(Asset->GetSkeleton()));
+	const bool bTypeMatch = (PinIfAvailable == nullptr) || UAnimationGraphSchema::IsLocalSpacePosePin(PinIfAvailable->PinType);
+	const bool bDirectionMatch = (PinIfAvailable == nullptr) || (PinIfAvailable->Direction == EGPD_Input);
 
 	if (bSkelMatch && bTypeMatch && bDirectionMatch)
 	{
@@ -712,11 +712,18 @@ void UAnimationGraphSchema::GetAssetsGraphHoverMessage(const TArray<FAssetData>&
 	if (UAnimationAsset* AnimationAsset = FAssetData::GetFirstAsset<UAnimationAsset>(Assets))
 	{
 		UAnimBlueprint* AnimBlueprint = Cast<UAnimBlueprint>(FBlueprintEditorUtils::FindBlueprintForGraph(HoverGraph));
-		const bool bSkelMatch = (AnimBlueprint != NULL) && (AnimBlueprint->TargetSkeleton->IsCompatible(AnimationAsset->GetSkeleton()));
+		const bool bSkelMatch = (AnimBlueprint != nullptr) && (AnimBlueprint->TargetSkeleton != nullptr) && (AnimBlueprint->TargetSkeleton->IsCompatible(AnimationAsset->GetSkeleton()));
 		if (!bSkelMatch)
 		{
 			OutOkIcon = false;
-			OutTooltipText = LOCTEXT("SkeletonsNotCompatible", "Skeletons are not compatible").ToString();
+			if(AnimBlueprint && AnimBlueprint->bIsTemplate)
+			{
+				OutTooltipText = LOCTEXT("TemplateNotAllowed", "Template animation blueprints cannot reference assets").ToString();
+			}
+			else
+			{
+				OutTooltipText = LOCTEXT("SkeletonsNotCompatible", "Skeletons are not compatible").ToString();
+			}
 		}
 		else if(UAnimMontage* Montage = FAssetData::GetFirstAsset<UAnimMontage>(Assets))
 		{
