Removing Initialize and Evaluate functions

For now, removing these as their semantics are due to change
Also converting Update and Become Relevant functions to use more-devived-typed update context (less casting required)

diff --git a/Engine/Source/Editor/AnimGraph/Private/AnimGraphNode_Base.cpp b/Engine/Source/Editor/AnimGraph/Private/AnimGraphNode_Base.cpp
--- a/Engine/Source/Editor/AnimGraph/Private/AnimGraphNode_Base.cpp
+++ b/Engine/Source/Editor/AnimGraph/Private/AnimGraphNode_Base.cpp
@@ -85,11 +85,6 @@ void UAnimGraphNode_Base::PostEditChangeProperty(FPropertyChangedEvent& Property
 		FOptionalPinManager::EvaluateOldShownPins(ShowPinForProperties, OldShownPins, this);
 		GetSchema()->ReconstructNode(*this);
 	}
-	else if(PropertyName == GET_MEMBER_NAME_CHECKED(UAnimGraphNode_Base, InitializeFunction))
-	{
-		GetFNode()->InitializeFunction.SetFromFunction(InitializeFunction.ResolveMember<UFunction>(GetAnimBlueprint()->SkeletonGeneratedClass));
-		GetSchema()->ReconstructNode(*this);
-	}
 	else if(PropertyName == GET_MEMBER_NAME_CHECKED(UAnimGraphNode_Base, BecomeRelevantFunction))
 	{
 		GetFNode()->BecomeRelevantFunction.SetFromFunction(BecomeRelevantFunction.ResolveMember<UFunction>(GetAnimBlueprint()->SkeletonGeneratedClass));
@@ -101,11 +96,6 @@ void UAnimGraphNode_Base::PostEditChangeProperty(FPropertyChangedEvent& Property
 		GetFNode()->UpdateFunction.SetFromFunction(UpdateFunction.ResolveMember<UFunction>(GetAnimBlueprint()->SkeletonGeneratedClass));
 		GetSchema()->ReconstructNode(*this);
 	}
-	else if(PropertyName == GET_MEMBER_NAME_CHECKED(UAnimGraphNode_Base, EvaluateFunction))
-	{
-		GetFNode()->EvaluateFunction.SetFromFunction(EvaluateFunction.ResolveMember<UFunction>(GetAnimBlueprint()->SkeletonGeneratedClass));
-		GetSchema()->ReconstructNode(*this);
-	}
 	
 	Super::PostEditChangeProperty(PropertyChangedEvent);
 
@@ -219,10 +209,8 @@ void UAnimGraphNode_Base::ValidateAnimNodeDuringCompilation(USkeleton* ForSkelet
 		}
 	};
 	
-	ValidateFunctionRef(InitializeFunction, LOCTEXT("MissingInitializeFunction", "Could not resolve initialize function for @@"));
 	ValidateFunctionRef(BecomeRelevantFunction, LOCTEXT("MissingBecomeRelevantFunction", "Could not resolve become relevant function for @@"));
 	ValidateFunctionRef(UpdateFunction, LOCTEXT("MissingUpdateFunction", "Could not resolve update function for @@"));
-	ValidateFunctionRef(EvaluateFunction, LOCTEXT("MissingEvaluateFunction", "Could not resolve evaluate function for @@"));
 }
 
 void UAnimGraphNode_Base::CopyTermDefaultsToDefaultObject(IAnimBlueprintCopyTermDefaultsContext& InCompilationContext, IAnimBlueprintNodeCopyTermDefaultsContext& InPerNodeContext, IAnimBlueprintGeneratedClassCompiledData& OutCompiledData)
@@ -506,11 +494,9 @@ void UAnimGraphNode_Base::ProcessDuringCompilation(IAnimBlueprintCompilationCont
 	Extension->ProcessNodePins(this, InCompilationContext, OutCompiledData);
 
 	// Resolve functions
-	GetFNode()->InitializeFunction.SetFromFunction(InitializeFunction.ResolveMember<UFunction>(GetAnimBlueprint()->SkeletonGeneratedClass));  
 	GetFNode()->BecomeRelevantFunction.SetFromFunction(BecomeRelevantFunction.ResolveMember<UFunction>(GetAnimBlueprint()->SkeletonGeneratedClass));
-	GetFNode()->UpdateFunction.SetFromFunction(UpdateFunction.ResolveMember<UFunction>(GetAnimBlueprint()->SkeletonGeneratedClass));
-	GetFNode()->EvaluateFunction.SetFromFunction(EvaluateFunction.ResolveMember<UFunction>(GetAnimBlueprint()->SkeletonGeneratedClass)); 
-	
+	GetFNode()->UpdateFunction.SetFromFunction(UpdateFunction.ResolveMember<UFunction>(GetAnimBlueprint()->SkeletonGeneratedClass)); 
+
 	// Call the override point
 	OnProcessDuringCompilation(InCompilationContext, OutCompiledData);
 }
diff --git a/Engine/Source/Editor/AnimGraph/Public/AnimGraphNode_Base.h b/Engine/Source/Editor/AnimGraph/Public/AnimGraphNode_Base.h
--- a/Engine/Source/Editor/AnimGraph/Public/AnimGraphNode_Base.h
+++ b/Engine/Source/Editor/AnimGraph/Public/AnimGraphNode_Base.h
@@ -209,21 +209,13 @@ class ANIMGRAPH_API UAnimGraphNode_Base : public UK2Node
 	UPROPERTY(Transient)
 	EBlueprintUsage BlueprintUsage;
 
-	// Function called when the node is initialized
-	UPROPERTY(EditAnywhere, Category = Functions, meta=(FunctionReference, PrototypeFunction="/Script/AnimGraphRuntime.AnimExecutionContextLibrary.Prototype_ThreadSafeAnimNodeCall"), DisplayName="On Initialize")
-	FMemberReference InitializeFunction;
-
 	// Function called when the node becomes relevant
-	UPROPERTY(EditAnywhere, Category = Functions, meta=(FunctionReference, PrototypeFunction="/Script/AnimGraphRuntime.AnimExecutionContextLibrary.Prototype_ThreadSafeAnimNodeCall"), DisplayName="On Become Relevant")
+	UPROPERTY(EditAnywhere, Category = Functions, meta=(FunctionReference, PrototypeFunction="/Script/AnimGraphRuntime.AnimExecutionContextLibrary.Prototype_ThreadSafeAnimUpdateCall"), DisplayName="On Become Relevant")
 	FMemberReference BecomeRelevantFunction;
 
 	// Function called when the node is updated
-	UPROPERTY(EditAnywhere, Category = Functions, meta=(FunctionReference, PrototypeFunction="/Script/AnimGraphRuntime.AnimExecutionContextLibrary.Prototype_ThreadSafeAnimNodeCall"), DisplayName="On Update")
+	UPROPERTY(EditAnywhere, Category = Functions, meta=(FunctionReference, PrototypeFunction="/Script/AnimGraphRuntime.AnimExecutionContextLibrary.Prototype_ThreadSafeAnimUpdateCall"), DisplayName="On Update")
 	FMemberReference UpdateFunction;
-
-	// Function called when the node is evaluated
-	UPROPERTY(EditAnywhere, Category = Functions, meta=(FunctionReference, PrototypeFunction="/Script/AnimGraphRuntime.AnimExecutionContextLibrary.Prototype_ThreadSafeAnimNodeCall"), DisplayName="On Evaluate")
-	FMemberReference EvaluateFunction;
 	
 	// UObject interface
 	virtual void PostEditChangeProperty(FPropertyChangedEvent& PropertyChangedEvent) override;
diff --git a/Engine/Source/Editor/AnimationBlueprintEditor/Private/AnimationNodes/SAnimationGraphNode.cpp b/Engine/Source/Editor/AnimationBlueprintEditor/Private/AnimationNodes/SAnimationGraphNode.cpp
--- a/Engine/Source/Editor/AnimationBlueprintEditor/Private/AnimationNodes/SAnimationGraphNode.cpp
+++ b/Engine/Source/Editor/AnimationBlueprintEditor/Private/AnimationNodes/SAnimationGraphNode.cpp
@@ -368,10 +368,6 @@ void SAnimationGraphNode::CreateBelowPinControls(TSharedPtr<SVerticalBox> MainBo
 			}
 		};
 
-		if(AnimNode->InitializeFunction.GetMemberGuid().IsValid())
-		{
-			AddFunctionBindingWidget("Functions", GET_MEMBER_NAME_CHECKED(UAnimGraphNode_Base, InitializeFunction));
-		}
 		if(AnimNode->BecomeRelevantFunction.GetMemberGuid().IsValid())
 		{
 			AddFunctionBindingWidget("Functions", GET_MEMBER_NAME_CHECKED(UAnimGraphNode_Base, BecomeRelevantFunction));
@@ -380,10 +376,6 @@ void SAnimationGraphNode::CreateBelowPinControls(TSharedPtr<SVerticalBox> MainBo
 		{
 			AddFunctionBindingWidget("Functions", GET_MEMBER_NAME_CHECKED(UAnimGraphNode_Base, UpdateFunction));
 		}
-		if(AnimNode->EvaluateFunction.GetMemberGuid().IsValid())
-		{
-			AddFunctionBindingWidget("Functions", GET_MEMBER_NAME_CHECKED(UAnimGraphNode_Base, EvaluateFunction));
-		}
 	}
 }
 
diff --git a/Engine/Source/Runtime/AnimGraphRuntime/Public/AnimExecutionContextLibrary.h b/Engine/Source/Runtime/AnimGraphRuntime/Public/AnimExecutionContextLibrary.h
--- a/Engine/Source/Runtime/AnimGraphRuntime/Public/AnimExecutionContextLibrary.h
+++ b/Engine/Source/Runtime/AnimGraphRuntime/Public/AnimExecutionContextLibrary.h
@@ -21,6 +21,10 @@ public:
 	/** Prototype function for thread-safe anim node calls */
 	UFUNCTION(BlueprintInternalUseOnly, meta=(BlueprintThreadSafe))
 	void Prototype_ThreadSafeAnimNodeCall(const FAnimExecutionContext& Context, const FAnimNodeReference& Node) {}
+
+	/** Prototype function for thread-safe anim update calls */
+	UFUNCTION(BlueprintInternalUseOnly, meta=(BlueprintThreadSafe))
+	void Prototype_ThreadSafeAnimUpdateCall(const FAnimUpdateContext& Context, const FAnimNodeReference& Node) {}	
 #endif
 	
 	/** Get the anim instance that hosts this anim node */
diff --git a/Engine/Source/Runtime/Engine/Classes/Animation/AnimNodeBase.h b/Engine/Source/Runtime/Engine/Classes/Animation/AnimNodeBase.h
--- a/Engine/Source/Runtime/Engine/Classes/Animation/AnimNodeBase.h
+++ b/Engine/Source/Runtime/Engine/Classes/Animation/AnimNodeBase.h
@@ -966,10 +966,8 @@ protected:
 
 private:
 	// Access functions
-	const FAnimNodeFunctionRef& GetInitializeFunction() const;
 	const FAnimNodeFunctionRef& GetBecomeRelevantFunction() const;
 	const FAnimNodeFunctionRef& GetUpdateFunction() const;
-	const FAnimNodeFunctionRef& GetEvaluateFunction() const;
 	
 private:
 	friend class IAnimClassInterface;
@@ -985,10 +983,6 @@ private:
 	const FAnimNodeData* NodeData = nullptr;
 
 #if WITH_EDITORONLY_DATA
-	// Function called on initialize
-	UPROPERTY(meta=(FoldProperty))
-	FAnimNodeFunctionRef InitializeFunction;
-
 	// Function called on become relevant
 	UPROPERTY(meta=(FoldProperty))
 	FAnimNodeFunctionRef BecomeRelevantFunction;
@@ -996,10 +990,6 @@ private:
 	// Function called on update
 	UPROPERTY(meta=(FoldProperty))
 	FAnimNodeFunctionRef UpdateFunction;
-
-	// Function called on evaluate
-	UPROPERTY(meta=(FoldProperty))
-	FAnimNodeFunctionRef EvaluateFunction;
 #endif
 };
 
diff --git a/Engine/Source/Runtime/Engine/Classes/Animation/AnimNodeFunctionRef.h b/Engine/Source/Runtime/Engine/Classes/Animation/AnimNodeFunctionRef.h
--- a/Engine/Source/Runtime/Engine/Classes/Animation/AnimNodeFunctionRef.h
+++ b/Engine/Source/Runtime/Engine/Classes/Animation/AnimNodeFunctionRef.h
@@ -66,20 +66,11 @@ private:
 	friend struct ::FComponentSpacePoseLink;
 	friend struct ::FAnimInstanceProxy;
 	
-	// Call the Initialize function of this node
-	static void Initialize(const FAnimationInitializeContext& InContext, FAnimNode_Base& InNode);
-
 	// Call the BecomeRelevant function of this node
 	static void BecomeRelevant(const FAnimationUpdateContext& InContext, FAnimNode_Base& InNode);
 
 	// Call the Update function of this node
 	static void Update(const FAnimationUpdateContext& InContext, FAnimNode_Base& InNode);
-
-	// Call the Evaluate function of this node
-	static void Evaluate(FPoseContext& InContext, FAnimNode_Base& InNode);
-
-	// Call the EvaluateComponentSpace function of this node
-	static void EvaluateComponentSpace(FComponentSpacePoseContext& InContext, FAnimNode_Base& InNode);
 };
 
 }}
diff --git a/Engine/Source/Runtime/Engine/Private/Animation/AnimInstanceProxy.cpp b/Engine/Source/Runtime/Engine/Private/Animation/AnimInstanceProxy.cpp
--- a/Engine/Source/Runtime/Engine/Private/Animation/AnimInstanceProxy.cpp
+++ b/Engine/Source/Runtime/Engine/Private/Animation/AnimInstanceProxy.cpp
@@ -296,8 +296,6 @@ void FAnimInstanceProxy::InitializeRootNode_WithRoot(FAnimNode_Base* InRootNode)
 	{
 		FAnimationUpdateSharedContext SharedContext;
 		FAnimationInitializeContext InitContext(this, &SharedContext);
-
-		UE::Anim::FNodeFunctionCaller::Initialize(InitContext, *InRootNode);
 		
 		if(InRootNode == RootNode)
 		{
@@ -1292,8 +1290,6 @@ void FAnimInstanceProxy::EvaluateAnimationNode_WithRoot(FPoseContext& Output, FA
 		}
 
 		TRACE_SCOPED_ANIM_NODE(Output);
-
-		UE::Anim::FNodeFunctionCaller::Evaluate(Output, *InRootNode);
 		
 		InRootNode->Evaluate_AnyThread(Output);
 	}
diff --git a/Engine/Source/Runtime/Engine/Private/Animation/AnimNodeBase.cpp b/Engine/Source/Runtime/Engine/Private/Animation/AnimNodeBase.cpp
--- a/Engine/Source/Runtime/Engine/Private/Animation/AnimNodeBase.cpp
+++ b/Engine/Source/Runtime/Engine/Private/Animation/AnimNodeBase.cpp
@@ -149,11 +149,6 @@ void FAnimNode_Base::ResetDynamics(ETeleportType InTeleportType)
 	PRAGMA_ENABLE_DEPRECATION_WARNINGS
 }
 
-const FAnimNodeFunctionRef& FAnimNode_Base::GetInitializeFunction() const
-{
-	return GET_ANIM_NODE_DATA(FAnimNodeFunctionRef, InitializeFunction);
-}
-
 const FAnimNodeFunctionRef& FAnimNode_Base::GetBecomeRelevantFunction() const
 {
 	return GET_ANIM_NODE_DATA(FAnimNodeFunctionRef, BecomeRelevantFunction);
@@ -164,11 +159,6 @@ const FAnimNodeFunctionRef& FAnimNode_Base::GetUpdateFunction() const
 	return GET_ANIM_NODE_DATA(FAnimNodeFunctionRef, UpdateFunction);
 }
 
-const FAnimNodeFunctionRef& FAnimNode_Base::GetEvaluateFunction() const
-{
-	return GET_ANIM_NODE_DATA(FAnimNodeFunctionRef, EvaluateFunction);
-}
-
 /////////////////////////////////////////////////////
 // FPoseLinkBase
 
@@ -214,7 +204,6 @@ void FPoseLinkBase::Initialize(const FAnimationInitializeContext& InContext)
 		FAnimationInitializeContext LinkContext(InContext);
 		LinkContext.SetNodeId(LinkID);
 		TRACE_SCOPED_ANIM_NODE(LinkContext);
-		UE::Anim::FNodeFunctionCaller::Initialize(LinkContext, *LinkedNode);
 		LinkedNode->Initialize_AnyThread(LinkContext);
 	}
 }
@@ -376,7 +365,6 @@ void FPoseLink::Evaluate(FPoseContext& Output)
 		{
 			Output.SetNodeId(LinkID);
 			TRACE_SCOPED_ANIM_NODE(Output);
-			UE::Anim::FNodeFunctionCaller::Evaluate(Output, *LinkedNode);
 			LinkedNode->Evaluate_AnyThread(Output);
 			TRACE_ANIM_NODE_BLENDABLE_ATTRIBUTES(Output, SourceID, LinkID);
 		}
@@ -487,7 +475,6 @@ void FComponentSpacePoseLink::EvaluateComponentSpace(FComponentSpacePoseContext&
 		{
 			Output.SetNodeId(LinkID);
 			TRACE_SCOPED_ANIM_NODE(Output);
-			UE::Anim::FNodeFunctionCaller::EvaluateComponentSpace(Output, *LinkedNode);
 			LinkedNode->EvaluateComponentSpace_AnyThread(Output);
 			TRACE_ANIM_NODE_BLENDABLE_ATTRIBUTES(Output, SourceID, LinkID);
 		}
diff --git a/Engine/Source/Runtime/Engine/Private/Animation/AnimNodeFunctionRef.cpp b/Engine/Source/Runtime/Engine/Private/Animation/AnimNodeFunctionRef.cpp
--- a/Engine/Source/Runtime/Engine/Private/Animation/AnimNodeFunctionRef.cpp
+++ b/Engine/Source/Runtime/Engine/Private/Animation/AnimNodeFunctionRef.cpp
@@ -28,7 +28,7 @@ namespace UE
 namespace Anim
 {
 
-template<typename ContextType>
+template<typename WrapperType, typename ContextType>
 static void CallFunctionHelper(const FAnimNodeFunctionRef& InFunction, ContextType InContext, FAnimNode_Base& InNode)
 {
 	if(InFunction.IsValid())
@@ -39,23 +39,15 @@ static void CallFunctionHelper(const FAnimNodeFunctionRef& InFunction, ContextTy
 			
 		struct FAnimNodeFunctionParams
 		{
-			FAnimExecutionContext ExecutionContext;
+			WrapperType ExecutionContext;
 			FAnimNodeReference NodeReference;
 		};
 			
-		FAnimNodeFunctionParams Params = { FAnimExecutionContext(ContextData), FAnimNodeReference(AnimInstance, InNode) };
+		FAnimNodeFunctionParams Params = { WrapperType(ContextData), FAnimNodeReference(AnimInstance, InNode) };
 			
 		InFunction.Call(AnimInstance, &Params);
 	}
 }
-
-void FNodeFunctionCaller::Initialize(const FAnimationInitializeContext& InContext, FAnimNode_Base& InNode)
-{
-	if(InNode.NodeData != nullptr)
-	{
-		CallFunctionHelper(InNode.GetInitializeFunction(), InContext, InNode);
-	}
-}
 	
 void FNodeFunctionCaller::BecomeRelevant(const FAnimationUpdateContext& InContext, FAnimNode_Base& InNode)
 {
@@ -68,7 +60,7 @@ void FNodeFunctionCaller::BecomeRelevant(const FAnimationUpdateContext& InContex
 			FAnimNodeRelevancyStatus Status = RelevancySubsystem.UpdateNodeRelevancy(InContext, InNode);
 			if(Status.HasJustBecomeRelevant())
 			{
-				CallFunctionHelper(Function, InContext, InNode);
+				CallFunctionHelper<FAnimUpdateContext>(Function, InContext, InNode);
 			}
 		}
 	}
@@ -78,23 +70,7 @@ void FNodeFunctionCaller::Update(const FAnimationUpdateContext& InContext, FAnim
 {
 	if(InNode.NodeData != nullptr)
 	{
-		CallFunctionHelper(InNode.GetUpdateFunction(), InContext, InNode);
-	}
-}
-
-void FNodeFunctionCaller::Evaluate(FPoseContext& InContext, FAnimNode_Base& InNode)
-{
-	if(InNode.NodeData != nullptr)
-	{
-		CallFunctionHelper(InNode.GetEvaluateFunction(), InContext, InNode);
-	}
-}
-
-void FNodeFunctionCaller::EvaluateComponentSpace(FComponentSpacePoseContext& InContext, FAnimNode_Base& InNode)
-{
-	if(InNode.NodeData != nullptr)
-	{
-		CallFunctionHelper(InNode.GetEvaluateFunction(), InContext, InNode);
+		CallFunctionHelper<FAnimUpdateContext>(InNode.GetUpdateFunction(), InContext, InNode);
 	}
 }
 
