Show mirror node in animation blueprint if no compatible mirror data tables exist

diff --git a/Engine/Source/Editor/AnimGraph/Classes/AnimGraphNode_Mirror.h b/Engine/Source/Editor/AnimGraph/Classes/AnimGraphNode_Mirror.h
--- a/Engine/Source/Editor/AnimGraph/Classes/AnimGraphNode_Mirror.h
+++ b/Engine/Source/Editor/AnimGraph/Classes/AnimGraphNode_Mirror.h
@@ -36,7 +36,7 @@ private:
 	static FText GetTitleGivenAssetInfo(const FText& AssetName);
 	void OnMirrorDataTableChanged(); 
 	void BindMirrorDataTableChangedDelegate(); 
-
+	bool HasMirrorDataTableForBlueprints(const TArray<UBlueprint*>& Blueprints) const;
 	/** Used for filtering in the Blueprint context menu when the sequence asset this node uses is unloaded */
 	FString UnloadedSkeletonName;
 };
diff --git a/Engine/Source/Editor/AnimGraph/Private/AnimGraphNode_Mirror.cpp b/Engine/Source/Editor/AnimGraph/Private/AnimGraphNode_Mirror.cpp
--- a/Engine/Source/Editor/AnimGraph/Private/AnimGraphNode_Mirror.cpp
+++ b/Engine/Source/Editor/AnimGraph/Private/AnimGraphNode_Mirror.cpp
@@ -208,7 +208,6 @@ void UAnimGraphNode_Mirror::GetMenuActions(FBlueprintActionDatabaseRegistrar& Ac
 	else if (QueryObject == GetClass())
 	{
 		FAssetRegistryModule& AssetRegistryModule = FModuleManager::LoadModuleChecked<FAssetRegistryModule>(TEXT("AssetRegistry"));
-		// define a filter to help in pulling UAnimSequence asset data from the registry
 		FARFilter Filter;
 		Filter.ClassNames.Add(UMirrorDataTable::StaticClass()->GetFName());
 		Filter.bRecursiveClasses = true;
@@ -232,6 +231,57 @@ void UAnimGraphNode_Mirror::GetMenuActions(FBlueprintActionDatabaseRegistrar& Ac
 			ActionRegistrar.AddBlueprintAction(Asset, NodeSpawner);
 		}
 	}
+
+	UClass* ActionKey = GetClass();
+	if (ActionRegistrar.IsOpenForRegistration(ActionKey))
+	{
+		// Add a default create option if there are no compatible tables.  This action is filtered out when 
+		// HasMirrorDataTableForBlueprints returns true because there is a table that can be used
+		UBlueprintNodeSpawner* EmptyNodeSpawner = UBlueprintNodeSpawner::Create(GetClass());
+		EmptyNodeSpawner->DefaultMenuSignature.MenuName = FText(LOCTEXT("MirrorNodeTitle", "Mirror"));
+		ActionRegistrar.AddBlueprintAction(ActionKey, EmptyNodeSpawner);
+	}
+}
+
+bool UAnimGraphNode_Mirror::HasMirrorDataTableForBlueprints(const TArray<UBlueprint*>& Blueprints) const
+{
+	// check to see if any mirror data table is available for the blueprint
+	FAssetRegistryModule& AssetRegistryModule = FModuleManager::LoadModuleChecked<FAssetRegistryModule>(TEXT("AssetRegistry"));
+	FARFilter Filter;
+	Filter.ClassNames.Add(UMirrorDataTable::StaticClass()->GetFName());
+
+	Filter.bRecursiveClasses = true;
+	TArray<FAssetData> MirrorDataTableList;
+	AssetRegistryModule.Get().GetAssets(Filter, MirrorDataTableList);
+
+	for (UBlueprint* Blueprint : Blueprints)
+	{
+		if (UAnimBlueprint* AnimBlueprint = Cast<UAnimBlueprint>(Blueprint))
+		{
+			for (auto AssetIt = MirrorDataTableList.CreateConstIterator(); AssetIt; ++AssetIt)
+			{
+				const FAssetData& Asset = *AssetIt;
+				if (Asset.IsAssetLoaded())
+				{
+					UMirrorDataTable* MirrorDataTable = Cast<UMirrorDataTable>(Asset.GetAsset());
+					if (AnimBlueprint->TargetSkeleton && AnimBlueprint->TargetSkeleton->IsCompatible(MirrorDataTable->Skeleton))
+					{
+						return true; 
+					}
+				}
+				else
+				{
+					FString TableSkeletonName;
+					Asset.GetTagValue("Skeleton", TableSkeletonName);
+					if (AnimBlueprint->TargetSkeleton && AnimBlueprint->TargetSkeleton->GetName() == TableSkeletonName)
+					{
+						return true;
+					}
+				}
+			}
+		}
+	}
+	return false;
 }
 
 bool UAnimGraphNode_Mirror::IsActionFilteredOut(class FBlueprintActionFilter const& Filter)
@@ -261,6 +311,15 @@ bool UAnimGraphNode_Mirror::IsActionFilteredOut(class FBlueprintActionFilter con
 					break;
 				}
 			}
+			else
+			{
+				// Only show the empty "Mirror" option if there is not a possible MirrorDataTable to select
+				if (HasMirrorDataTableForBlueprints(FilterContext.Blueprints))
+				{
+					bIsFilteredOut = true;
+					break;
+				}
+			}
 		}
 		else
 		{
