Fixed crash when editing callsites for event graph property accesses

diff --git a/Engine/Source/Editor/AnimGraph/Private/AnimBlueprintCompiler.cpp b/Engine/Source/Editor/AnimGraph/Private/AnimBlueprintCompiler.cpp
--- a/Engine/Source/Editor/AnimGraph/Private/AnimBlueprintCompiler.cpp
+++ b/Engine/Source/Editor/AnimGraph/Private/AnimBlueprintCompiler.cpp
@@ -196,6 +196,11 @@ void FAnimBlueprintCompilerContext::CreateClassVariablesFromBlueprint()
 			}
 		};
 
+		for (UEdGraph* UbergraphPage : Blueprint->UbergraphPages)
+		{
+			ForAllSubGraphs(UbergraphPage, ProcessGraph);
+		}
+
 		for (UEdGraph* Graph : Blueprint->FunctionGraphs)
 		{
 			ForAllSubGraphs(Graph, ProcessGraph);
