Anim node references

Added the abiity to tag and retrieve any anim graph node (similar to how we could reference linked anim graph nodes previously).
Ported linked anim graph nodes to use the new system
Added the ability to reference any anim graph node by tag (via a new custom node, spawnable from the context menu, with the appearance of an actor reference in a level blueprint)
Added tag display and editing in the bottom-right of anim graph nodes
Added new override point to SGraphNodeK2Var to allow for title widget parameters to be overriden by child classes

diff --git a/Engine/Source/Editor/AnimGraph/Private/AnimBlueprintCompiler.cpp b/Engine/Source/Editor/AnimGraph/Private/AnimBlueprintCompiler.cpp
--- a/Engine/Source/Editor/AnimGraph/Private/AnimBlueprintCompiler.cpp
+++ b/Engine/Source/Editor/AnimGraph/Private/AnimBlueprintCompiler.cpp
@@ -1917,11 +1917,10 @@ void FAnimBlueprintCompilerContext::ProcessFoldedPropertyRecords()
 	NewAnimBlueprintClass->NodeTypeMap.Empty();
 	
 	// First index & setup the node data
-	int32 NodeIndex = 0;
-	for(FAnimNodeData& NodeData : NewAnimBlueprintClass->AnimNodeData)
+	for(const TPair<int32, FProperty*>& IndexPropertyPair : AllocatedPropertiesByIndex)
 	{
-		const int32 NodePropertyIndex = AllocatedAnimNodeIndices.Num() - 1 - NodeIndex;
-		FStructProperty* NodeStructProperty = CastFieldChecked<FStructProperty>(AllocatedPropertiesByIndex.FindChecked(NodePropertyIndex));
+		int32 NodeIndex = AllocatedPropertiesByIndex.Num() - 1 - IndexPropertyPair.Key;
+		FStructProperty* NodeStructProperty = CastFieldChecked<FStructProperty>(IndexPropertyPair.Value);
 		const UScriptStruct* Struct = NodeStructProperty->Struct;
 		const FAnimNodeStructData AnimNodeStructData = NewAnimBlueprintClass->NodeTypeMap.Add(Struct, FAnimNodeStructData(Struct));
 		const int32 NumProperties = AnimNodeStructData.GetNumProperties();
@@ -1933,9 +1932,10 @@ void FAnimBlueprintCompilerContext::ProcessFoldedPropertyRecords()
 			NewAnimBlueprintClass->NodeTypeMap.Add(SuperStruct, FAnimNodeStructData(SuperStruct));
 			SuperStruct = Cast<UScriptStruct>(SuperStruct->GetSuperStruct());
 		}
-		
+
+		FAnimNodeData& NodeData = NewAnimBlueprintClass->AnimNodeData[NodeIndex];
 		NodeData.AnimClassInterface = NewAnimBlueprintClass;
-		NodeData.NodeIndex = NodeIndex++;
+		NodeData.NodeIndex = NodeIndex;
 		check(NumProperties >= 0);
 		NodeData.Entries.SetNum(NumProperties);
 		for(uint32& Entry : NodeData.Entries)
diff --git a/Engine/Source/Editor/AnimGraph/Private/AnimBlueprintCompiler.h b/Engine/Source/Editor/AnimGraph/Private/AnimBlueprintCompiler.h
--- a/Engine/Source/Editor/AnimGraph/Private/AnimBlueprintCompiler.h
+++ b/Engine/Source/Editor/AnimGraph/Private/AnimBlueprintCompiler.h
@@ -47,6 +47,7 @@ class FAnimBlueprintCompilerContext : public FKismetCompilerContext
 	friend class FAnimBlueprintCompilationBracketContext;
 	friend class FAnimBlueprintPostExpansionStepContext;
 	friend class FAnimBlueprintCopyTermDefaultsContext;
+	friend class UK2Node_AnimNodeReference;
 
 protected:
 	typedef FKismetCompilerContext Super;
diff --git a/Engine/Source/Editor/AnimGraph/Private/AnimBlueprintExtension.cpp b/Engine/Source/Editor/AnimGraph/Private/AnimBlueprintExtension.cpp
--- a/Engine/Source/Editor/AnimGraph/Private/AnimBlueprintExtension.cpp
+++ b/Engine/Source/Editor/AnimGraph/Private/AnimBlueprintExtension.cpp
@@ -11,6 +11,7 @@
 #include "AnimBlueprintExtension_PropertyAccess.h"
 #include "Kismet2/BlueprintEditorUtils.h"
 #include "AnimBlueprintExtension_NodeRelevancy.h"
+#include "AnimBlueprintExtension_Tag.h"
 
 // Set used to refresh extensions. Checks that an extension has a reference from an anim node for each refresh.
 static TSet<TSubclassOf<UAnimBlueprintExtension>> RefreshSet;
@@ -84,6 +85,11 @@ void UAnimBlueprintExtension::RequestExtensionsForNode(UAnimGraphNode_Base* InAn
 			ExtensionClasses.Add(UAnimBlueprintExtension_NodeRelevancy::StaticClass());
 		}
 
+		if(InAnimGraphNode->Tag != NAME_None)
+		{
+			ExtensionClasses.Add(UAnimBlueprintExtension_Tag::StaticClass());
+		}
+		
 		InAnimGraphNode->GetRequiredExtensions(ExtensionClasses);
 	
 		for(const TSubclassOf<UAnimBlueprintExtension>& ExtensionClass : ExtensionClasses)
diff --git a/Engine/Source/Editor/AnimGraph/Private/AnimBlueprintExtension_Tag.cpp b/Engine/Source/Editor/AnimGraph/Private/AnimBlueprintExtension_Tag.cpp
new file mode 100644
--- /dev/null
+++ b/Engine/Source/Editor/AnimGraph/Private/AnimBlueprintExtension_Tag.cpp
@@ -0,0 +1,54 @@
+// Copyright Epic Games, Inc. All Rights Reserved.
+
+#include "AnimBlueprintExtension_Tag.h"
+
+#include "AnimGraphNode_Base.h"
+#include "AnimGraphNode_LinkedAnimGraphBase.h"
+#include "IAnimBlueprintCompilationBracketContext.h"
+#include "IAnimBlueprintCompilationContext.h"
+
+#define LOCTEXT_NAMESPACE "UAnimBlueprintExtension_Tag"
+
+void UAnimBlueprintExtension_Tag::HandleStartCompilingClass(const UClass* InClass, IAnimBlueprintCompilationBracketContext& InCompilationContext, IAnimBlueprintGeneratedClassCompiledData& OutCompiledData)
+{
+	TaggedNodes.Empty();
+	Subsystem.NodeIndices.Empty();
+}
+
+void UAnimBlueprintExtension_Tag::HandleFinishCompilingClass(const UClass* InClass, IAnimBlueprintCompilationBracketContext& InCompilationContext, IAnimBlueprintGeneratedClassCompiledData& OutCompiledData)
+{
+	for(const TPair<FName, UAnimGraphNode_Base*>& TaggedNodePair : TaggedNodes)
+	{
+		if(const int32* IndexPtr = InCompilationContext.GetAllocatedAnimNodeIndices().Find(TaggedNodePair.Value))
+		{
+			Subsystem.NodeIndices.Add(TaggedNodePair.Value->GetTag(), *IndexPtr);
+		}
+	}
+}
+
+void UAnimBlueprintExtension_Tag::AddTaggedNode(UAnimGraphNode_Base* InNode, IAnimBlueprintCompilationContext& InCompilationContext)
+{
+	if(InNode->GetTag() != NAME_None)
+	{
+		if(UAnimGraphNode_Base** ExistingNode = TaggedNodes.Find(InNode->GetTag()))
+		{
+			InCompilationContext.GetMessageLog().Error(*FText::Format(LOCTEXT("DuplicateLabelError", "Nodes @@ and @@ have the same reference tag '{0}'"), FText::FromName(InNode->GetTag())).ToString(), InNode, *ExistingNode);
+		}
+		else
+		{
+			TaggedNodes.Add(InNode->GetTag(), InNode);
+		}
+	}
+}
+
+UAnimGraphNode_Base* UAnimBlueprintExtension_Tag::FindTaggedNode(FName InTag) const
+{
+	if(UAnimGraphNode_Base*const* ExistingNode = TaggedNodes.Find(InTag))
+	{
+		return *ExistingNode;
+	}
+
+	return nullptr;
+}
+
+#undef LOCTEXT_NAMESPACE
\ No newline at end of file
diff --git a/Engine/Source/Editor/AnimGraph/Private/AnimBlueprintExtension_Tag.h b/Engine/Source/Editor/AnimGraph/Private/AnimBlueprintExtension_Tag.h
new file mode 100644
--- /dev/null
+++ b/Engine/Source/Editor/AnimGraph/Private/AnimBlueprintExtension_Tag.h
@@ -0,0 +1,36 @@
+// Copyright Epic Games, Inc. All Rights Reserved.
+
+#pragma once
+
+#include "CoreMinimal.h"
+#include "AnimBlueprintExtension.h"
+#include "Animation/AnimSubsystem_Tag.h"
+#include "AnimBlueprintExtension_Tag.generated.h"
+
+class UAnimGraphNode_Base;
+class IAnimBlueprintCompilationContext;
+
+UCLASS(MinimalAPI)
+class UAnimBlueprintExtension_Tag : public UAnimBlueprintExtension
+{
+	GENERATED_BODY()
+
+public:
+	// Add a tagged node
+	void AddTaggedNode(UAnimGraphNode_Base* InNode, IAnimBlueprintCompilationContext& InCompilationContext);
+
+	// Find a tagged node in this animation blueprint
+	UAnimGraphNode_Base* FindTaggedNode(FName InTag) const;
+
+private:
+	// UAnimBlueprintExtension interface
+	virtual void HandleStartCompilingClass(const UClass* InClass, IAnimBlueprintCompilationBracketContext& InCompilationContext, IAnimBlueprintGeneratedClassCompiledData& OutCompiledData) override;
+	virtual void HandleFinishCompilingClass(const UClass* InClass, IAnimBlueprintCompilationBracketContext& InCompilationContext, IAnimBlueprintGeneratedClassCompiledData& OutCompiledData) override;
+
+private:
+	// Map of tag -> tagged nodes
+	TMap<FName, UAnimGraphNode_Base*> TaggedNodes;
+	
+	UPROPERTY()
+	FAnimSubsystem_Tag Subsystem;
+};
\ No newline at end of file
diff --git a/Engine/Source/Editor/AnimGraph/Private/AnimGraphNode_Base.cpp b/Engine/Source/Editor/AnimGraph/Private/AnimGraphNode_Base.cpp
--- a/Engine/Source/Editor/AnimGraph/Private/AnimGraphNode_Base.cpp
+++ b/Engine/Source/Editor/AnimGraph/Private/AnimGraphNode_Base.cpp
@@ -21,6 +21,7 @@
 #include "AnimBlueprintExtension_Base.h"
 #include "AnimBlueprintExtension_Attributes.h"
 #include "AnimBlueprintExtension_PropertyAccess.h"
+#include "AnimBlueprintExtension_Tag.h"
 #include "IAnimBlueprintCompilationContext.h"
 #include "AnimBlueprintCompilationContext.h"
 #include "AnimBlueprintExtension.h"
@@ -102,6 +103,11 @@ void UAnimGraphNode_Base::PostEditChangeProperty(FPropertyChangedEvent& Property
 		GetFNode()->UpdateFunction.SetFromFunction(UpdateFunction.ResolveMember<UFunction>(GetAnimBlueprint()->SkeletonGeneratedClass));
 		GetSchema()->ReconstructNode(*this);
 	}
+	else if(PropertyName == GET_MEMBER_NAME_CHECKED(UAnimGraphNode_Base, Tag))
+	{
+		GetAnimBlueprint()->RequestRefreshExtensions();
+		FBlueprintEditorUtils::MarkBlueprintAsStructurallyModified(GetAnimBlueprint());
+	}
 	
 	Super::PostEditChangeProperty(PropertyChangedEvent);
 
@@ -512,6 +518,13 @@ void UAnimGraphNode_Base::ProcessDuringCompilation(IAnimBlueprintCompilationCont
 	GetFNode()->BecomeRelevantFunction.SetFromFunction(BecomeRelevantFunction.ResolveMember<UFunction>(GetBlueprintClassFromNode()));
 	GetFNode()->UpdateFunction.SetFromFunction(UpdateFunction.ResolveMember<UFunction>(GetBlueprintClassFromNode())); 
 
+	// Insert tag, if any
+	if(Tag != NAME_None)
+	{
+		UAnimBlueprintExtension_Tag* TagExtension = UAnimBlueprintExtension::GetExtension<UAnimBlueprintExtension_Tag>(GetAnimBlueprint());
+		TagExtension->AddTaggedNode(this, InCompilationContext);
+	}
+	
 	// Call the override point
 	OnProcessDuringCompilation(InCompilationContext, OutCompiledData);
 }
@@ -628,6 +641,11 @@ void UAnimGraphNode_Base::AddSearchMetaDataInfo(TArray<struct FSearchTagDataPair
 		OutTaggedMetaData.Add(FSearchTagDataPair(FFindInBlueprintSearchTags::FiB_Name, FText::FromName(BindingPair.Key)));
 		OutTaggedMetaData.Add(FSearchTagDataPair(LOCTEXT("Binding", "Binding"), BindingPair.Value.PathAsText));
 	}
+
+	if(Tag != NAME_None)
+	{
+		OutTaggedMetaData.Add(FSearchTagDataPair(LOCTEXT("Tag", "Tag"), FText::FromName(Tag)));
+	}
 }
 
 bool UAnimGraphNode_Base::IsPinExposedAndLinked(const FString& InPinName, const EEdGraphPinDirection InDirection) const
@@ -657,6 +675,17 @@ bool UAnimGraphNode_Base::IsPinBindable(const UEdGraphPin* InPin) const
 	return false;
 }
 
+void UAnimGraphNode_Base::SetTag(FName InTag)
+{
+	FName OldTag = Tag;
+	Tag = InTag;
+	if(Tag != OldTag)
+	{
+		GetAnimBlueprint()->RequestRefreshExtensions();
+		FBlueprintEditorUtils::MarkBlueprintAsStructurallyModified(GetAnimBlueprint());
+	}
+}
+
 FProperty* UAnimGraphNode_Base::GetPinProperty(const UEdGraphPin* InPin) const
 {
 	// Compare FName without number to make sure we catch array properties that are split into multiple pins
diff --git a/Engine/Source/Editor/AnimGraph/Private/AnimGraphNode_LinkedAnimGraph.cpp b/Engine/Source/Editor/AnimGraph/Private/AnimGraphNode_LinkedAnimGraph.cpp
--- a/Engine/Source/Editor/AnimGraph/Private/AnimGraphNode_LinkedAnimGraph.cpp
+++ b/Engine/Source/Editor/AnimGraph/Private/AnimGraphNode_LinkedAnimGraph.cpp
@@ -11,6 +11,7 @@
 #include "BlueprintActionDatabaseRegistrar.h"
 #include "IAnimBlueprintCopyTermDefaultsContext.h"
 #include "KismetCompiler.h"
+#include "UObject/UE5MainStreamObjectVersion.h"
 
 #define LOCTEXT_NAMESPACE "UAnimGraphNode_LinkedAnimGraph"
 
@@ -134,4 +135,17 @@ bool UAnimGraphNode_LinkedAnimGraph::IsActionFilteredOut(class FBlueprintActionF
 	return bIsFilteredOut;
 }
 
+void UAnimGraphNode_LinkedAnimGraph::Serialize(FArchive& Ar)
+{
+	Super::Serialize(Ar);
+
+	Ar.UsingCustomVersion(FUE5MainStreamObjectVersion::GUID);
+	
+	if(Ar.IsLoading() && Ar.CustomVer(FUE5MainStreamObjectVersion::GUID) < FUE5MainStreamObjectVersion::AnimGraphNodeTaggingAdded)
+	{
+		// Transfer old tag to new system
+		SetTagInternal(Node.Tag_DEPRECATED);
+	}
+}
+
 #undef LOCTEXT_NAMESPACE
\ No newline at end of file
diff --git a/Engine/Source/Editor/AnimGraph/Private/AnimGraphNode_LinkedAnimGraphBase.cpp b/Engine/Source/Editor/AnimGraph/Private/AnimGraphNode_LinkedAnimGraphBase.cpp
--- a/Engine/Source/Editor/AnimGraph/Private/AnimGraphNode_LinkedAnimGraphBase.cpp
+++ b/Engine/Source/Editor/AnimGraph/Private/AnimGraphNode_LinkedAnimGraphBase.cpp
@@ -89,12 +89,7 @@ FText UAnimGraphNode_LinkedAnimGraphBase::GetNodeTitle(ENodeTitleType::Type Titl
 	}
 	else if(TitleType == ENodeTitleType::ListView)
 	{
-		if(Node.Tag != NAME_None)
-		{
-			Args.Add(TEXT("Tag"), FText::FromName(Node.Tag));
-			return FText::Format(LOCTEXT("TitleListFormatTagged", "{TargetClass} ({Tag})"), Args);
-		}
-		else if(TargetClass)
+		if(TargetClass)
 		{
 			return FText::Format(LOCTEXT("TitleListFormat", "{TargetClass}"), Args);
 		}
@@ -105,12 +100,7 @@ FText UAnimGraphNode_LinkedAnimGraphBase::GetNodeTitle(ENodeTitleType::Type Titl
 	}
 	else
 	{
-		if(Node.Tag != NAME_None)
-		{
-			Args.Add(TEXT("Tag"), FText::FromName(Node.Tag));
-			return FText::Format(LOCTEXT("TitleFormatTagged", "{TargetClass} ({Tag})\n{NodeType} "), Args);
-		}
-		else if(TargetClass)
+		if(TargetClass)
 		{
 			return FText::Format(LOCTEXT("TitleFormat", "{TargetClass}\n{NodeType}"), Args);
 		}
@@ -140,28 +130,6 @@ void UAnimGraphNode_LinkedAnimGraphBase::ValidateAnimNodeDuringCompilation(USkel
 
 	const FAnimNode_LinkedAnimGraph& Node = *GetLinkedAnimGraphNode();
 
-	// Check for duplicate tags in this anim blueprint
-	for(UEdGraph* Graph : Graphs)
-	{
-		TArray<UAnimGraphNode_LinkedAnimGraphBase*> LinkedAnimGraphNodes;
-		Graph->GetNodesOfClass(LinkedAnimGraphNodes);
-
-		for(UAnimGraphNode_LinkedAnimGraphBase* LinkedAnimGraphNode : LinkedAnimGraphNodes)
-		{
-			if(LinkedAnimGraphNode == OriginalNode)
-			{
-				continue;
-			}
-
-			FAnimNode_LinkedAnimGraph& InnerNode = *LinkedAnimGraphNode->GetLinkedAnimGraphNode();
-
-			if(InnerNode.Tag != NAME_None && InnerNode.Tag == Node.Tag)
-			{
-				MessageLog.Error(*FText::Format(LOCTEXT("DuplicateTagErrorFormat", "Node @@ and node @@ both have the same tag '{0}'."), FText::FromName(Node.Tag)).ToString(), this, LinkedAnimGraphNode);
-			}
-		}
-	}
-
 	// Check we don't try to spawn our own blueprint
 	if(GetTargetClass() == AnimBP->GetAnimBlueprintGeneratedClass())
 	{
diff --git a/Engine/Source/Editor/AnimGraph/Private/K2Node_AnimNodeReference.cpp b/Engine/Source/Editor/AnimGraph/Private/K2Node_AnimNodeReference.cpp
new file mode 100644
--- /dev/null
+++ b/Engine/Source/Editor/AnimGraph/Private/K2Node_AnimNodeReference.cpp
@@ -0,0 +1,153 @@
+// Copyright Epic Games, Inc. All Rights Reserved.
+
+#include "K2Node_AnimNodeReference.h"
+
+#include "AnimBlueprintCompiler.h"
+#include "AnimBlueprintExtension.h"
+#include "AnimBlueprintExtension_Tag.h"
+#include "AnimExecutionContextLibrary.h"
+#include "BlueprintActionDatabaseRegistrar.h"
+#include "BlueprintNodeSpawner.h"
+#include "EditorCategoryUtils.h"
+#include "Animation/AnimNodeReference.h"
+#include "Kismet2/BlueprintEditorUtils.h"
+#include "FindInBlueprintManager.h"
+#include "K2Node_CallFunction.h"
+
+#define LOCTEXT_NAMESPACE "K2Node_AnimNodeReference"
+
+void UK2Node_AnimNodeReference::ExpandNode(FKismetCompilerContext& InCompilerContext, UEdGraph* InSourceGraph)
+{
+	bool bFound = false;
+	
+	UAnimBlueprintExtension_Tag* Extension = UAnimBlueprintExtension::FindExtension<UAnimBlueprintExtension_Tag>(Cast<UAnimBlueprint>(GetBlueprint()));
+	if(Extension)
+	{
+		bFound = true;
+		
+		if(UAnimGraphNode_Base* LabelledNode = Extension->FindTaggedNode(Tag))
+		{
+			// Expand to function call
+			UK2Node_CallFunction* CallFunctionNode = InCompilerContext.SpawnIntermediateNode<UK2Node_CallFunction>(this, InSourceGraph);
+			CallFunctionNode->SetFromFunction(GetDefault<UAnimExecutionContextLibrary>()->FindFunctionChecked(GET_FUNCTION_NAME_CHECKED(UAnimExecutionContextLibrary, GetAnimNodeReference)));
+			CallFunctionNode->AllocateDefaultPins();
+
+			UEdGraphPin* ValuePin = FindPinChecked(TEXT("Value"), EGPD_Output);
+
+			if(UEdGraphPin* IndexPin = CallFunctionNode->FindPin(TEXT("Index"), EGPD_Input))
+			{
+				FAnimBlueprintCompilerContext& AnimBlueprintCompilerContext = static_cast<FAnimBlueprintCompilerContext&>(InCompilerContext);
+
+				const int32 NodeIndex = AnimBlueprintCompilerContext.GetAllocationIndexOfNode(LabelledNode);
+				IndexPin->DefaultValue = FString::FromInt(NodeIndex);
+			}
+
+			if(UEdGraphPin* ReturnValuePin = CallFunctionNode->FindPin(TEXT("ReturnValue"), EGPD_Output))
+			{
+				InCompilerContext.MovePinLinksToIntermediate(*ValuePin, *ReturnValuePin);
+			}
+		}
+	}
+
+	if(!bFound)
+	{
+		InCompilerContext.MessageLog.Error(*FText::Format(LOCTEXT("MissingTaggedNodeError", "@@ cannot find referenced node with tag '{0}'"), FText::FromName(Tag)).ToString(), this);
+	}
+}
+
+void UK2Node_AnimNodeReference::AllocateDefaultPins()
+{
+	FEdGraphPinType PinType;
+	PinType.PinCategory = UEdGraphSchema_K2::PC_Struct;
+	PinType.PinSubCategoryObject = FAnimNodeReference::StaticStruct();
+
+	CreatePin(EGPD_Output, PinType, TEXT("Value"));
+}
+
+FText UK2Node_AnimNodeReference::GetLabelText() const
+{
+	return FText::FromName(Tag);
+}
+
+FText UK2Node_AnimNodeReference::GetNodeTitle(ENodeTitleType::Type TitleType) const
+{
+	return FText::Format(LOCTEXT("NodeTitleFormat", "Node Reference: {0}"), FText::FromName(Tag));
+}
+
+void UK2Node_AnimNodeReference::AddSearchMetaDataInfo(TArray<FSearchTagDataPair>& OutTaggedMetaData) const
+{
+	Super::AddSearchMetaDataInfo(OutTaggedMetaData);
+
+	if(Tag != NAME_None)
+	{
+		OutTaggedMetaData.Emplace(LOCTEXT("Tag", "Tag"), FText::FromName(Tag));
+	}
+}
+
+void UK2Node_AnimNodeReference::GetMenuActions(FBlueprintActionDatabaseRegistrar& ActionRegistrar) const
+{
+	UClass* NodeClass = GetClass();
+	const UAnimBlueprint* AnimBlueprint = Cast<UAnimBlueprint>(ActionRegistrar.GetActionKeyFilter());
+	if (AnimBlueprint && ActionRegistrar.IsOpenForRegistration(AnimBlueprint))
+	{
+		// Let us spawn a node for any labelled anim graph node we find
+		TArray<UAnimGraphNode_Base*> AnimGraphNodes;
+		FBlueprintEditorUtils::GetAllNodesOfClass(AnimBlueprint, AnimGraphNodes);
+
+		for(UAnimGraphNode_Base* AnimGraphNodeToReference : AnimGraphNodes)
+		{
+			if(AnimGraphNodeToReference->GetTag() != NAME_None)
+			{
+				UBlueprintNodeSpawner* NodeSpawner = UBlueprintNodeSpawner::Create(NodeClass);
+				check(NodeSpawner);
+				NodeSpawner->CustomizeNodeDelegate = UBlueprintNodeSpawner::FCustomizeNodeDelegate::CreateLambda([WeakAnimGraphNodeToReference = TWeakObjectPtr<UAnimGraphNode_Base>(AnimGraphNodeToReference)](UEdGraphNode* InNode, bool bInIsTemplate)
+				{
+					UK2Node_AnimNodeReference* AnimNodeReference = CastChecked<UK2Node_AnimNodeReference>(InNode);
+					AnimNodeReference->Tag = WeakAnimGraphNodeToReference->GetTag();
+				});
+				ActionRegistrar.AddBlueprintAction(AnimBlueprint, NodeSpawner);
+			}
+		}
+	}
+}
+
+bool UK2Node_AnimNodeReference::IsActionFilteredOut(class FBlueprintActionFilter const& Filter)
+{
+	bool bIsFilteredOut = false;
+	FBlueprintActionContext const& FilterContext = Filter.Context;
+
+	if(Tag != NAME_None)
+	{
+		for (UBlueprint* Blueprint : FilterContext.Blueprints)
+		{
+			if(!Blueprint->IsA<UAnimBlueprint>())
+			{
+				// Not an animation Blueprint, cannot use
+				bIsFilteredOut = true;
+				break;
+			}
+		}
+	}
+	else
+	{
+		bIsFilteredOut = true;
+	}
+	return bIsFilteredOut;
+}
+
+FText UK2Node_AnimNodeReference::GetMenuCategory() const
+{
+	return FEditorCategoryUtils::GetCommonCategory(FCommonEditorCategory::Variables);
+}
+
+bool UK2Node_AnimNodeReference::IsCompatibleWithGraph(UEdGraph const* TargetGraph) const
+{
+	return Cast<UAnimBlueprint>(FBlueprintEditorUtils::FindBlueprintForGraph(TargetGraph)) != nullptr;
+}
+
+FText UK2Node_AnimNodeReference::GetTooltipText() const
+{
+	return LOCTEXT("NodeReferenceTooltip", "Gets a reference to an anim graph node in this Animation Blueprint");
+}
+
+#undef LOCTEXT_NAMESPACE
\ No newline at end of file
diff --git a/Engine/Source/Editor/AnimGraph/Public/AnimBlueprintExtension.h b/Engine/Source/Editor/AnimGraph/Public/AnimBlueprintExtension.h
--- a/Engine/Source/Editor/AnimGraph/Public/AnimBlueprintExtension.h
+++ b/Engine/Source/Editor/AnimGraph/Public/AnimBlueprintExtension.h
@@ -35,12 +35,19 @@ public:
 	/** Request an anim blueprint extension for an anim blueprint. It is illegal to perform this operation once compilation has commenced, use GetExtension instead. */
 	static UAnimBlueprintExtension* RequestExtension(UAnimBlueprint* InAnimBlueprint, TSubclassOf<UAnimBlueprintExtension> InExtensionType);
 
-	/** Get an already-requested extension for an anim blueprint. */
+	/** Get an already-requested extension for an anim blueprint. Will assert if the extension is not found. */
 	template<typename ExtensionType>
-    static ExtensionType* GetExtension(UAnimBlueprint* InAnimBlueprint)
+	static ExtensionType* GetExtension(UAnimBlueprint* InAnimBlueprint)
 	{
 		return CastChecked<ExtensionType>(GetExtension(InAnimBlueprint, ExtensionType::StaticClass()));
 	}
+
+	/** Get an already-requested extension for an anim blueprint. @return nullptr if the extension is not found. */
+	template<typename ExtensionType>
+	static ExtensionType* FindExtension(UAnimBlueprint* InAnimBlueprint)
+	{
+		return Cast<ExtensionType>(GetExtension(InAnimBlueprint, ExtensionType::StaticClass()));
+	}
 	
 	/** Get an already-requested an anim blueprint extension for an anim blueprint. */
 	static UAnimBlueprintExtension* GetExtension(UAnimBlueprint* InAnimBlueprint, TSubclassOf<UAnimBlueprintExtension> InExtensionType);
diff --git a/Engine/Source/Editor/AnimGraph/Public/AnimGraphNode_Base.h b/Engine/Source/Editor/AnimGraph/Public/AnimGraphNode_Base.h
--- a/Engine/Source/Editor/AnimGraph/Public/AnimGraphNode_Base.h
+++ b/Engine/Source/Editor/AnimGraph/Public/AnimGraphNode_Base.h
@@ -220,7 +220,13 @@ class ANIMGRAPH_API UAnimGraphNode_Base : public UK2Node
 	// Function called when the node is updated
 	UPROPERTY(EditAnywhere, Category = Functions, meta=(FunctionReference, AllowFunctionLibraries, PrototypeFunction="/Script/AnimGraphRuntime.AnimExecutionContextLibrary.Prototype_ThreadSafeAnimUpdateCall"), DisplayName="On Update")
 	FMemberReference UpdateFunction;
-	
+
+private:
+	// Optional reference tag name. If this is set then this node can be referenced from elsewhere in this animation blueprint using an anim node reference
+	UPROPERTY(EditAnywhere, Category = Tag)
+	FName Tag;
+
+public:
 	// UObject interface
 	virtual void PostEditChangeProperty(FPropertyChangedEvent& PropertyChangedEvent) override;
 	virtual void PreEditChange(FProperty* PropertyAboutToChange) override;
@@ -453,6 +459,12 @@ class ANIMGRAPH_API UAnimGraphNode_Base : public UK2Node
 	
 	// Check whether the named pin is bindable
 	bool IsPinBindable(const UEdGraphPin* InPin) const;
+
+	// Get the tag for this node, if any
+	FName GetTag() const { return Tag; }
+
+	// Set the tag for this node
+	void SetTag(FName InTag);
 	
 protected:
 	friend class FAnimBlueprintCompilerContext;
@@ -461,6 +473,9 @@ protected:
 	friend class UAnimBlueprintExtension_Base;
 	friend class SAnimationGraphNode;
 
+	// Set the tag for this node but without regenerating any BP data for tagging
+	void SetTagInternal(FName InTag) { Tag = InTag; }
+	
 	// Gets the animation FNode type represented by this ed graph node
 	UScriptStruct* GetFNodeType() const;
 
diff --git a/Engine/Source/Editor/AnimGraph/Public/AnimGraphNode_LinkedAnimGraph.h b/Engine/Source/Editor/AnimGraph/Public/AnimGraphNode_LinkedAnimGraph.h
--- a/Engine/Source/Editor/AnimGraph/Public/AnimGraphNode_LinkedAnimGraph.h
+++ b/Engine/Source/Editor/AnimGraph/Public/AnimGraphNode_LinkedAnimGraph.h
@@ -19,7 +19,10 @@ public:
 
 	UPROPERTY(EditAnywhere, Category = Settings)
 	FAnimNode_LinkedAnimGraph Node;
-
+	
+	// Begin UObject
+	virtual void Serialize(FArchive& Ar) override;
+	
 	// Begin UEdGraphNode
 	virtual void PostPasteNode() override;
 
diff --git a/Engine/Source/Editor/AnimGraph/Public/K2Node_AnimNodeReference.h b/Engine/Source/Editor/AnimGraph/Public/K2Node_AnimNodeReference.h
new file mode 100644
--- /dev/null
+++ b/Engine/Source/Editor/AnimGraph/Public/K2Node_AnimNodeReference.h
@@ -0,0 +1,39 @@
+// Copyright Epic Games, Inc. All Rights Reserved.
+
+#pragma once
+
+#include "CoreMinimal.h"
+#include "K2Node.h"
+#include "AnimGraphNode_Base.h"
+
+#include "K2Node_AnimNodeReference.generated.h"
+
+UCLASS()
+class ANIMGRAPH_API UK2Node_AnimNodeReference : public UK2Node
+{
+public:
+	GENERATED_BODY()
+
+	// Get the text used for the node's label
+	FText GetLabelText() const;
+	
+private:
+	/** The node tag we reference */
+	UPROPERTY()
+	FName Tag;
+
+	// UEdGraphNode interface
+	virtual void AllocateDefaultPins() override;
+	virtual FText GetNodeTitle(ENodeTitleType::Type TitleType) const override;
+	virtual void AddSearchMetaDataInfo(TArray<FSearchTagDataPair>& OutTaggedMetaData) const override;
+	virtual bool IsCompatibleWithGraph(UEdGraph const* TargetGraph) const override;
+	
+	// UK2Node interface
+	virtual void ExpandNode(class FKismetCompilerContext& CompilerContext, UEdGraph* SourceGraph) override;
+	virtual bool IsNodePure() const override { return true; }
+	virtual void GetMenuActions(FBlueprintActionDatabaseRegistrar& ActionRegistrar) const override;
+	virtual bool IsActionFilteredOut(class FBlueprintActionFilter const& Filter) override;
+	virtual FText GetMenuCategory() const override;
+	virtual bool DrawNodeAsVariable() const override { return true; }
+	virtual FText GetTooltipText() const override;
+};
\ No newline at end of file
diff --git a/Engine/Source/Editor/AnimationBlueprintEditor/Private/AnimationGraphFactory.cpp b/Engine/Source/Editor/AnimationBlueprintEditor/Private/AnimationGraphFactory.cpp
--- a/Engine/Source/Editor/AnimationBlueprintEditor/Private/AnimationGraphFactory.cpp
+++ b/Engine/Source/Editor/AnimationBlueprintEditor/Private/AnimationGraphFactory.cpp
@@ -37,6 +37,8 @@
 #include "KismetPins/SGraphPinExec.h"
 #include "AnimGraphNode_BlendSpaceGraph.h"
 #include "AnimationNodes/SGraphNodeBlendSpaceGraph.h"
+#include "K2Node_AnimNodeReference.h"
+#include "AnimationNodes/SAnimNodeReference.h"
 
 TSharedPtr<class SGraphNode> FAnimationGraphNodeFactory::CreateNode(class UEdGraphNode* InNode) const 
 {
@@ -91,7 +93,11 @@ TSharedPtr<class SGraphNode> FAnimationGraphNodeFactory::CreateNode(class UEdGra
 	{
 		return SNew(SGraphNodeAnimStateEntry, EntryNode);
 	}
-
+	else if (UK2Node_AnimNodeReference* AnimNodeReference = Cast<UK2Node_AnimNodeReference>(InNode))
+	{
+		return SNew(SAnimNodeReference, AnimNodeReference);
+	}
+	
 	return nullptr;
 }
 
diff --git a/Engine/Source/Editor/AnimationBlueprintEditor/Private/AnimationNodes/SAnimNodeReference.cpp b/Engine/Source/Editor/AnimationBlueprintEditor/Private/AnimationNodes/SAnimNodeReference.cpp
new file mode 100644
--- /dev/null
+++ b/Engine/Source/Editor/AnimationBlueprintEditor/Private/AnimationNodes/SAnimNodeReference.cpp
@@ -0,0 +1,67 @@
+// Copyright Epic Games, Inc. All Rights Reserved.
+
+#include "SAnimNodeReference.h"
+
+#include "SLevelOfDetailBranchNode.h"
+#include "EditorStyleSet.h"
+#include "K2Node_AnimNodeReference.h"
+#include "SGraphPin.h"
+#include "Widgets/Layout/SWrapBox.h"
+#include "Widgets/Layout/SSpacer.h"
+#include "Widgets/Text/STextBlock.h"
+
+#define LOCTEXT_NAMESPACE "SAnimNodeReferenceNode"
+
+void SAnimNodeReference::Construct(const FArguments& InArgs, UK2Node_AnimNodeReference* InNode)
+{
+	GraphNode = InNode;
+	UpdateGraphNode();
+}
+
+TSharedRef<SWidget> SAnimNodeReference::UpdateTitleWidget(FText InTitleText, TSharedPtr<SWidget> InTitleWidget, EHorizontalAlignment& InOutTitleHAlign, FMargin& InOutTitleMargin) const
+{
+	UK2Node_AnimNodeReference* K2Node_AnimNodeReference = CastChecked<UK2Node_AnimNodeReference>(GraphNode);
+
+	InTitleWidget =
+		SNew(SLevelOfDetailBranchNode)
+		.UseLowDetailSlot(this, &SAnimNodeReference::UseLowDetailNodeTitles)
+		.LowDetail()
+		[
+			SNew(SSpacer)
+		]
+		.HighDetail()
+		[
+			SNew(SVerticalBox)
+			+SVerticalBox::Slot()
+			.AutoHeight()
+			[
+				SNew(STextBlock)
+				.Text_Lambda([K2Node_AnimNodeReference]()
+				{
+					return K2Node_AnimNodeReference->GetLabelText();
+				})
+			]
+			+SVerticalBox::Slot()
+			.AutoHeight()
+			[
+				SNew(STextBlock)
+				.Text(LOCTEXT("NodeSubTitle", "Anim Node Reference"))
+				.TextStyle(&FEditorStyle::Get().GetWidgetStyle<FTextBlockStyle>("AnimGraph.AnimNodeReference.Subtitle"))
+			]
+		];
+
+	InOutTitleHAlign = HAlign_Left;
+	InOutTitleMargin = FMargin(12.0f, 8.0f, 36.0f, 6.0f);
+	
+	return InTitleWidget.ToSharedRef();
+}
+
+TSharedPtr<SGraphPin> SAnimNodeReference::CreatePinWidget(UEdGraphPin* Pin) const
+{
+	TSharedPtr<SGraphPin> DefaultWidget = SGraphNodeK2Var::CreatePinWidget(Pin);
+	DefaultWidget->SetShowLabel(false);
+
+	return DefaultWidget;
+}
+
+#undef LOCTEXT_NAMESPACE
\ No newline at end of file
diff --git a/Engine/Source/Editor/AnimationBlueprintEditor/Private/AnimationNodes/SAnimNodeReference.h b/Engine/Source/Editor/AnimationBlueprintEditor/Private/AnimationNodes/SAnimNodeReference.h
new file mode 100644
--- /dev/null
+++ b/Engine/Source/Editor/AnimationBlueprintEditor/Private/AnimationNodes/SAnimNodeReference.h
@@ -0,0 +1,24 @@
+// Copyright Epic Games, Inc. All Rights Reserved.
+
+#pragma once
+
+#include "CoreMinimal.h"
+#include "KismetNodes/SGraphNodeK2Var.h"
+
+class UK2Node_AnimNodeReference;
+
+class SAnimNodeReference : public SGraphNodeK2Var
+{
+public:
+	SLATE_BEGIN_ARGS(SAnimNodeReference) {}
+	SLATE_END_ARGS()
+
+public:
+	void Construct(const FArguments& InArgs, UK2Node_AnimNodeReference* InNode);
+
+	// SGraphNode interface
+	virtual TSharedPtr<SGraphPin> CreatePinWidget(UEdGraphPin* Pin) const override;
+	
+	// SGraphNodeK2Var interface
+	virtual TSharedRef<SWidget> UpdateTitleWidget(FText InTitleText, TSharedPtr<SWidget> InTitleWidget, EHorizontalAlignment& InOutTitleHAlign, FMargin& InOutTitleMargin) const override;
+};
\ No newline at end of file
diff --git a/Engine/Source/Editor/AnimationBlueprintEditor/Private/AnimationNodes/SAnimationGraphNode.cpp b/Engine/Source/Editor/AnimationBlueprintEditor/Private/AnimationNodes/SAnimationGraphNode.cpp
--- a/Engine/Source/Editor/AnimationBlueprintEditor/Private/AnimationNodes/SAnimationGraphNode.cpp
+++ b/Engine/Source/Editor/AnimationBlueprintEditor/Private/AnimationNodes/SAnimationGraphNode.cpp
@@ -24,6 +24,7 @@
 #include "IPropertyRowGenerator.h"
 #include "IDetailTreeNode.h"
 #include "Widgets/Layout/SGridPanel.h"
+#include "Widgets/Text/SInlineEditableTextBlock.h"
 
 #define LOCTEXT_NAMESPACE "AnimationGraphNode"
 
@@ -280,17 +281,29 @@ void SAnimationGraphNode::CreateBelowPinControls(TSharedPtr<SVerticalBox> MainBo
 {
 	if (UAnimGraphNode_Base* AnimNode = CastChecked<UAnimGraphNode_Base>(GraphNode, ECastCheckedType::NullAllowed))
 	{
-		TSharedRef<SWidget> NodeFunctionsWidget = CreateNodeFunctionsWidget(AnimNode, MakeAttributeSP(this, &SAnimationGraphNode::UseLowDetailNodeTitles));
-		
+		auto UseLowDetailNode = [this]()
+		{
+			return GetCurrentLOD() <= EGraphRenderingLOD::LowDetail;
+		};
+
 		// Insert above the error reporting bar
-		MainBox->InsertSlot(FMath::Max(0, MainBox->NumSlots() - 1))
+		MainBox->InsertSlot(FMath::Max(0, MainBox->NumSlots() - TagAndFunctionsSlotReverseIndex))
 		.AutoHeight()
 		.Padding(4.0f)
 		[
-			SNew(SBox)
+			SNew(SVerticalBox)
 			.IsEnabled_Lambda([this](){ return IsNodeEditable(); })
+			+SVerticalBox::Slot()
+			.AutoHeight()
+			[
+				CreateNodeFunctionsWidget(AnimNode, MakeAttributeLambda(UseLowDetailNode))
+			]
+			+SVerticalBox::Slot()
+			.AutoHeight()
+			.Padding(4.0f, 0.0f, 4.0f, 4.0f)
+			.HAlign(HAlign_Right)
 			[
-				NodeFunctionsWidget
+				CreateNodeTagWidget(AnimNode, MakeAttributeLambda(UseLowDetailNode))
 			]
 		];
 	}
@@ -308,6 +321,8 @@ public:
 
 	void Construct(const FArguments& InArgs, UAnimGraphNode_Base* InNode)
 	{
+		UseLowDetail = InArgs._UseLowDetail;
+		
 		FPropertyEditorModule& PropertyEditorModule = FModuleManager::Get().LoadModuleChecked<FPropertyEditorModule>("PropertyEditor");
 		PropertyRowGenerator = PropertyEditorModule.CreatePropertyRowGenerator(FPropertyRowGeneratorArgs());
 		PropertyRowGenerator->RegisterInstancedCustomPropertyTypeLayout(FMemberReference::StaticStruct()->GetFName(), FOnGetPropertyTypeCustomizationInstance::CreateStatic(&FBlueprintMemberReferenceDetails::MakeInstance));
@@ -433,6 +448,26 @@ TSharedRef<SWidget> SAnimationGraphNode::CreateNodeFunctionsWidget(UAnimGraphNod
 		.UseLowDetail(InUseLowDetail);
 }
 
+TSharedRef<SWidget> SAnimationGraphNode::CreateNodeTagWidget(UAnimGraphNode_Base* InAnimNode, TAttribute<bool> InUseLowDetail)
+{
+	return SNew(SLevelOfDetailBranchNode)
+		.Visibility_Lambda([InAnimNode](){ return InAnimNode->Tag != NAME_None ? EVisibility::Visible : EVisibility::Collapsed; })
+		.UseLowDetailSlot(InUseLowDetail)
+		.LowDetail()
+		[
+			SNew(SSpacer)
+			.Size(FVector2D(24.0f, 24.f))
+		]
+		.HighDetail()
+		[
+			SNew(SInlineEditableTextBlock)
+			.ToolTipText_Lambda([InAnimNode](){ return FText::Format(LOCTEXT("TagFormat_Tooltip", "Tag: {0}\nThis node can be referenced elsewhere in this Anim Blueprint using this tag"), FText::FromName(InAnimNode->GetTag())); })
+			.Style(&FEditorStyle::Get().GetWidgetStyle<FInlineEditableTextBlockStyle>("AnimGraph.Node.Tag"))
+			.Text_Lambda([InAnimNode](){ return FText::FromName(InAnimNode->GetTag()); })
+			.OnTextCommitted_Lambda([InAnimNode](const FText& InText, ETextCommit::Type InCommitType){ InAnimNode->SetTag(*InText.ToString()); })
+		];
+}
+
 void SAnimationGraphNode::ReconfigurePinWidgetsForPropertyBindings(UAnimGraphNode_Base* InAnimGraphNode, TSharedRef<SGraphNode> InGraphNodeWidget, TFunctionRef<TSharedPtr<SGraphPin>(UEdGraphPin*)> InFindWidgetForPin)
 {
 	for(UEdGraphPin* Pin : InAnimGraphNode->Pins)
diff --git a/Engine/Source/Editor/AnimationBlueprintEditor/Private/AnimationNodes/SGraphNodeBlendSpaceGraph.cpp b/Engine/Source/Editor/AnimationBlueprintEditor/Private/AnimationNodes/SGraphNodeBlendSpaceGraph.cpp
--- a/Engine/Source/Editor/AnimationBlueprintEditor/Private/AnimationNodes/SGraphNodeBlendSpaceGraph.cpp
+++ b/Engine/Source/Editor/AnimationBlueprintEditor/Private/AnimationNodes/SGraphNodeBlendSpaceGraph.cpp
@@ -114,6 +114,11 @@ TSharedRef<SWidget> SGraphNodeBlendSpaceGraph::CreateNodeBody()
 	
 	UAnimGraphNode_BlendSpaceGraphBase* BlendSpaceNode = CastChecked<UAnimGraphNode_BlendSpaceGraphBase>(GraphNode);
 
+	auto UseLowDetailNode = [this]()
+	{
+		return GetCurrentLOD() <= EGraphRenderingLOD::LowDetail;
+	};
+	
 	return SNew(SVerticalBox)
 		+SVerticalBox::Slot()
 		.AutoHeight()
@@ -123,13 +128,13 @@ TSharedRef<SWidget> SGraphNodeBlendSpaceGraph::CreateNodeBody()
 		+SVerticalBox::Slot()
 		.AutoHeight()
 		[
-			SAnimationGraphNode::CreateNodeFunctionsWidget(BlendSpaceNode, MakeAttributeSP(this, &SGraphNodeBlendSpaceGraph::UseLowDetailNodeTitles))
+			SAnimationGraphNode::CreateNodeFunctionsWidget(BlendSpaceNode, MakeAttributeLambda(UseLowDetailNode))
 		]
 		+SVerticalBox::Slot()
 		.AutoHeight()
 		[
 			SNew(SLevelOfDetailBranchNode)
-			.UseLowDetailSlot(this, &SGraphNodeBlendSpaceGraph::UseLowDetailNodeTitles)
+			.UseLowDetailSlot_Lambda(UseLowDetailNode)
 			.LowDetail()
 			[
 				SNew(SSpacer)
@@ -149,6 +154,13 @@ TSharedRef<SWidget> SGraphNodeBlendSpaceGraph::CreateNodeBody()
 					return NAME_None;
 				}))
 			]
+		]
+		+SVerticalBox::Slot()
+		.AutoHeight()
+		.HAlign(HAlign_Right)
+		.Padding(8.0f, 4.0f, 8.0f, 8.0f)
+		[
+			SAnimationGraphNode::CreateNodeTagWidget(BlendSpaceNode, MakeAttributeLambda(UseLowDetailNode))
 		];
 }
 
diff --git a/Engine/Source/Editor/AnimationBlueprintEditor/Private/AnimationNodes/SGraphNodeBlendSpacePlayer.cpp b/Engine/Source/Editor/AnimationBlueprintEditor/Private/AnimationNodes/SGraphNodeBlendSpacePlayer.cpp
--- a/Engine/Source/Editor/AnimationBlueprintEditor/Private/AnimationNodes/SGraphNodeBlendSpacePlayer.cpp
+++ b/Engine/Source/Editor/AnimationBlueprintEditor/Private/AnimationNodes/SGraphNodeBlendSpacePlayer.cpp
@@ -31,8 +31,8 @@ void SGraphNodeBlendSpacePlayer::CreateBelowPinControls(TSharedPtr<SVerticalBox>
 {
 	SAnimationGraphNode::CreateBelowPinControls(MainBox);
 
-	// Insert above the error reporting bar
-	MainBox->InsertSlot(FMath::Max(0, MainBox->NumSlots() - 1))
+	// Insert above the error reporting bar (but above the tag/functions)
+	MainBox->InsertSlot(FMath::Max(0, MainBox->NumSlots() - DebugGridSlotReverseIndex))
 	.AutoHeight()
 	.VAlign(VAlign_Fill)
 	.Padding(0.0f)
diff --git a/Engine/Source/Editor/AnimationBlueprintEditor/Private/AnimationNodes/SGraphNodeBlendSpacePlayer.h b/Engine/Source/Editor/AnimationBlueprintEditor/Private/AnimationNodes/SGraphNodeBlendSpacePlayer.h
--- a/Engine/Source/Editor/AnimationBlueprintEditor/Private/AnimationNodes/SGraphNodeBlendSpacePlayer.h
+++ b/Engine/Source/Editor/AnimationBlueprintEditor/Private/AnimationNodes/SGraphNodeBlendSpacePlayer.h
@@ -18,6 +18,9 @@ public:
 	SLATE_BEGIN_ARGS(SGraphNodeBlendSpacePlayer){}
 	SLATE_END_ARGS()
 
+	// Reverse index of the debug grid widget
+	static const int32 DebugGridSlotReverseIndex = 2;
+	
 	void Construct(const FArguments& InArgs, UAnimGraphNode_Base* InNode);
 
 	// SGraphNode interface
diff --git a/Engine/Source/Editor/AnimationBlueprintEditor/Private/AnimationNodes/SGraphNodeSequencePlayer.cpp b/Engine/Source/Editor/AnimationBlueprintEditor/Private/AnimationNodes/SGraphNodeSequencePlayer.cpp
--- a/Engine/Source/Editor/AnimationBlueprintEditor/Private/AnimationNodes/SGraphNodeSequencePlayer.cpp
+++ b/Engine/Source/Editor/AnimationBlueprintEditor/Private/AnimationNodes/SGraphNodeSequencePlayer.cpp
@@ -77,15 +77,20 @@ void SGraphNodeSequencePlayer::UpdateGraphNode()
 void SGraphNodeSequencePlayer::CreateBelowPinControls(TSharedPtr<SVerticalBox> MainBox)
 {
 	SAnimationGraphNode::CreateBelowPinControls(MainBox);
+
+	auto UseLowDetailNode = [this]()
+	{
+		return GetCurrentLOD() <= EGraphRenderingLOD::LowDetail;
+	};
 	
-	// Insert above the error reporting bar
-	MainBox->InsertSlot(FMath::Max(0, MainBox->NumSlots() - 1))
+	// Insert above the error reporting bar (but above the tag/functions)
+	MainBox->InsertSlot(FMath::Max(0, MainBox->NumSlots() - DebugSliderSlotReverseIndex))
 	.AutoHeight()
 	.VAlign( VAlign_Fill )
 	.Padding(FMargin(4))
 	[
 		SNew(SLevelOfDetailBranchNode)
-		.UseLowDetailSlot(this, &SGraphNodeSequencePlayer::UseLowDetailNodeTitles)
+		.UseLowDetailSlot_Lambda(UseLowDetailNode)
 		.LowDetail()
 		[
 			SNew(SSpacer)
diff --git a/Engine/Source/Editor/AnimationBlueprintEditor/Private/AnimationNodes/SGraphNodeSequencePlayer.h b/Engine/Source/Editor/AnimationBlueprintEditor/Private/AnimationNodes/SGraphNodeSequencePlayer.h
--- a/Engine/Source/Editor/AnimationBlueprintEditor/Private/AnimationNodes/SGraphNodeSequencePlayer.h
+++ b/Engine/Source/Editor/AnimationBlueprintEditor/Private/AnimationNodes/SGraphNodeSequencePlayer.h
@@ -17,6 +17,9 @@ public:
 	SLATE_BEGIN_ARGS(SGraphNodeSequencePlayer){}
 	SLATE_END_ARGS()
 
+	// Reverse index of the debug slider widget
+	static const int32 DebugSliderSlotReverseIndex = 2;
+
 	void Construct(const FArguments& InArgs, UAnimGraphNode_Base* InNode);
 
 	// SNodePanel::SNode interface
diff --git a/Engine/Source/Editor/AnimationBlueprintEditor/Public/AnimationNodes/SAnimationGraphNode.h b/Engine/Source/Editor/AnimationBlueprintEditor/Public/AnimationNodes/SAnimationGraphNode.h
--- a/Engine/Source/Editor/AnimationBlueprintEditor/Public/AnimationNodes/SAnimationGraphNode.h
+++ b/Engine/Source/Editor/AnimationBlueprintEditor/Public/AnimationNodes/SAnimationGraphNode.h
@@ -20,6 +20,12 @@ public:
 	SLATE_BEGIN_ARGS(SAnimationGraphNode) {}
 	SLATE_END_ARGS()
 
+	// Reverse index of the error reporting bar slot
+	static const int32 ErrorReportingSlotReverseIndex = 0;
+
+	// Reverse index of the tag/functions slot
+	static const int32 TagAndFunctionsSlotReverseIndex = 1;
+	
 	void Construct(const FArguments& InArgs, UAnimGraphNode_Base* InNode);
 
 	// Tweak any created pin widgets so they respond to bindings
@@ -27,6 +33,9 @@ public:
 
 	// Create below-widget controls for editing anim node functions
 	static TSharedRef<SWidget> CreateNodeFunctionsWidget(UAnimGraphNode_Base* InAnimNode, TAttribute<bool> InUseLowDetail);
+
+	// Create below-widget controls for editing anim node tags
+	static TSharedRef<SWidget> CreateNodeTagWidget(UAnimGraphNode_Base* InAnimNode, TAttribute<bool> InUseLowDetail);
 	
 protected:
 	// SWidget interface
diff --git a/Engine/Source/Editor/EditorStyle/Private/StarshipStyle.cpp b/Engine/Source/Editor/EditorStyle/Private/StarshipStyle.cpp
--- a/Engine/Source/Editor/EditorStyle/Private/StarshipStyle.cpp
+++ b/Engine/Source/Editor/EditorStyle/Private/StarshipStyle.cpp
@@ -5855,6 +5855,18 @@ void FStarshipEditorStyle::FStyle::SetupPersonaStyle()
 		Set("AnimGraph.Attribute.RootMotionDelta.Color", FStyleColors::AccentGreen);
 		Set("AnimGraph.Attribute.InertialBlending.Icon", new IMAGE_BRUSH_SVG("Starship/AnimationBlueprintEditor/AnimationGraphInertialBlending", Icon16x16));
 		Set("AnimGraph.Attribute.InertialBlending.Color", FStyleColors::AccentOrange);
+
+		Set("AnimGraph.AnimNodeReference.Subtitle", FTextBlockStyle(NormalText)
+			.SetFont(DEFAULT_FONT("Italic", 8))
+			.SetColorAndOpacity(FLinearColor(218.0f/255.0f,218.0f/255.0f,96.0f/255.0f, 0.5f))
+		);
+
+		FTextBlockStyle TagTextStyle = FTextBlockStyle(NormalText)
+			.SetFont(DEFAULT_FONT("Italic", 9))
+			.SetColorAndOpacity(FLinearColor(218.0f/255.0f,218.0f/255.0f,96.0f/255.0f, 0.5f));
+		
+		Set("AnimGraph.Node.Tag", FInlineEditableTextBlockStyle(FCoreStyle::Get().GetWidgetStyle<FInlineEditableTextBlockStyle>("InlineEditableTextBlockStyle"))
+			.SetTextStyle(TagTextStyle));
 	}
 	
 	// Property Access 
diff --git a/Engine/Source/Editor/GraphEditor/Private/KismetNodes/SGraphNodeK2Var.cpp b/Engine/Source/Editor/GraphEditor/Private/KismetNodes/SGraphNodeK2Var.cpp
--- a/Engine/Source/Editor/GraphEditor/Private/KismetNodes/SGraphNodeK2Var.cpp
+++ b/Engine/Source/Editor/GraphEditor/Private/KismetNodes/SGraphNodeK2Var.cpp
@@ -33,6 +33,25 @@ FSlateColor SGraphNodeK2Var::GetVariableColor() const
 	return GraphNode->GetNodeTitleColor();
 }
 
+TSharedRef<SWidget> SGraphNodeK2Var::UpdateTitleWidget(FText InTitleText, TSharedPtr<SWidget> InTitleWidget, EHorizontalAlignment& InOutTitleHAlign, FMargin& InOutTitleMargin) const
+{
+	TSharedPtr<SWidget> TitleWidget = InTitleWidget;
+	
+	if (InTitleText.IsEmpty())
+	{
+		TitleWidget = SNullWidget::NullWidget;
+	}
+
+	if(!TitleWidget.IsValid())
+	{
+		TitleWidget = SNew(STextBlock)
+		.TextStyle( FEditorStyle::Get(), "Graph.Node.NodeTitle" )
+		.Text(InTitleText);
+	}
+
+	return TitleWidget.ToSharedRef();
+}
+
 void SGraphNodeK2Var::UpdateGraphNode()
 {
 	InputPins.Empty();
@@ -163,18 +182,8 @@ void SGraphNodeK2Var::UpdateGraphNode()
 			ContentAreaMargin.Top += 16;
 		}
 	}
-	
-	if (TitleText.IsEmpty())
-	{
-		TitleWidget = SNullWidget::NullWidget;
-	}
 
-	if(!TitleWidget.IsValid())
-	{
-		TitleWidget = SNew(STextBlock)
-						.TextStyle( FEditorStyle::Get(), "Graph.Node.NodeTitle" )
-						.Text(TitleText);
-	}
+	TitleWidget = UpdateTitleWidget(TitleText, TitleWidget, TitleHAlign, TitleMargin);
 
 	SetupErrorReporting();
 
diff --git a/Engine/Source/Editor/GraphEditor/Public/KismetNodes/SGraphNodeK2Var.h b/Engine/Source/Editor/GraphEditor/Public/KismetNodes/SGraphNodeK2Var.h
--- a/Engine/Source/Editor/GraphEditor/Public/KismetNodes/SGraphNodeK2Var.h
+++ b/Engine/Source/Editor/GraphEditor/Public/KismetNodes/SGraphNodeK2Var.h
@@ -25,4 +25,7 @@ public:
 
 protected:
 	FSlateColor GetVariableColor() const;
+
+	// Allow derived classes to override title widget
+	virtual TSharedRef<SWidget> UpdateTitleWidget(FText InTitleText, TSharedPtr<SWidget> InTitleWidget, EHorizontalAlignment& InOutTitleHAlign, FMargin& InOutTitleMargin) const;
 };
diff --git a/Engine/Source/Runtime/Core/Public/UObject/UE5MainStreamObjectVersion.h b/Engine/Source/Runtime/Core/Public/UObject/UE5MainStreamObjectVersion.h
--- a/Engine/Source/Runtime/Core/Public/UObject/UE5MainStreamObjectVersion.h
+++ b/Engine/Source/Runtime/Core/Public/UObject/UE5MainStreamObjectVersion.h
@@ -175,6 +175,9 @@ struct CORE_API FUE5MainStreamObjectVersion
 		// Add bVisibleInRayTracing flag to Skeletal Mesh Sections
 		SkelMeshSectionVisibleInRayTracingFlagAdded,
 
+		// Add generic tagging of all anim graph nodes in anim blueprints
+		AnimGraphNodeTaggingAdded,
+		
 		// -----<new versions can be added above this line>-------------------------------------------------
 		VersionPlusOne,
 		LatestVersion = VersionPlusOne - 1
diff --git a/Engine/Source/Runtime/Engine/Classes/Animation/AnimInstance.h b/Engine/Source/Runtime/Engine/Classes/Animation/AnimInstance.h
--- a/Engine/Source/Runtime/Engine/Classes/Animation/AnimInstance.h
+++ b/Engine/Source/Runtime/Engine/Classes/Animation/AnimInstance.h
@@ -815,7 +815,8 @@ public:
 	void GetSubInstancesByTag(FName InTag, TArray<UAnimInstance*>& OutSubInstances) const;
 
 	/** Runs through all nodes, attempting to find all linked instances that match the name/tag */
-	UFUNCTION(BlueprintPure, Category = "Animation Blueprint Linking")
+	UE_DEPRECATED(5.0, "Tags are unique so this function is no longer supported. Please use GetLinkedAnimGraphInstanceByTag instead")
+	UFUNCTION(BlueprintPure, Category = "Animation Blueprint Linking", meta=(DeprecatedFunction, DeprecationMessage="Tags are unique so this function is no longer supported. Please use GetLinkedAnimGraphInstanceByTag instead"))
 	void GetLinkedAnimGraphInstancesByTag(FName InTag, TArray<UAnimInstance*>& OutLinkedInstances) const;
 
 	UE_DEPRECATED(4.24, "Function renamed, please use LinkAnimGraphByTag")
diff --git a/Engine/Source/Runtime/Engine/Classes/Animation/AnimNode_LinkedAnimGraph.h b/Engine/Source/Runtime/Engine/Classes/Animation/AnimNode_LinkedAnimGraph.h
--- a/Engine/Source/Runtime/Engine/Classes/Animation/AnimNode_LinkedAnimGraph.h
+++ b/Engine/Source/Runtime/Engine/Classes/Animation/AnimNode_LinkedAnimGraph.h
@@ -41,10 +41,11 @@ public:
 	UPROPERTY(EditAnywhere, Category = Settings)
 	TSubclassOf<UAnimInstance> InstanceClass;
 
-	/** Optional tag used to identify this linked instance */
-	UPROPERTY(EditAnywhere, Category = Settings)
-	FName Tag;
-
+#if WITH_EDITORONLY_DATA
+	UPROPERTY()
+	FName Tag_DEPRECATED;
+#endif
+	
 	// The root node of the dynamically-linked graph
 	FAnimNode_Base* LinkedRoot;
 
diff --git a/Engine/Source/Runtime/Engine/Classes/Components/SkeletalMeshComponent.h b/Engine/Source/Runtime/Engine/Classes/Components/SkeletalMeshComponent.h
--- a/Engine/Source/Runtime/Engine/Classes/Components/SkeletalMeshComponent.h
+++ b/Engine/Source/Runtime/Engine/Classes/Components/SkeletalMeshComponent.h
@@ -916,12 +916,18 @@ public:
 		UAnimInstance* GetLinkedAnimGraphInstanceByTag(FName InTag) const;
 
 	UE_DEPRECATED(4.24, "Function renamed, please use GetLinkedAnimGraphInstancesByTag")
-	void GetSubInstancesByTag(FName InTag, TArray<UAnimInstance*>& OutSubInstances) const { GetLinkedAnimGraphInstancesByTag(InTag, OutSubInstances); }
+	void GetSubInstancesByTag(FName InTag, TArray<UAnimInstance*>& OutSubInstances) const
+	{
+		PRAGMA_DISABLE_DEPRECATION_WARNINGS
+		GetLinkedAnimGraphInstancesByTag(InTag, OutSubInstances);
+		PRAGMA_ENABLE_DEPRECATION_WARNINGS
+	}
 
 	/**
 	 * Returns all tagged linked instance nodes that match the tag.
 	 */
-	UFUNCTION(BlueprintPure, Category = "Components|SkeletalMesh|Animation Blueprint Linking", meta = (Keywords = "AnimBlueprint"))
+	UE_DEPRECATED(5.0, "Tags are unique so this funciton is no longer supported. Please use GetLinkedAnimGraphInstanceByTag instead")
+	UFUNCTION(BlueprintPure, Category = "Components|SkeletalMesh|Animation Blueprint Linking", meta = (Keywords = "AnimBlueprint", DeprecatedFunction, DeprecationMessage="Tags are unique so this funciton is no longer supported. Please use GetLinkedAnimGraphInstanceByTag instead"))
 	void GetLinkedAnimGraphInstancesByTag(FName InTag, TArray<UAnimInstance*>& OutLinkedInstances) const;
 
 	UE_DEPRECATED(4.24, "Function renamed, please use LinkAnimGraphByTag")
diff --git a/Engine/Source/Runtime/Engine/Private/Animation/AnimInstance.cpp b/Engine/Source/Runtime/Engine/Private/Animation/AnimInstance.cpp
--- a/Engine/Source/Runtime/Engine/Private/Animation/AnimInstance.cpp
+++ b/Engine/Source/Runtime/Engine/Private/Animation/AnimInstance.cpp
@@ -31,6 +31,7 @@
 #include "Animation/AnimNode_LinkedAnimLayer.h"
 #include "UObject/UE5MainStreamObjectVersion.h"
 #include "Animation/AnimSubsystem.h"
+#include "Animation/AnimSubsystem_Tag.h"
 
 /** Anim stats */
 
@@ -2707,14 +2708,11 @@ UAnimInstance* UAnimInstance::GetLinkedAnimGraphInstanceByTag(FName InTag) const
 {
 	if (IAnimClassInterface* AnimBlueprintClass = IAnimClassInterface::GetFromClass(GetClass()))
 	{
-		const TArray<FStructProperty*>& LinkedAnimGraphNodeProperties = AnimBlueprintClass->GetLinkedAnimGraphNodeProperties();
-		for(const FStructProperty* LinkedAnimGraphNodeProperty : LinkedAnimGraphNodeProperties)
+		const FAnimSubsystem_Tag& TagSubsystem = AnimBlueprintClass->GetSubsystem<FAnimSubsystem_Tag>();
+		const FAnimNode_LinkedAnimGraph* LinkedAnimGraph = TagSubsystem.FindNodeByTag<FAnimNode_LinkedAnimGraph>(InTag, this);
+		if(LinkedAnimGraph)
 		{
-			const FAnimNode_LinkedAnimGraph* LinkedAnimGraph = LinkedAnimGraphNodeProperty->ContainerPtrToValuePtr<FAnimNode_LinkedAnimGraph>(this);
-			if(LinkedAnimGraph && LinkedAnimGraph->Tag == InTag)
-			{
-				return LinkedAnimGraph->GetTargetInstance<UAnimInstance>();
-			}
+			return LinkedAnimGraph->GetTargetInstance<UAnimInstance>();
 		}
 	}
 
@@ -2723,17 +2721,9 @@ UAnimInstance* UAnimInstance::GetLinkedAnimGraphInstanceByTag(FName InTag) const
 
 void UAnimInstance::GetLinkedAnimGraphInstancesByTag(FName InTag, TArray<UAnimInstance*>& OutLinkedInstances) const
 {
-	if (IAnimClassInterface* AnimBlueprintClass = IAnimClassInterface::GetFromClass(GetClass()))
+	if(UAnimInstance* Instance = GetLinkedAnimGraphInstanceByTag(InTag))
 	{
-		const TArray<FStructProperty*>& LinkedAnimGraphNodeProperties = AnimBlueprintClass->GetLinkedAnimGraphNodeProperties();
-		for (const FStructProperty* LinkedAnimGraphNodeProperty : LinkedAnimGraphNodeProperties)
-		{
-			const FAnimNode_LinkedAnimGraph* LinkedAnimGraph = LinkedAnimGraphNodeProperty->ContainerPtrToValuePtr<FAnimNode_LinkedAnimGraph>(this);
-			if(LinkedAnimGraph && LinkedAnimGraph->Tag == InTag)
-			{
-				OutLinkedInstances.Add(LinkedAnimGraph->GetTargetInstance<UAnimInstance>());
-			}
-		}
+		OutLinkedInstances.Add(Instance);
 	}
 }
 
@@ -2741,14 +2731,11 @@ void UAnimInstance::LinkAnimGraphByTag(FName InTag, TSubclassOf<UAnimInstance> I
 {
 	if (IAnimClassInterface* AnimBlueprintClass = IAnimClassInterface::GetFromClass(GetClass()))
 	{
-		const TArray<FStructProperty*>& LinkedAnimGraphNodeProperties = AnimBlueprintClass->GetLinkedAnimGraphNodeProperties();
-		for (const FStructProperty* LinkedAnimGraphNodeProperty : LinkedAnimGraphNodeProperties)
+		const FAnimSubsystem_Tag& TagSubsystem = AnimBlueprintClass->GetSubsystem<FAnimSubsystem_Tag>();
+		FAnimNode_LinkedAnimGraph* LinkedAnimGraph = TagSubsystem.FindNodeByTag<FAnimNode_LinkedAnimGraph>(InTag, this);
+		if(LinkedAnimGraph)
 		{
-			FAnimNode_LinkedAnimGraph* LinkedAnimGraph = LinkedAnimGraphNodeProperty->ContainerPtrToValuePtr<FAnimNode_LinkedAnimGraph>(this);
-			if(LinkedAnimGraph && LinkedAnimGraph->Tag == InTag)
-			{
-				LinkedAnimGraph->SetAnimClass(InClass, this);
-			}
+			LinkedAnimGraph->SetAnimClass(InClass, this);
 		}
 	}
 }
diff --git a/Engine/Source/Runtime/Engine/Private/Animation/AnimNode_LinkedAnimGraph.cpp b/Engine/Source/Runtime/Engine/Private/Animation/AnimNode_LinkedAnimGraph.cpp
--- a/Engine/Source/Runtime/Engine/Private/Animation/AnimNode_LinkedAnimGraph.cpp
+++ b/Engine/Source/Runtime/Engine/Private/Animation/AnimNode_LinkedAnimGraph.cpp
@@ -31,7 +31,9 @@ static float GetBlendDuration(const IAnimClassInterface* PriorAnimBPClass, const
 
 FAnimNode_LinkedAnimGraph::FAnimNode_LinkedAnimGraph()
 	: InstanceClass(nullptr)
-	, Tag(NAME_None)
+#if WITH_EDITORONLY_DATA
+	, Tag_DEPRECATED(NAME_None)
+#endif
 	, LinkedRoot(nullptr)
 	, NodeIndex(INDEX_NONE)
 	, CachedLinkedNodeIndex(INDEX_NONE)
diff --git a/Engine/Source/Runtime/Engine/Private/Animation/AnimSubsystem_Tag.cpp b/Engine/Source/Runtime/Engine/Private/Animation/AnimSubsystem_Tag.cpp
new file mode 100644
--- /dev/null
+++ b/Engine/Source/Runtime/Engine/Private/Animation/AnimSubsystem_Tag.cpp
@@ -0,0 +1,44 @@
+﻿// Copyright Epic Games, Inc. All Rights Reserved.
+
+#include "Animation/AnimSubsystem_Tag.h"
+#include "Animation/AnimClassInterface.h"
+#include "Animation/AnimInstance.h"
+
+int32 FAnimSubsystem_Tag::FindNodeIndexByTag(FName InTag) const
+{
+	if(const int32* NodeIndexPtr = NodeIndices.Find(InTag))
+	{
+		const TArray<FStructProperty*>& AnimNodeProperties = AnimClass->GetAnimNodeProperties();
+		// As the index is patched during compilation, it needs to be reversed here
+		int32 ReverseIndex = AnimNodeProperties.Num() - 1 - *NodeIndexPtr;
+		return ReverseIndex;
+	}
+
+	return INDEX_NONE;
+}
+
+FAnimNode_Base* FAnimSubsystem_Tag::FindNodeByTag_Internal(FName InTag, UAnimInstance* InInstance, UScriptStruct* InNodeType) const
+{
+	ensure(InInstance->GetClass()->IsChildOf(IAnimClassInterface::GetActualAnimClass(AnimClass)));
+
+	if(const int32* NodeIndexPtr = NodeIndices.Find(InTag))
+	{
+		const TArray<FStructProperty*>& AnimNodeProperties = AnimClass->GetAnimNodeProperties();
+		// As the index is patched during compilation, it needs to be reversed here
+		int32 ReverseIndex = AnimNodeProperties.Num() - 1 - *NodeIndexPtr;
+		
+		check(AnimNodeProperties.IsValidIndex(ReverseIndex));
+		FStructProperty* NodeProperty = AnimClass->GetAnimNodeProperties()[ReverseIndex];
+		if(NodeProperty->Struct->IsChildOf(InNodeType))
+		{
+			return NodeProperty->ContainerPtrToValuePtr<FAnimNode_Base>(InInstance);
+		}
+	}
+
+	return nullptr;
+}
+
+void FAnimSubsystem_Tag::OnPostLoadDefaults(FAnimSubsystemPostLoadDefaultsContext& InContext)
+{
+	AnimClass = IAnimClassInterface::GetFromClass(InContext.DefaultAnimInstance->GetClass());
+}
\ No newline at end of file
diff --git a/Engine/Source/Runtime/Engine/Private/Components/SkeletalMeshComponent.cpp b/Engine/Source/Runtime/Engine/Private/Components/SkeletalMeshComponent.cpp
--- a/Engine/Source/Runtime/Engine/Private/Components/SkeletalMeshComponent.cpp
+++ b/Engine/Source/Runtime/Engine/Private/Components/SkeletalMeshComponent.cpp
@@ -3024,7 +3024,9 @@ void USkeletalMeshComponent::GetLinkedAnimGraphInstancesByTag(FName InTag, TArra
 {
 	if(AnimScriptInstance)
 	{
+		PRAGMA_DISABLE_DEPRECATION_WARNINGS
 		AnimScriptInstance->GetLinkedAnimGraphInstancesByTag(InTag, OutLinkedInstances);
+		PRAGMA_ENABLE_DEPRECATION_WARNINGS
 	}
 }
 
diff --git a/Engine/Source/Runtime/Engine/Public/Animation/AnimSubsystem_Tag.h b/Engine/Source/Runtime/Engine/Public/Animation/AnimSubsystem_Tag.h
new file mode 100644
--- /dev/null
+++ b/Engine/Source/Runtime/Engine/Public/Animation/AnimSubsystem_Tag.h
@@ -0,0 +1,65 @@
+// Copyright Epic Games, Inc. All Rights Reserved.
+
+#pragma once
+
+#include "AnimSubsystem.h"
+#include "AnimSubsystem_Tag.generated.h"
+
+struct FAnimNode_Base;
+class UAnimInstance;
+
+/** Allows anim nodes to be tagged by name at edit time and retrieved by tag at runtime */
+USTRUCT()
+struct ENGINE_API FAnimSubsystem_Tag : public FAnimSubsystem
+{
+	GENERATED_BODY()
+
+public:
+	/**
+	 * Find a node's index by tag.
+	 * @param	InTag		The name tag for the node to find
+	 * @return INDEX_NONE if the tagged node was not found
+	 */	
+	int32 FindNodeIndexByTag(FName InTag) const;
+	
+	/**
+	 * Find a node by tag.
+	 * @param	InTag		The name tag for the node to find
+	 * @param	InInstance	The instance that the node is in
+	 * @return nullptr if the tagged node was not found
+	 */
+	template<typename NodeType>
+	NodeType* FindNodeByTag(FName InTag, UAnimInstance* InInstance) const
+	{
+		return static_cast<NodeType*>(FindNodeByTag_Internal(InTag, InInstance, NodeType::StaticStruct()));
+	}
+
+	/**
+	 * Find a node by tag (const version)
+	 * @param	InTag		The name tag for the node to find
+	 * @param	InInstance	The instance that the node is in
+	 * @return nullptr if the tagged node was not found
+	 */	
+	template<typename NodeType>
+	const NodeType* FindNodeByTag(FName InTag, const UAnimInstance* InInstance) const
+	{
+		return static_cast<const NodeType*>(FindNodeByTag_Internal(InTag, const_cast<UAnimInstance*>(InInstance), NodeType::StaticStruct()));
+	}
+	
+private:
+	// FAnimSubsystem interface
+	virtual void OnPostLoadDefaults(FAnimSubsystemPostLoadDefaultsContext& InContext) override;
+
+	// Internal impl for FindNodeByTag
+	FAnimNode_Base* FindNodeByTag_Internal(FName InTag, UAnimInstance* InInstance, UScriptStruct* InNodeType) const;
+	
+private:
+	friend class UAnimBlueprintExtension_Tag;
+	
+	// Map of name->node index
+	UPROPERTY()
+	TMap<FName, int32> NodeIndices;
+
+	// Anim class that contains this subsystem
+	IAnimClassInterface* AnimClass = nullptr;
+};
\ No newline at end of file
