Fixed anim node functions not being called correctly

Due to a previous change, functions were being called on the skeleton class
Also removes an incorrect guard that was preventing binding to native function libraries

diff --git a/Engine/Plugins/Runtime/PropertyAccess/Source/PropertyAccessEditor/Private/SPropertyBinding.cpp b/Engine/Plugins/Runtime/PropertyAccess/Source/PropertyAccessEditor/Private/SPropertyBinding.cpp
--- a/Engine/Plugins/Runtime/PropertyAccess/Source/PropertyAccessEditor/Private/SPropertyBinding.cpp
+++ b/Engine/Plugins/Runtime/PropertyAccess/Source/PropertyAccessEditor/Private/SPropertyBinding.cpp
@@ -177,7 +177,7 @@ void SPropertyBinding::ForEachBindableFunction(UClass* FromClass, Predicate Pred
 			for (TObjectIterator<UClass> ClassIt; ClassIt; ++ClassIt)
 			{
 				UClass* TestClass = *ClassIt;
-				if (!TestClass->IsNative() && TestClass->IsChildOf(UBlueprintFunctionLibrary::StaticClass()) && (!TestClass->HasAnyClassFlags(CLASS_Abstract)))
+				if (TestClass->IsChildOf(UBlueprintFunctionLibrary::StaticClass()) && (!TestClass->HasAnyClassFlags(CLASS_Abstract)))
 				{
 					CheckBindableClass(TestClass);
 				}
diff --git a/Engine/Source/Editor/AnimGraph/Private/AnimGraphNode_Base.cpp b/Engine/Source/Editor/AnimGraph/Private/AnimGraphNode_Base.cpp
--- a/Engine/Source/Editor/AnimGraph/Private/AnimGraphNode_Base.cpp
+++ b/Engine/Source/Editor/AnimGraph/Private/AnimGraphNode_Base.cpp
@@ -88,19 +88,19 @@ void UAnimGraphNode_Base::PostEditChangeProperty(FPropertyChangedEvent& Property
 	}
 	else if(PropertyName == GET_MEMBER_NAME_CHECKED(UAnimGraphNode_Base, InitialUpdateFunction))
 	{
-		GetFNode()->InitialUpdateFunction.SetFromFunction(InitialUpdateFunction.ResolveMember<UFunction>(GetAnimBlueprint()->SkeletonGeneratedClass));
+		GetFNode()->InitialUpdateFunction.SetFromFunction(InitialUpdateFunction.ResolveMember<UFunction>(GetBlueprintClassFromNode()));
 		GetAnimBlueprint()->RequestRefreshExtensions();
 		GetSchema()->ReconstructNode(*this);
 	}
 	else if(PropertyName == GET_MEMBER_NAME_CHECKED(UAnimGraphNode_Base, BecomeRelevantFunction))
 	{
-		GetFNode()->BecomeRelevantFunction.SetFromFunction(BecomeRelevantFunction.ResolveMember<UFunction>(GetAnimBlueprint()->SkeletonGeneratedClass));
+		GetFNode()->BecomeRelevantFunction.SetFromFunction(BecomeRelevantFunction.ResolveMember<UFunction>(GetBlueprintClassFromNode()));
 		GetAnimBlueprint()->RequestRefreshExtensions();
 		GetSchema()->ReconstructNode(*this);
 	}
 	else if(PropertyName == GET_MEMBER_NAME_CHECKED(UAnimGraphNode_Base, UpdateFunction))
 	{
-		GetFNode()->UpdateFunction.SetFromFunction(UpdateFunction.ResolveMember<UFunction>(GetAnimBlueprint()->SkeletonGeneratedClass));
+		GetFNode()->UpdateFunction.SetFromFunction(UpdateFunction.ResolveMember<UFunction>(GetBlueprintClassFromNode()));
 		GetSchema()->ReconstructNode(*this);
 	}
 	else if(PropertyName == GET_MEMBER_NAME_CHECKED(UAnimGraphNode_Base, Tag))
diff --git a/Engine/Source/Runtime/Engine/Private/Animation/AnimNodeFunctionRef.cpp b/Engine/Source/Runtime/Engine/Private/Animation/AnimNodeFunctionRef.cpp
--- a/Engine/Source/Runtime/Engine/Private/Animation/AnimNodeFunctionRef.cpp
+++ b/Engine/Source/Runtime/Engine/Private/Animation/AnimNodeFunctionRef.cpp
@@ -13,7 +13,9 @@ void FAnimNodeFunctionRef::SetFromFunction(UFunction* InFunction)
 	if(InFunction)
 	{
 		FunctionName = InFunction->GetFName();
-		if(UClass* OwnerClass = InFunction->GetOwnerClass())
+		UClass* OwnerClass = InFunction->GetOwnerClass();
+		// Only need the class name if the function is static, as the function will be called on the passed-in UObject otherwise.
+		if(OwnerClass && InFunction->HasAnyFunctionFlags(FUNC_Static))
 		{
 			ensureAlwaysMsgf(OwnerClass->IsChildOf(UAnimInstance::StaticClass()) || OwnerClass->IsChildOf(UBlueprintFunctionLibrary::StaticClass()), TEXT("Function class must derive from either UAnimInstance or UBlueprintFunctionLibrary"));
 			ClassName = *OwnerClass->GetPathName();
