[UE][FEATURE] Support render of multiple selected nodes in Anim Graph Editor Viewport

diff --git a/Engine/Source/Editor/AnimGraph/Private/EditModes/AnimNodeEditMode.cpp b/Engine/Source/Editor/AnimGraph/Private/EditModes/AnimNodeEditMode.cpp
--- a/Engine/Source/Editor/AnimGraph/Private/EditModes/AnimNodeEditMode.cpp
+++ b/Engine/Source/Editor/AnimGraph/Private/EditModes/AnimNodeEditMode.cpp
@@ -12,9 +12,7 @@
 #define LOCTEXT_NAMESPACE "AnimNodeEditMode"
 
 FAnimNodeEditMode::FAnimNodeEditMode()
-	: AnimNode(nullptr)
-	, RuntimeAnimNode(nullptr)
-	, bManipulating(false)
+	: bManipulating(false)
 	, bInTransaction(false)
 {
 	// Disable grid drawing for this mode as the viewport handles this
@@ -37,15 +35,18 @@ IPersonaPreviewScene& FAnimNodeEditMode::GetAnimPreviewScene() const
 
 void FAnimNodeEditMode::GetOnScreenDebugInfo(TArray<FText>& OutDebugInfo) const
 {
-	if (AnimNode != nullptr)
+	for (EditorRuntimeNodePair CurrentNodePair : AnimNodes)
 	{
-		AnimNode->GetOnScreenDebugInfo(OutDebugInfo, RuntimeAnimNode, GetAnimPreviewScene().GetPreviewMeshComponent());
+		if ((CurrentNodePair.EditorAnimNode != nullptr) && (CurrentNodePair.RuntimeAnimNode != nullptr))
+		{
+			CurrentNodePair.EditorAnimNode->GetOnScreenDebugInfo(OutDebugInfo, CurrentNodePair.RuntimeAnimNode, GetAnimPreviewScene().GetPreviewMeshComponent());
+		}
 	}
 }
 
 ECoordSystem FAnimNodeEditMode::GetWidgetCoordinateSystem() const
 {
-	UAnimGraphNode_SkeletalControlBase* SkelControl = Cast<UAnimGraphNode_SkeletalControlBase>(AnimNode);
+	UAnimGraphNode_SkeletalControlBase* SkelControl = Cast<UAnimGraphNode_SkeletalControlBase>(GetActiveWidgetAnimNode());
 	if (SkelControl != nullptr)
 	{
 		PRAGMA_DISABLE_DEPRECATION_WARNINGS
@@ -58,7 +59,7 @@ ECoordSystem FAnimNodeEditMode::GetWidgetCoordinateSystem() const
 
 UE::Widget::EWidgetMode FAnimNodeEditMode::GetWidgetMode() const
 {
-	UAnimGraphNode_SkeletalControlBase* SkelControl = Cast<UAnimGraphNode_SkeletalControlBase>(AnimNode);
+	UAnimGraphNode_SkeletalControlBase* SkelControl = Cast<UAnimGraphNode_SkeletalControlBase>(GetActiveWidgetAnimNode());
 	if (SkelControl != nullptr)
 	{
 		PRAGMA_DISABLE_DEPRECATION_WARNINGS
@@ -71,7 +72,7 @@ UE::Widget::EWidgetMode FAnimNodeEditMode::GetWidgetMode() const
 
 UE::Widget::EWidgetMode FAnimNodeEditMode::ChangeToNextWidgetMode(UE::Widget::EWidgetMode CurWidgetMode)
 {
-	UAnimGraphNode_SkeletalControlBase* SkelControl = Cast<UAnimGraphNode_SkeletalControlBase>(AnimNode);
+	UAnimGraphNode_SkeletalControlBase* SkelControl = Cast<UAnimGraphNode_SkeletalControlBase>(GetActiveWidgetAnimNode());
 	if (SkelControl != nullptr)
 	{
 		PRAGMA_DISABLE_DEPRECATION_WARNINGS
@@ -84,7 +85,7 @@ UE::Widget::EWidgetMode FAnimNodeEditMode::ChangeToNextWidgetMode(UE::Widget::EW
 
 bool FAnimNodeEditMode::SetWidgetMode(UE::Widget::EWidgetMode InWidgetMode)
 {
-	UAnimGraphNode_SkeletalControlBase* SkelControl = Cast<UAnimGraphNode_SkeletalControlBase>(AnimNode);
+	UAnimGraphNode_SkeletalControlBase* SkelControl = Cast<UAnimGraphNode_SkeletalControlBase>(GetActiveWidgetAnimNode());
 	if (SkelControl != nullptr)
 	{
 		PRAGMA_DISABLE_DEPRECATION_WARNINGS
@@ -97,7 +98,7 @@ bool FAnimNodeEditMode::SetWidgetMode(UE::Widget::EWidgetMode InWidgetMode)
 
 FName FAnimNodeEditMode::GetSelectedBone() const
 {
-	UAnimGraphNode_SkeletalControlBase* SkelControl = Cast<UAnimGraphNode_SkeletalControlBase>(AnimNode);
+	UAnimGraphNode_SkeletalControlBase* const SkelControl = Cast<UAnimGraphNode_SkeletalControlBase>(GetActiveWidgetAnimNode());
 	if (SkelControl != nullptr)
 	{
 		PRAGMA_DISABLE_DEPRECATION_WARNINGS
@@ -110,16 +111,20 @@ FName FAnimNodeEditMode::GetSelectedBone() const
 
 void FAnimNodeEditMode::EnterMode(UAnimGraphNode_Base* InEditorNode, FAnimNode_Base* InRuntimeNode)
 {
-	AnimNode = InEditorNode;
-	RuntimeAnimNode = InRuntimeNode;
+	check(InEditorNode && InRuntimeNode); // Expect valid Node ptrs.
 
-	UAnimGraphNode_SkeletalControlBase* SkelControl = Cast<UAnimGraphNode_SkeletalControlBase>(AnimNode);
-	if (SkelControl != nullptr)
+	if (InEditorNode && InRuntimeNode)
 	{
-		PRAGMA_DISABLE_DEPRECATION_WARNINGS
-		SkelControl->MoveSelectActorLocation(GetAnimPreviewScene().GetPreviewMeshComponent(), (FAnimNode_SkeletalControlBase*)RuntimeAnimNode);
-		SkelControl->CopyNodeDataTo(RuntimeAnimNode);
-		PRAGMA_ENABLE_DEPRECATION_WARNINGS
+		AnimNodes.Add(EditorRuntimeNodePair(InEditorNode, InRuntimeNode));
+
+		UAnimGraphNode_SkeletalControlBase* SkelControl = Cast<UAnimGraphNode_SkeletalControlBase>(InEditorNode);
+		if (SkelControl != nullptr)
+		{
+			PRAGMA_DISABLE_DEPRECATION_WARNINGS
+			SkelControl->MoveSelectActorLocation(GetAnimPreviewScene().GetPreviewMeshComponent(), static_cast<FAnimNode_SkeletalControlBase*>(InRuntimeNode));
+			SkelControl->CopyNodeDataTo(InRuntimeNode);
+			PRAGMA_ENABLE_DEPRECATION_WARNINGS
+		}
 	}
 
 	GetModeManager()->SetCoordSystem(GetWidgetCoordinateSystem());
@@ -128,31 +133,42 @@ void FAnimNodeEditMode::EnterMode(UAnimGraphNode_Base* InEditorNode, FAnimNode_B
 
 void FAnimNodeEditMode::ExitMode()
 {
-	UAnimGraphNode_SkeletalControlBase* SkelControl = Cast<UAnimGraphNode_SkeletalControlBase>(AnimNode);
-	if (SkelControl != nullptr)
+	for (EditorRuntimeNodePair CurrentNodePair : AnimNodes)
 	{
-		PRAGMA_DISABLE_DEPRECATION_WARNINGS
-		SkelControl->DeselectActor(GetAnimPreviewScene().GetPreviewMeshComponent());
-		PRAGMA_ENABLE_DEPRECATION_WARNINGS
+		if (CurrentNodePair.EditorAnimNode != nullptr)
+		{
+			UAnimGraphNode_SkeletalControlBase* SkelControl = Cast<UAnimGraphNode_SkeletalControlBase>(CurrentNodePair.EditorAnimNode);
+			if (SkelControl != nullptr)
+			{
+				PRAGMA_DISABLE_DEPRECATION_WARNINGS
+				SkelControl->DeselectActor(GetAnimPreviewScene().GetPreviewMeshComponent());
+				PRAGMA_ENABLE_DEPRECATION_WARNINGS
+			}
+		}
 	}
 
-	AnimNode = nullptr;
-	RuntimeAnimNode = nullptr;
+	AnimNodes.Empty();
 }
 
 void FAnimNodeEditMode::Render(const FSceneView* View, FViewport* Viewport, FPrimitiveDrawInterface* PDI)
 {
-	if (AnimNode != nullptr)
+	for (EditorRuntimeNodePair CurrentNodePair : AnimNodes)
 	{
-		AnimNode->Draw(PDI, GetAnimPreviewScene().GetPreviewMeshComponent());
+		if (CurrentNodePair.EditorAnimNode != nullptr)
+		{
+			CurrentNodePair.EditorAnimNode->Draw(PDI, GetAnimPreviewScene().GetPreviewMeshComponent());
+		}
 	}
 }
 
 void FAnimNodeEditMode::DrawHUD(FEditorViewportClient* ViewportClient, FViewport* Viewport, const FSceneView* View, FCanvas* Canvas)
 {
-	if (AnimNode != nullptr)
+	for (EditorRuntimeNodePair CurrentNodePair : AnimNodes)
 	{
-		AnimNode->DrawCanvas(*Viewport, *const_cast<FSceneView*>(View), *Canvas, GetAnimPreviewScene().GetPreviewMeshComponent());
+		if (CurrentNodePair.EditorAnimNode != nullptr)
+		{
+			CurrentNodePair.EditorAnimNode->DrawCanvas(*Viewport, *const_cast<FSceneView*>(View), *Canvas, GetAnimPreviewScene().GetPreviewMeshComponent());
+		}
 	}
 }
 
@@ -163,13 +179,17 @@ bool FAnimNodeEditMode::HandleClick(FEditorViewportClient* InViewportClient, HHi
 		HActor* ActorHitProxy = static_cast<HActor*>(HitProxy);
 		GetAnimPreviewScene().SetSelectedActor(ActorHitProxy->Actor);
 
-		UAnimGraphNode_SkeletalControlBase* SkelControl = Cast<UAnimGraphNode_SkeletalControlBase>(AnimNode);
-		if (SkelControl != nullptr)
+		for (EditorRuntimeNodePair CurrentNodePair : AnimNodes)
 		{
-			PRAGMA_DISABLE_DEPRECATION_WARNINGS
-			SkelControl->ProcessActorClick(ActorHitProxy);
-			PRAGMA_ENABLE_DEPRECATION_WARNINGS
+			UAnimGraphNode_SkeletalControlBase* SkelControl = Cast<UAnimGraphNode_SkeletalControlBase>(CurrentNodePair.EditorAnimNode);
+			if (SkelControl != nullptr)
+			{
+				PRAGMA_DISABLE_DEPRECATION_WARNINGS
+				SkelControl->ProcessActorClick(ActorHitProxy);
+				PRAGMA_ENABLE_DEPRECATION_WARNINGS
+			}
 		}
+
 		return true;
 	}
 
@@ -178,11 +198,12 @@ bool FAnimNodeEditMode::HandleClick(FEditorViewportClient* InViewportClient, HHi
 
 FVector FAnimNodeEditMode::GetWidgetLocation() const
 {
-	UAnimGraphNode_SkeletalControlBase* SkelControl = Cast<UAnimGraphNode_SkeletalControlBase>(AnimNode);
-	if (SkelControl != nullptr)
+	UAnimGraphNode_SkeletalControlBase* const EditorSkelControl = Cast<UAnimGraphNode_SkeletalControlBase>(GetActiveWidgetAnimNode());
+	FAnimNode_SkeletalControlBase* const RuntimeSkelControl = static_cast<FAnimNode_SkeletalControlBase*>(GetActiveWidgetRuntimeAnimNode());
+	if ((EditorSkelControl != nullptr) && (RuntimeSkelControl != nullptr))
 	{
 		PRAGMA_DISABLE_DEPRECATION_WARNINGS
-		return SkelControl->GetWidgetLocation(GetAnimPreviewScene().GetPreviewMeshComponent(), (FAnimNode_SkeletalControlBase*)RuntimeAnimNode);
+		return EditorSkelControl->GetWidgetLocation(GetAnimPreviewScene().GetPreviewMeshComponent(), RuntimeSkelControl);
 		PRAGMA_ENABLE_DEPRECATION_WARNINGS
 	}
 
@@ -194,8 +215,16 @@ bool FAnimNodeEditMode::StartTracking(FEditorViewportClient* InViewportClient, F
 	if (!bInTransaction)
 	{
 		GEditor->BeginTransaction(LOCTEXT("EditSkelControlNodeTransaction", "Edit Skeletal Control Node"));
-		AnimNode->SetFlags(RF_Transactional);
-		AnimNode->Modify();
+
+		for (EditorRuntimeNodePair CurrentNodePair : AnimNodes)
+		{
+			if (CurrentNodePair.EditorAnimNode != nullptr)
+			{
+				CurrentNodePair.EditorAnimNode->SetFlags(RF_Transactional);
+				CurrentNodePair.EditorAnimNode->Modify();
+			}
+		}
+
 		bInTransaction = true;
 	}
 
@@ -310,33 +339,36 @@ bool FAnimNodeEditMode::ShouldDrawWidget() const
 
 void FAnimNodeEditMode::DoTranslation(FVector& InTranslation)
 {
-	UAnimGraphNode_SkeletalControlBase* SkelControl = Cast<UAnimGraphNode_SkeletalControlBase>(AnimNode);
-	if (SkelControl != nullptr)
+	UAnimGraphNode_SkeletalControlBase* const EditorSkelControl = Cast<UAnimGraphNode_SkeletalControlBase>(GetActiveWidgetAnimNode());
+	FAnimNode_SkeletalControlBase* const RuntimeSkelControl = static_cast<FAnimNode_SkeletalControlBase*>(GetActiveWidgetRuntimeAnimNode());
+	if ((EditorSkelControl != nullptr) && (RuntimeSkelControl != nullptr))
 	{
 		PRAGMA_DISABLE_DEPRECATION_WARNINGS
-		SkelControl->DoTranslation(GetAnimPreviewScene().GetPreviewMeshComponent(), InTranslation, (FAnimNode_SkeletalControlBase*)RuntimeAnimNode);
+		EditorSkelControl->DoTranslation(GetAnimPreviewScene().GetPreviewMeshComponent(), InTranslation, RuntimeSkelControl);
 		PRAGMA_ENABLE_DEPRECATION_WARNINGS
 	}
 }
 
 void FAnimNodeEditMode::DoRotation(FRotator& InRotation)
 {
-	UAnimGraphNode_SkeletalControlBase* SkelControl = Cast<UAnimGraphNode_SkeletalControlBase>(AnimNode);
-	if (SkelControl != nullptr)
+	UAnimGraphNode_SkeletalControlBase* const EditorSkelControl = Cast<UAnimGraphNode_SkeletalControlBase>(GetActiveWidgetAnimNode());
+	FAnimNode_SkeletalControlBase* const RuntimeSkelControl = static_cast<FAnimNode_SkeletalControlBase*>(GetActiveWidgetRuntimeAnimNode());
+	if ((EditorSkelControl != nullptr) && (RuntimeSkelControl != nullptr))
 	{
 		PRAGMA_DISABLE_DEPRECATION_WARNINGS
-		SkelControl->DoRotation(GetAnimPreviewScene().GetPreviewMeshComponent(), InRotation, (FAnimNode_SkeletalControlBase*)RuntimeAnimNode);
+		EditorSkelControl->DoRotation(GetAnimPreviewScene().GetPreviewMeshComponent(), InRotation, RuntimeSkelControl);
 		PRAGMA_ENABLE_DEPRECATION_WARNINGS
 	}
 }
 
 void FAnimNodeEditMode::DoScale(FVector& InScale)
 {
-	UAnimGraphNode_SkeletalControlBase* SkelControl = Cast<UAnimGraphNode_SkeletalControlBase>(AnimNode);
-	if (SkelControl != nullptr)
+	UAnimGraphNode_SkeletalControlBase* const EditorSkelControl = Cast<UAnimGraphNode_SkeletalControlBase>(GetActiveWidgetAnimNode());
+	FAnimNode_SkeletalControlBase* const RuntimeSkelControl = static_cast<FAnimNode_SkeletalControlBase*>(GetActiveWidgetRuntimeAnimNode());
+	if ((EditorSkelControl != nullptr) && (RuntimeSkelControl != nullptr))
 	{
 		PRAGMA_DISABLE_DEPRECATION_WARNINGS
-		SkelControl->DoScale(GetAnimPreviewScene().GetPreviewMeshComponent(), InScale, (FAnimNode_SkeletalControlBase*)RuntimeAnimNode);
+		EditorSkelControl->DoScale(GetAnimPreviewScene().GetPreviewMeshComponent(), InScale, RuntimeSkelControl);
 		PRAGMA_ENABLE_DEPRECATION_WARNINGS
 	}
 }
@@ -346,15 +378,43 @@ void FAnimNodeEditMode::Tick(FEditorViewportClient* ViewportClient, float DeltaT
 	IAnimNodeEditMode::Tick(ViewportClient, DeltaTime);
 
 	// Keep actor location in sync with animation
-	UAnimGraphNode_SkeletalControlBase* SkelControl = Cast<UAnimGraphNode_SkeletalControlBase>(AnimNode);
-	if (SkelControl != nullptr)
+	UAnimGraphNode_SkeletalControlBase* const EditorSkelControl = Cast<UAnimGraphNode_SkeletalControlBase>(GetActiveWidgetAnimNode());
+	FAnimNode_SkeletalControlBase* const RuntimeSkelControl = static_cast<FAnimNode_SkeletalControlBase*>(GetActiveWidgetRuntimeAnimNode());
+	if ((EditorSkelControl != nullptr) && (RuntimeSkelControl != nullptr))
 	{
 		PRAGMA_DISABLE_DEPRECATION_WARNINGS
-		SkelControl->MoveSelectActorLocation(GetAnimPreviewScene().GetPreviewMeshComponent(), (FAnimNode_SkeletalControlBase*)RuntimeAnimNode);
+		EditorSkelControl->MoveSelectActorLocation(GetAnimPreviewScene().GetPreviewMeshComponent(), RuntimeSkelControl);
 		PRAGMA_ENABLE_DEPRECATION_WARNINGS
 	}
 }
 
+void FAnimNodeEditMode::Exit()
+{
+	IAnimNodeEditMode::Exit();
+
+	ExitMode(); // Editor modes can be destroyed without a call to ExitMode. Calling it here to ensure everything is cleaned up correctly.
+}
+
+UAnimGraphNode_Base* FAnimNodeEditMode::GetActiveWidgetAnimNode() const
+{
+	if (AnimNodes.Num() > 0)
+	{
+		return AnimNodes.Last().EditorAnimNode;
+	}
+
+	return nullptr;
+}
+
+FAnimNode_Base* FAnimNodeEditMode::GetActiveWidgetRuntimeAnimNode() const
+{
+	if (AnimNodes.Num() > 0)
+	{
+		return AnimNodes.Last().RuntimeAnimNode;
+	}
+
+	return nullptr;
+}
+
 void FAnimNodeEditMode::ConvertToComponentSpaceTransform(const USkeletalMeshComponent* SkelComp, const FTransform & InTransform, FTransform & OutCSTransform, int32 BoneIndex, EBoneControlSpace Space)
 {
 	USkeleton* Skeleton = SkelComp->SkeletalMesh->GetSkeleton();
diff --git a/Engine/Source/Editor/AnimGraph/Public/AnimNodeEditMode.h b/Engine/Source/Editor/AnimGraph/Public/AnimNodeEditMode.h
--- a/Engine/Source/Editor/AnimGraph/Public/AnimNodeEditMode.h
+++ b/Engine/Source/Editor/AnimGraph/Public/AnimNodeEditMode.h
@@ -51,6 +51,7 @@ public:
 	virtual bool GetCustomInputCoordinateSystem(FMatrix& InMatrix, void* InData) override;
 	virtual bool ShouldDrawWidget() const override;
 	virtual void Tick(FEditorViewportClient* ViewportClient, float DeltaTime) override;
+	virtual void Exit() override;
 
 protected:
 	// local conversion functions for drawing
@@ -65,12 +66,26 @@ protected:
 	static FVector ConvertWidgetLocation(const USkeletalMeshComponent* InSkelComp, FCSPose<FCompactHeapPose>& InMeshBases, const FName& BoneName, const FVector& InLocation, const EBoneControlSpace Space);
 	static FVector ConvertWidgetLocation(const USkeletalMeshComponent* InSkelComp, FCSPose<FCompactHeapPose>& InMeshBases, const FBoneSocketTarget& Target, const FVector& InLocation, const EBoneControlSpace Space);
 
+	virtual UAnimGraphNode_Base* GetActiveWidgetAnimNode() const; // Return the editor node associated with the selected widget. All widget operations are performed on this node.
+	virtual FAnimNode_Base*	GetActiveWidgetRuntimeAnimNode() const; // Return the runtime node associated with the selected widget. All widget operations are performed on this node.
+
 protected:
-	/** The node we are operating on */
-	class UAnimGraphNode_Base* AnimNode;
 
-	/** The runtime node in the preview scene */
-	struct FAnimNode_Base* RuntimeAnimNode;
+	struct EditorRuntimeNodePair
+	{
+		EditorRuntimeNodePair(UAnimGraphNode_Base* InEditorAnimNode, FAnimNode_Base* InRuntimeAnimNode)
+			: EditorAnimNode(InEditorAnimNode)
+			, RuntimeAnimNode(InRuntimeAnimNode)
+		{}
+
+		/** The node we are operating on */
+		UAnimGraphNode_Base* EditorAnimNode;
+
+		/** The runtime node in the preview scene */
+		FAnimNode_Base* RuntimeAnimNode;
+	};
+
+	TArray< EditorRuntimeNodePair > AnimNodes;
 
 private:
 	bool bManipulating;
diff --git a/Engine/Source/Editor/AnimationBlueprintEditor/Private/AnimationBlueprintEditor.cpp b/Engine/Source/Editor/AnimationBlueprintEditor/Private/AnimationBlueprintEditor.cpp
--- a/Engine/Source/Editor/AnimationBlueprintEditor/Private/AnimationBlueprintEditor.cpp
+++ b/Engine/Source/Editor/AnimationBlueprintEditor/Private/AnimationBlueprintEditor.cpp
@@ -110,6 +110,33 @@ namespace AnimationBlueprintEditorTabs
 	const FName CurveNamesTab(TEXT("AnimCurveViewerTab"));
 };
 
+/////////////////////////////////////////////////////
+// SortedContainerDifference
+
+/** Algorithm to find the difference between two sorted sets of unique values - outputs two sets, all the elements that are in set A but not in set B and all the elements that are in set B but not in set A **/
+template <typename TContainerType, typename TPredicate> void SortedContainerDifference(const TContainerType& LhsContainer, const TContainerType& RhsContainer, TContainerType& OutLhsDifference, TContainerType& OutRhsDifference, const TPredicate& SortPredicate)
+{
+	for (unsigned int LhsIndex = 0, RhsIndex = 0, LhsMax = LhsContainer.Num(), RhsMax = RhsContainer.Num(); (LhsIndex < LhsMax) || (RhsIndex < RhsMax); )
+	{
+		if ((LhsIndex < LhsMax) && (!(RhsIndex < RhsMax) || SortPredicate(LhsContainer[LhsIndex], RhsContainer[RhsIndex])))
+		{
+			OutRhsDifference.Add(LhsContainer[LhsIndex]);
+			++LhsIndex;
+		}
+		else if ((RhsIndex < RhsMax) && (!(LhsIndex < LhsMax) || SortPredicate(RhsContainer[RhsIndex], LhsContainer[LhsIndex])))
+		{
+			OutLhsDifference.Add(RhsContainer[RhsIndex]);
+			++RhsIndex;
+		}
+		else
+		{
+			++LhsIndex;
+			++RhsIndex;
+		}
+	}
+}
+
+
 /////////////////////////////////////////////////////
 // SAnimBlueprintPreviewPropertyEditor
 
@@ -1492,16 +1519,16 @@ void FAnimationBlueprintEditor::ClearSelectedActor()
 	GetPreviewScene()->ClearSelectedActor();
 }
 
-void FAnimationBlueprintEditor::ClearSelectedAnimGraphNode()
+void FAnimationBlueprintEditor::ClearSelectedAnimGraphNodes()
 {
-	SelectedAnimGraphNode.Reset();
+	SelectedAnimGraphNodes.Empty();
 }
 
 void FAnimationBlueprintEditor::DeselectAll()
 {
 	GetSkeletonTree()->DeselectAll();
 	ClearSelectedActor();
-	ClearSelectedAnimGraphNode();
+	ClearSelectedAnimGraphNodes();
 }
 
 void FAnimationBlueprintEditor::PostRedo(bool bSuccess)
@@ -1538,13 +1565,14 @@ void FAnimationBlueprintEditor::NotifyPostChange(const FPropertyChangedEvent& Pr
 	FBlueprintEditor::NotifyPostChange(PropertyChangedEvent, PropertyThatChanged);
 
 	// When you change properties on a node, call CopyNodeDataToPreviewNode to allow pushing those to preview instance, for live editing
-	UAnimGraphNode_Base* SelectedNode = SelectedAnimGraphNode.Get();
-	if (SelectedNode)
+	for (TWeakObjectPtr< class UAnimGraphNode_Base > CurrentAnimGraphNode : SelectedAnimGraphNodes)
 	{
-		FAnimNode_Base* PreviewNode = FindAnimNode(SelectedNode);
-		if (PreviewNode)
+		if (UAnimGraphNode_Base* CurrentNode = CurrentAnimGraphNode.Get())
 		{
-			SelectedNode->CopyNodeDataToPreviewNode(PreviewNode);
+			if (FAnimNode_Base* PreviewNode = FindAnimNode(CurrentNode))
+			{
+				CurrentNode->CopyNodeDataToPreviewNode(PreviewNode);
+			}
 		}
 	}
 }
@@ -1652,8 +1680,8 @@ void FAnimationBlueprintEditor::OnBlueprintPostCompile(UBlueprint* InBlueprint)
 			}
 		}
 
-		// reset the selected skeletal control node
-		SelectedAnimGraphNode.Reset();
+		// reset the selected skeletal control nodes
+		ClearSelectedAnimGraphNodes();
 
 		// if the user manipulated Pin values directly from the node, then should copy updated values to the internal node to retain data consistency
 		OnPostCompile();
@@ -1754,32 +1782,67 @@ void FAnimationBlueprintEditor::OnSelectedNodesChangedImpl(const TSet<class UObj
 {
 	FBlueprintEditor::OnSelectedNodesChangedImpl(NewSelection);
 
-	IPersonaEditorModeManager* PersonaEditorModeManager = static_cast<IPersonaEditorModeManager*>(&GetEditorModeManager());
+	IPersonaEditorModeManager* const PersonaEditorModeManager = static_cast<IPersonaEditorModeManager*>(&GetEditorModeManager());
 
-	if (UAnimGraphNode_Base* SelectedAnimGraphNodePtr = SelectedAnimGraphNode.Get())
+	if (PersonaEditorModeManager)
 	{
-		FAnimNode_Base* PreviewNode = FindAnimNode(SelectedAnimGraphNodePtr);
-		if (PersonaEditorModeManager)
+		// Update the list of selected nodes, being careful to maintain the order of the list as this is an important requirement of the UI.
+
+		using FSelectedNodePtr = TWeakObjectPtr< class UAnimGraphNode_Base >;
+
+		TArray< FSelectedNodePtr > AddSelection;	// Nodes that should be added to the current selection.
+		TArray< FSelectedNodePtr > RemSelection;	// Nodes that should be removed from the current selection.
+
+		// Compare the set of nodes in 'NewSelection' with the list of previously selected nodes to identify nodes that should be added / removed from the selection.
 		{
-			SelectedAnimGraphNodePtr->OnNodeSelected(false, *PersonaEditorModeManager, PreviewNode);
+			TArray< FSelectedNodePtr > OldSelectionSorted(SelectedAnimGraphNodes);
+			TArray< FSelectedNodePtr > NewSelectionSorted;
+
+			for (UObject* NewSelectedObject : NewSelection)
+			{
+				if (UAnimGraphNode_Base* NewSelectedAnimGraphNode = Cast<UAnimGraphNode_Base>(NewSelectedObject))
+				{
+					NewSelectionSorted.Add(NewSelectedAnimGraphNode);
+				}
+			}
+
+			auto SortPredicate = [](const FSelectedNodePtr& Lhs, const FSelectedNodePtr& Rhs) { return Lhs.Get() < Rhs.Get();  };
+
+			OldSelectionSorted.Sort(SortPredicate);
+			NewSelectionSorted.Sort(SortPredicate);
+
+			SortedContainerDifference(OldSelectionSorted, NewSelectionSorted, AddSelection, RemSelection, SortPredicate);
 		}
 
-		SelectedAnimGraphNode.Reset();
-	}
+		// Register de-selection with all the previously selected nodes.
+		for (FSelectedNodePtr CurrentAnimGraphNode : SelectedAnimGraphNodes)
+		{
+			UAnimGraphNode_Base* const CurrentAnimGraphNodePtr = CurrentAnimGraphNode.Get();
+			if (FAnimNode_Base* const PreviewNode = FindAnimNode(CurrentAnimGraphNodePtr))
+			{
+				CurrentAnimGraphNodePtr->OnNodeSelected(false, *PersonaEditorModeManager, PreviewNode);
+			}
+		}
 
-	// if we only have one node selected, let it know
-	UAnimGraphNode_Base* NewSelectedAnimGraphNode = nullptr;
-	if (NewSelection.Num() == 1)
-	{
-		NewSelectedAnimGraphNode = Cast<UAnimGraphNode_Base>(*NewSelection.CreateConstIterator());
-		if (NewSelectedAnimGraphNode != nullptr)
+		// Remove all the nodes that are no longer selected.
+		for (FSelectedNodePtr CurrentAnimGraphNode : RemSelection)
 		{
-			SelectedAnimGraphNode = NewSelectedAnimGraphNode;
+			SelectedAnimGraphNodes.Remove(CurrentAnimGraphNode);
+		}
 
-			FAnimNode_Base* PreviewNode = FindAnimNode(NewSelectedAnimGraphNode);
-			if (PreviewNode && PersonaEditorModeManager)
+		// Add all the newly selected nodes.
+		for (FSelectedNodePtr CurrentAnimGraphNode : AddSelection)
+		{
+			SelectedAnimGraphNodes.Add(CurrentAnimGraphNode);
+		}
+
+		// Register re-selection with all the currently selected nodes.
+		for (FSelectedNodePtr CurrentAnimGraphNode : SelectedAnimGraphNodes)
+		{
+			UAnimGraphNode_Base* const CurrentAnimGraphNodePtr = CurrentAnimGraphNode.Get();
+			if (FAnimNode_Base* const PreviewNode = FindAnimNode(CurrentAnimGraphNodePtr))
 			{
-				NewSelectedAnimGraphNode->OnNodeSelected(true, *PersonaEditorModeManager, PreviewNode);
+				CurrentAnimGraphNodePtr->OnNodeSelected(true, *PersonaEditorModeManager, PreviewNode);
 			}
 		}
 	}
diff --git a/Engine/Source/Editor/AnimationBlueprintEditor/Private/AnimationBlueprintEditor.h b/Engine/Source/Editor/AnimationBlueprintEditor/Private/AnimationBlueprintEditor.h
--- a/Engine/Source/Editor/AnimationBlueprintEditor/Private/AnimationBlueprintEditor.h
+++ b/Engine/Source/Editor/AnimationBlueprintEditor/Private/AnimationBlueprintEditor.h
@@ -116,8 +116,8 @@ public:
 	/** Clears the selected actor */
 	void ClearSelectedActor();
 
-	/** Clears the selected anim graph node */
-	void ClearSelectedAnimGraphNode();
+	/** Clears the selected anim graph nodes */
+	void ClearSelectedAnimGraphNodes();
 
 	/** Clears the selection (both sockets and bones). Also broadcasts this */
 	void DeselectAll();
@@ -344,7 +344,7 @@ private:
 	TSharedPtr<class ISkeletonTree> SkeletonTree;
 
 	// selected anim graph node 
-	TWeakObjectPtr<class UAnimGraphNode_Base> SelectedAnimGraphNode;
+	TArray< TWeakObjectPtr< class UAnimGraphNode_Base > > SelectedAnimGraphNodes;
 
 	/** Sequence Browser **/
 	TWeakPtr<class IAnimationSequenceBrowser> SequenceBrowser;
