Fixing bad merge

diff --git a/Engine/Source/Editor/AnimGraph/Private/EditModes/AnimNodeEditMode.cpp b/Engine/Source/Editor/AnimGraph/Private/EditModes/AnimNodeEditMode.cpp
--- a/Engine/Source/Editor/AnimGraph/Private/EditModes/AnimNodeEditMode.cpp
+++ b/Engine/Source/Editor/AnimGraph/Private/EditModes/AnimNodeEditMode.cpp
@@ -304,8 +304,6 @@ bool FAnimNodeEditMode::InputDelta(FEditorViewportClient* InViewportClient, FVie
 			DoScale(InScale);
 		}
 
-		AnimNode->PostEditRefreshDebuggedComponent();
-		
 		InViewport->Invalidate();
 	}
 
