Fix crash compiling child anim BPs with certain configurations

Copy mutable data (in the same way as constant data) to ensure that child class containers are coorectly initialized

diff --git a/Engine/Source/Editor/AnimGraph/Private/AnimBlueprintCompiler.cpp b/Engine/Source/Editor/AnimGraph/Private/AnimBlueprintCompiler.cpp
--- a/Engine/Source/Editor/AnimGraph/Private/AnimBlueprintCompiler.cpp
+++ b/Engine/Source/Editor/AnimGraph/Private/AnimBlueprintCompiler.cpp
@@ -770,6 +770,28 @@ void FAnimBlueprintCompilerContext::CopyTermDefaultsToDefaultObject(UObject* Def
 			}
 		}
 		
+		// Copy from root mutable data to our new mutable data
+		if (NewAnimBlueprintClass->GetMutableNodeData(DefaultObject) && RootAnimClass->GetMutableNodeData(RootDefaultObject))
+		{
+			// These properties should have been linked and cached by now
+			check(RootAnimClass->MutableNodeDataProperty);
+			check(RootAnimClass->MutableNodeDataProperty->Struct);
+			check(NewAnimBlueprintClass->MutableNodeDataProperty);
+			check(NewAnimBlueprintClass->MutableNodeDataProperty->Struct);
+			
+			for (TFieldIterator<FProperty> PropertyIt(RootAnimClass->MutableNodeDataProperty->Struct); PropertyIt; ++PropertyIt)
+			{
+				FProperty* RootProperty = *PropertyIt;
+				FProperty* ChildProperty = FindFProperty<FProperty>(NewAnimBlueprintClass->MutableNodeDataProperty->Struct, *RootProperty->GetName());
+				check(ChildProperty != nullptr);
+		 
+				const uint8* SourcePtr = RootProperty->ContainerPtrToValuePtr<uint8>(RootAnimClass->GetMutableNodeData(RootDefaultObject));
+				uint8* DestPtr = ChildProperty->ContainerPtrToValuePtr<uint8>(NewAnimBlueprintClass->GetMutableNodeData(DefaultObject));
+				check(SourcePtr != nullptr && DestPtr != nullptr);
+				RootProperty->CopyCompleteValue(DestPtr, SourcePtr);
+			}
+		}
+		
 		// Re-initialize node data tables (they would be overwritten in the loop above)
 		NewAnimBlueprintClass->InitializeAnimNodeData(DefaultObject, true);
 	}
