Fixed null blendspaces being referenced by child anim BPs that use blendspace graphs

Ensure that RF_Public is set so that external packages (like child BPs) can access the sub-object

diff --git a/Engine/Source/Editor/AnimGraph/Private/AnimBlueprintExtension_BlendSpaceGraph.cpp b/Engine/Source/Editor/AnimGraph/Private/AnimBlueprintExtension_BlendSpaceGraph.cpp
--- a/Engine/Source/Editor/AnimGraph/Private/AnimBlueprintExtension_BlendSpaceGraph.cpp
+++ b/Engine/Source/Editor/AnimGraph/Private/AnimBlueprintExtension_BlendSpaceGraph.cpp
@@ -13,6 +13,10 @@ UBlendSpace* UAnimBlueprintExtension_BlendSpaceGraph::AddBlendSpace(UBlendSpace*
 {
 	UBlendSpace* CopiedBlendSpace = DuplicateObject(InSourceBlendSpace, Class);
 	CopiedBlendSpace->ClearFlags(RF_Transient);
+
+	// RF_Public is required because this blendspace may need to be referenced in child classes
+	CopiedBlendSpace->SetFlags(RF_Public);
+	
 	Subsystem.BlendSpaces.Add(CopiedBlendSpace);
 	return CopiedBlendSpace;
 }
\ No newline at end of file
