Fix for blend masks not being added when adding input pins to the layered blend per bone node
#rnx

diff --git a/Engine/Source/Editor/AnimGraph/Private/AnimGraphNode_LayeredBoneBlend.cpp b/Engine/Source/Editor/AnimGraph/Private/AnimGraphNode_LayeredBoneBlend.cpp
--- a/Engine/Source/Editor/AnimGraph/Private/AnimGraphNode_LayeredBoneBlend.cpp
+++ b/Engine/Source/Editor/AnimGraph/Private/AnimGraphNode_LayeredBoneBlend.cpp
@@ -43,6 +43,21 @@ void UAnimGraphNode_LayeredBoneBlend::PostEditChangeProperty(struct FPropertyCha
 	if (PropertyName == GET_MEMBER_NAME_STRING_CHECKED(FAnimNode_LayeredBoneBlend, BlendMode))
 	{
 		// If we  change blend modes, we need to resize our containers
+		FScopedTransaction Transaction(LOCTEXT("ChangeBlendMode", "Change Blend Mode"));
+		Modify();
+
+		const int32 NumPoses = Node.BlendPoses.Num();
+		if (Node.BlendMode == ELayeredBoneBlendMode::BlendMask)
+		{
+			Node.LayerSetup.Reset();
+			Node.BlendMasks.SetNum(NumPoses);
+		}
+		else
+		{
+			Node.BlendMasks.Reset();
+			Node.LayerSetup.SetNum(NumPoses);
+		}
+
 		FBlueprintEditorUtils::MarkBlueprintAsStructurallyModified(GetBlueprint());
 	}
 
diff --git a/Engine/Source/Runtime/AnimGraphRuntime/Public/AnimNodes/AnimNode_LayeredBoneBlend.h b/Engine/Source/Runtime/AnimGraphRuntime/Public/AnimNodes/AnimNode_LayeredBoneBlend.h
--- a/Engine/Source/Runtime/AnimGraphRuntime/Public/AnimNodes/AnimNode_LayeredBoneBlend.h
+++ b/Engine/Source/Runtime/AnimGraphRuntime/Public/AnimNodes/AnimNode_LayeredBoneBlend.h
@@ -129,6 +129,7 @@ public:
 	void AddPose()
 	{
 		BlendWeights.Add(1.f);
+		BlendMasks.Add(nullptr);
 		new (BlendPoses) FPoseLink();
 		new (LayerSetup) FInputBlendPose();
 	}
@@ -136,6 +137,7 @@ public:
 	void RemovePose(int32 PoseIndex)
 	{
 		BlendWeights.RemoveAt(PoseIndex);
+		BlendMasks.RemoveAt(PoseIndex);
 		BlendPoses.RemoveAt(PoseIndex);
 		LayerSetup.RemoveAt(PoseIndex);
 	}
