Fix crash selecting linked anim layer class

Prevent filtering and selection of anim blueprints with incompatible/null skeletons, depending on context.
Additionally prevents non-template anim blueprints from being selected within templates and adds compile-time checks for compatibilty

diff --git a/Engine/Source/Editor/AnimGraph/Private/AnimGraphNode_LinkedAnimGraph.cpp b/Engine/Source/Editor/AnimGraph/Private/AnimGraphNode_LinkedAnimGraph.cpp
--- a/Engine/Source/Editor/AnimGraph/Private/AnimGraphNode_LinkedAnimGraph.cpp
+++ b/Engine/Source/Editor/AnimGraph/Private/AnimGraphNode_LinkedAnimGraph.cpp
@@ -62,6 +62,14 @@ void UAnimGraphNode_LinkedAnimGraph::SetupFromAsset(const FAssetData& InAssetDat
 			SkeletonName.Empty();
 		}
 
+		FString TagTemplateValue;
+		InAssetData.GetTagValue("bIsTemplate", TagTemplateValue);
+		bIsTemplateAnimBlueprint = TagTemplateValue.Equals(TEXT("True"));
+
+		FString BlueprintTypeValue;
+		InAssetData.GetTagValue("BlueprintType", BlueprintTypeValue);
+		bIsInterfaceBlueprint = BlueprintTypeValue.Equals(TEXT("BPTYPE_Interface"));
+		
 		if(!bInIsTemplateNode)
 		{
 			UAnimBlueprint* AnimBlueprint = CastChecked<UAnimBlueprint>(InAssetData.GetAsset());
@@ -140,7 +148,11 @@ bool UAnimGraphNode_LinkedAnimGraph::IsActionFilteredOut(class FBlueprintActionF
 {
 	bool bIsFilteredOut = false;
 
-	if(!SkeletonName.IsEmpty())
+	if(bIsInterfaceBlueprint)
+	{
+		bIsFilteredOut = true;
+	}
+	else if(!SkeletonName.IsEmpty())
 	{
 		FBlueprintActionContext const& FilterContext = Filter.Context;
 
@@ -148,7 +160,7 @@ bool UAnimGraphNode_LinkedAnimGraph::IsActionFilteredOut(class FBlueprintActionF
 		{
 			if (UAnimBlueprint* AnimBlueprint = Cast<UAnimBlueprint>(Blueprint))
 			{
-				if(AnimBlueprint->TargetSkeleton != nullptr && !AnimBlueprint->TargetSkeleton->IsCompatibleSkeletonByAssetString(SkeletonName))
+				if(!AnimBlueprint->IsCompatibleByAssetString(SkeletonName, bIsTemplateAnimBlueprint, bIsInterfaceBlueprint))
 				{
 					bIsFilteredOut = true;
 					break;
diff --git a/Engine/Source/Editor/AnimGraph/Private/AnimGraphNode_LinkedAnimGraphBase.cpp b/Engine/Source/Editor/AnimGraph/Private/AnimGraphNode_LinkedAnimGraphBase.cpp
--- a/Engine/Source/Editor/AnimGraph/Private/AnimGraphNode_LinkedAnimGraphBase.cpp
+++ b/Engine/Source/Editor/AnimGraph/Private/AnimGraphNode_LinkedAnimGraphBase.cpp
@@ -196,6 +196,15 @@ void UAnimGraphNode_LinkedAnimGraphBase::ValidateAnimNodeDuringCompilation(USkel
 	{
 		MessageLog.Error(TEXT("Linked instance node @@ targets instance class @@ which it is inside, this would cause a loop."), this, AnimBP->GetAnimBlueprintGeneratedClass());
 	}
+
+	// Check for compatibility
+	if(UAnimBlueprint* TargetAnimBP = Cast<UAnimBlueprint>(UBlueprint::GetBlueprintFromClass(GetTargetClass())))
+	{
+		if(!AnimBP->IsCompatible(TargetAnimBP))
+		{
+			MessageLog.Error(TEXT("Linked instance node @@ targets instance class @@ which is incompatible."), this, GetTargetClass());
+		}
+	}
 }
 
 void UAnimGraphNode_LinkedAnimGraphBase::CreateOutputPins()
@@ -417,13 +426,17 @@ bool UAnimGraphNode_LinkedAnimGraphBase::OnShouldFilterInstanceBlueprint(const F
 		return true;
 	}
 
-	// Check skeleton
-	FAssetDataTagMapSharedView::FFindTagResult Result = AssetData.TagsAndValues.FindTag("TargetSkeleton");
-	if (Result.IsSet())
+	// Check skeleton & flags
+	FAssetDataTagMapSharedView::FFindTagResult TargetSkeletonResult = AssetData.TagsAndValues.FindTag("TargetSkeleton");
+	FAssetDataTagMapSharedView::FFindTagResult IsTemplateResult = AssetData.TagsAndValues.FindTag("bIsTemplate");
+	FAssetDataTagMapSharedView::FFindTagResult BlueprintTypeResult = AssetData.TagsAndValues.FindTag("BlueprintType");
+	if (TargetSkeletonResult.IsSet())
 	{
+		const bool bIsTemplate = IsTemplateResult.IsSet() && IsTemplateResult.Equals(TEXT("True"));
+		const bool bIsInterface = BlueprintTypeResult.IsSet() && BlueprintTypeResult.Equals(TEXT("BPTYPE_Interface"));
 		if (UAnimBlueprint* CurrentBlueprint = Cast<UAnimBlueprint>(GetBlueprint()))
 		{
-			if (!CurrentBlueprint->TargetSkeleton->IsCompatibleSkeletonByAssetString(Result.GetValue()))
+			if (!CurrentBlueprint->IsCompatibleByAssetString(TargetSkeletonResult.GetValue(), bIsTemplate, bIsInterface))
 			{
 				return true;
 			}
diff --git a/Engine/Source/Editor/AnimGraph/Public/AnimGraphNode_LinkedAnimGraphBase.h b/Engine/Source/Editor/AnimGraph/Public/AnimGraphNode_LinkedAnimGraphBase.h
--- a/Engine/Source/Editor/AnimGraph/Public/AnimGraphNode_LinkedAnimGraphBase.h
+++ b/Engine/Source/Editor/AnimGraph/Public/AnimGraphNode_LinkedAnimGraphBase.h
@@ -83,6 +83,12 @@ protected:
 
 	// Skeleton name used for filtering unloaded classes 
 	FString SkeletonName;
+
+	// Template flag used for filtering unloaded classes
+	bool bIsTemplateAnimBlueprint = false;
+
+	// Interface flag used for filtering unloaded classes
+	bool bIsInterfaceBlueprint = false;
 };
 
 UE_DEPRECATED(4.24, "UAnimGraphNode_SubInstanceBase has been renamed to UAnimGraphNode_LinkedAnimGraphBase")
diff --git a/Engine/Source/Runtime/Engine/Classes/Animation/AnimBlueprint.h b/Engine/Source/Runtime/Engine/Classes/Animation/AnimBlueprint.h
--- a/Engine/Source/Runtime/Engine/Classes/Animation/AnimBlueprint.h
+++ b/Engine/Source/Runtime/Engine/Classes/Animation/AnimBlueprint.h
@@ -210,6 +210,20 @@ public:
 	// Queue a refresh of the set of anim blueprint extensions that this anim blueprint hosts.
 	// Usually called from anim graph nodes to ensure that extensions that are no longer required are cleaned up.
 	void RequestRefreshExtensions() { bRefreshExtensions = true; }
+
+	// Check if the anim BP is compatible with this one (for linked instancing). Checks target skeleton, template flags
+	// blueprint type.
+	// Note compatibility is directional - e.g. template anim BPs can be instanced within any 'regular' anim BP, but not
+	// vice versa
+	// @param	InAnimBlueprint		The anim blueprint to check for compatibility
+	bool IsCompatible(const UAnimBlueprint* InAnimBlueprint) const;
+	
+	// Check if the asset path of a skeleton, template and interface flags are compatible with this anim blueprint
+	// (for linked instancing)
+	// @param	InSkeletonAsset		The asset path of the skeleton asset used by the anim blueprint
+	// @param	bInIsTemplate		Whether the anim blueprint to check is a template
+	// @param	bInIsInterface		Whether the anim blueprint to check is an interface
+	bool IsCompatibleByAssetString(const FString& InSkeletonAsset, bool bInIsTemplate, bool bInIsInterface) const;
 	
 public:
 	// Array of overrides to asset containing nodes in the parent that have been overridden
diff --git a/Engine/Source/Runtime/Engine/Private/AnimBlueprint.cpp b/Engine/Source/Runtime/Engine/Private/AnimBlueprint.cpp
--- a/Engine/Source/Runtime/Engine/Private/AnimBlueprint.cpp
+++ b/Engine/Source/Runtime/Engine/Private/AnimBlueprint.cpp
@@ -348,3 +348,39 @@ FAnimBlueprintDebugData* UAnimBlueprint::GetDebugData() const
 	return nullptr;
 #endif // WITH_EDITORONLY_DATA
 }
+
+#if WITH_EDITORONLY_DATA
+bool UAnimBlueprint::IsCompatible(const UAnimBlueprint* InAnimBlueprint) const
+{
+	if(InAnimBlueprint->bIsTemplate)
+	{
+		// Directionality here, templates are compatible with all anim BPs, but not necessarily vice versa
+		return true;
+	}
+	
+	if(InAnimBlueprint->BlueprintType == BPTYPE_Interface)
+	{
+		// Interfaces dont bother with skeleton checks - assume compatibility is driven by the interface machinery
+		return true;
+	}
+
+	return (TargetSkeleton != nullptr && TargetSkeleton->IsCompatible(InAnimBlueprint->TargetSkeleton));
+}
+
+bool UAnimBlueprint::IsCompatibleByAssetString(const FString& InSkeletonAsset, bool bInIsTemplate, bool bInIsInterface) const
+{
+	if(bInIsTemplate)
+	{
+		// Directionality here, templates are compatible with all anim BPs, but not necessarily vice versa
+		return true;
+	}
+	
+	if(bInIsInterface)
+	{
+		// Interfaces dont bother with skeleton checks - assume compatibility is driven by the interface machinery
+		return true;
+	}
+
+	return (TargetSkeleton != nullptr && TargetSkeleton->IsCompatibleSkeletonByAssetString(InSkeletonAsset));
+}
+#endif
\ No newline at end of file
