Reinstate code that was removed in CL-18260151 to fix a bad merge

diff --git a/Engine/Source/Editor/AnimGraph/Private/EditModes/AnimNodeEditMode.cpp b/Engine/Source/Editor/AnimGraph/Private/EditModes/AnimNodeEditMode.cpp
--- a/Engine/Source/Editor/AnimGraph/Private/EditModes/AnimNodeEditMode.cpp
+++ b/Engine/Source/Editor/AnimGraph/Private/EditModes/AnimNodeEditMode.cpp
@@ -304,6 +304,14 @@ bool FAnimNodeEditMode::InputDelta(FEditorViewportClient* InViewportClient, FVie
 			DoScale(InScale);
 		}
 
+		for (EditorRuntimeNodePair CurrentNodePair : AnimNodes)
+		{
+			if (CurrentNodePair.EditorAnimNode)
+			{
+				CurrentNodePair.EditorAnimNode->PostEditRefreshDebuggedComponent();
+			}
+		}
+
 		InViewport->Invalidate();
 	}
 
