[FIX] Crash when exiting anim graph editor - removeing call to ExitMode from Exit function as Exit is called from within the owning objects dtor but ExitMode calls v-function on owning object which is unsafe

diff --git a/Engine/Source/Editor/AnimGraph/Private/EditModes/AnimNodeEditMode.cpp b/Engine/Source/Editor/AnimGraph/Private/EditModes/AnimNodeEditMode.cpp
--- a/Engine/Source/Editor/AnimGraph/Private/EditModes/AnimNodeEditMode.cpp
+++ b/Engine/Source/Editor/AnimGraph/Private/EditModes/AnimNodeEditMode.cpp
@@ -398,7 +398,7 @@ void FAnimNodeEditMode::Exit()
 {
 	IAnimNodeEditMode::Exit();
 
-	ExitMode(); // Editor modes can be destroyed without a call to ExitMode. Calling it here to ensure everything is cleaned up correctly.
+	AnimNodes.Empty();
 }
 
 UAnimGraphNode_Base* FAnimNodeEditMode::GetActiveWidgetAnimNode() const
