Fix crash when selecting property access chains involving structs with native break functions

diff --git a/Engine/Source/Editor/AnimGraph/Private/AnimBlueprintExtension_PropertyAccess.cpp b/Engine/Source/Editor/AnimGraph/Private/AnimBlueprintExtension_PropertyAccess.cpp
--- a/Engine/Source/Editor/AnimGraph/Private/AnimBlueprintExtension_PropertyAccess.cpp
+++ b/Engine/Source/Editor/AnimGraph/Private/AnimBlueprintExtension_PropertyAccess.cpp
@@ -263,7 +263,19 @@ void UAnimBlueprintExtension_PropertyAccess::ExpandPropertyAccess(FKismetCompile
 				// Create a break struct/split pin node
 				const UEdGraphSchema_K2* K2Schema = CastChecked<UEdGraphSchema_K2>(InParentGraph->GetSchema());
 				UK2Node* SplitPinNode = K2Schema->CreateSplitPinNode(CurrentPin, UEdGraphSchema_K2::FCreateSplitPinNodeParams(&InCompilerContext, InParentGraph));
-				UEdGraphPin* InputPin = SplitPinNode->FindPinChecked(ScriptStruct->GetFName(), EGPD_Input);
+				UEdGraphPin* InputPin = SplitPinNode->FindPinByPredicate([ScriptStruct](UEdGraphPin* InPin)
+				{
+					if(InPin && InPin->Direction == EGPD_Input && InPin->PinType.PinCategory == UEdGraphSchema_K2::PC_Struct)
+					{
+						if(UScriptStruct* PinStruct = Cast<UScriptStruct>(InPin->PinType.PinSubCategoryObject.Get()))
+						{
+							return PinStruct == ScriptStruct;
+						}
+					}
+
+					return false;
+				});
+				check(InputPin);
 				CurrentPin->MakeLinkTo(InputPin);
 
 				// Current pin is the property of the struct
