Anim blueprint editor modes now display transform toolbar

Deferred to the active editor modes when determining transform toolbar visibility
Updated existing edit modes to correctly update their available widget modes

diff --git a/Engine/Source/Editor/AnimGraph/Private/EditModes/CCDIKEditMode.cpp b/Engine/Source/Editor/AnimGraph/Private/EditModes/CCDIKEditMode.cpp
--- a/Engine/Source/Editor/AnimGraph/Private/EditModes/CCDIKEditMode.cpp
+++ b/Engine/Source/Editor/AnimGraph/Private/EditModes/CCDIKEditMode.cpp
@@ -45,6 +45,11 @@ UE::Widget::EWidgetMode FCCDIKEditMode::GetWidgetMode() const
 	return UE::Widget::WM_None;
 }
 
+bool FCCDIKEditMode::UsesTransformWidget(UE::Widget::EWidgetMode InWidgetMode) const
+{
+	return InWidgetMode == UE::Widget::WM_Translate;
+}
+
 void FCCDIKEditMode::DoTranslation(FVector& InTranslation)
 {
 	USkeletalMeshComponent* SkelComp = GetAnimPreviewScene().GetPreviewMeshComponent();
diff --git a/Engine/Source/Editor/AnimGraph/Private/EditModes/CCDIKEditMode.h b/Engine/Source/Editor/AnimGraph/Private/EditModes/CCDIKEditMode.h
--- a/Engine/Source/Editor/AnimGraph/Private/EditModes/CCDIKEditMode.h
+++ b/Engine/Source/Editor/AnimGraph/Private/EditModes/CCDIKEditMode.h
@@ -12,6 +12,7 @@ public:
 	virtual void ExitMode() override;
 	virtual FVector GetWidgetLocation() const override;
 	virtual UE::Widget::EWidgetMode GetWidgetMode() const override;
+	virtual bool UsesTransformWidget(UE::Widget::EWidgetMode InWidgetMode) const override;
 	virtual void DoTranslation(FVector& InTranslation) override;
 	virtual void Render(const FSceneView* View, FViewport* Viewport, FPrimitiveDrawInterface* PDI) override;
 private:
diff --git a/Engine/Source/Editor/AnimGraph/Private/EditModes/FabrikEditMode.cpp b/Engine/Source/Editor/AnimGraph/Private/EditModes/FabrikEditMode.cpp
--- a/Engine/Source/Editor/AnimGraph/Private/EditModes/FabrikEditMode.cpp
+++ b/Engine/Source/Editor/AnimGraph/Private/EditModes/FabrikEditMode.cpp
@@ -38,6 +38,11 @@ UE::Widget::EWidgetMode FFabrikEditMode::GetWidgetMode() const
 	return UE::Widget::WM_Translate;
 }
 
+bool FFabrikEditMode::UsesTransformWidget(UE::Widget::EWidgetMode InWidgetMode) const
+{
+	return InWidgetMode == UE::Widget::WM_Translate;
+}
+
 void FFabrikEditMode::DoTranslation(FVector& InTranslation)
 {
 	USkeletalMeshComponent* SkelComp = GetAnimPreviewScene().GetPreviewMeshComponent();
diff --git a/Engine/Source/Editor/AnimGraph/Private/EditModes/FabrikEditMode.h b/Engine/Source/Editor/AnimGraph/Private/EditModes/FabrikEditMode.h
--- a/Engine/Source/Editor/AnimGraph/Private/EditModes/FabrikEditMode.h
+++ b/Engine/Source/Editor/AnimGraph/Private/EditModes/FabrikEditMode.h
@@ -14,6 +14,7 @@ public:
 	virtual void ExitMode() override;
 	virtual FVector GetWidgetLocation() const override;
 	virtual UE::Widget::EWidgetMode GetWidgetMode() const override;
+	virtual bool UsesTransformWidget(UE::Widget::EWidgetMode InWidgetMode) const override;
 	virtual void DoTranslation(FVector& InTranslation) override;
 
 private:
diff --git a/Engine/Source/Editor/AnimGraph/Private/EditModes/LookAtEditMode.cpp b/Engine/Source/Editor/AnimGraph/Private/EditModes/LookAtEditMode.cpp
--- a/Engine/Source/Editor/AnimGraph/Private/EditModes/LookAtEditMode.cpp
+++ b/Engine/Source/Editor/AnimGraph/Private/EditModes/LookAtEditMode.cpp
@@ -39,6 +39,11 @@ UE::Widget::EWidgetMode FLookAtEditMode::GetWidgetMode() const
 	return UE::Widget::WM_Translate;
 }
 
+bool FLookAtEditMode::UsesTransformWidget(UE::Widget::EWidgetMode InWidgetMode) const
+{
+	return InWidgetMode == UE::Widget::WM_Translate;
+}
+
 FName FLookAtEditMode::GetSelectedBone() const
 {
 	return RuntimeNode->LookAtTarget.GetTargetSetup();
diff --git a/Engine/Source/Editor/AnimGraph/Private/EditModes/LookAtEditMode.h b/Engine/Source/Editor/AnimGraph/Private/EditModes/LookAtEditMode.h
--- a/Engine/Source/Editor/AnimGraph/Private/EditModes/LookAtEditMode.h
+++ b/Engine/Source/Editor/AnimGraph/Private/EditModes/LookAtEditMode.h
@@ -15,6 +15,7 @@ public:
 	virtual FVector GetWidgetLocation() const override;
 	virtual ECoordSystem GetWidgetCoordinateSystem() const override;
 	virtual UE::Widget::EWidgetMode GetWidgetMode() const override;
+	virtual bool UsesTransformWidget(UE::Widget::EWidgetMode InWidgetMode) const override;
 	virtual FName GetSelectedBone() const override;
 	virtual void DoTranslation(FVector& InTranslation) override;
 	virtual bool ShouldDrawWidget() const override;
diff --git a/Engine/Source/Editor/AnimGraph/Private/EditModes/ModifyBoneEditMode.cpp b/Engine/Source/Editor/AnimGraph/Private/EditModes/ModifyBoneEditMode.cpp
--- a/Engine/Source/Editor/AnimGraph/Private/EditModes/ModifyBoneEditMode.cpp
+++ b/Engine/Source/Editor/AnimGraph/Private/EditModes/ModifyBoneEditMode.cpp
@@ -206,6 +206,11 @@ bool FModifyBoneEditMode::SetWidgetMode(UE::Widget::EWidgetMode InWidgetMode)
 	return false;
 }
 
+bool FModifyBoneEditMode::UsesTransformWidget(UE::Widget::EWidgetMode InWidgetMode) const
+{
+	return FindValidWidgetMode(InWidgetMode) == InWidgetMode;
+}
+
 FName FModifyBoneEditMode::GetSelectedBone() const
 {
 	return GraphNode->Node.BoneToModify.BoneName;
diff --git a/Engine/Source/Editor/AnimGraph/Private/EditModes/ModifyBoneEditMode.h b/Engine/Source/Editor/AnimGraph/Private/EditModes/ModifyBoneEditMode.h
--- a/Engine/Source/Editor/AnimGraph/Private/EditModes/ModifyBoneEditMode.h
+++ b/Engine/Source/Editor/AnimGraph/Private/EditModes/ModifyBoneEditMode.h
@@ -20,6 +20,7 @@ public:
 	virtual UE::Widget::EWidgetMode GetWidgetMode() const override;
 	virtual UE::Widget::EWidgetMode ChangeToNextWidgetMode(UE::Widget::EWidgetMode InCurWidgetMode) override;
 	virtual bool SetWidgetMode(UE::Widget::EWidgetMode InWidgetMode) override;
+	virtual bool UsesTransformWidget(UE::Widget::EWidgetMode InWidgetMode) const override;
 	virtual FName GetSelectedBone() const override;
 	virtual void DoTranslation(FVector& InTranslation) override;
 	virtual void DoRotation(FRotator& InRotation) override;
diff --git a/Engine/Source/Editor/AnimGraph/Private/EditModes/ObserveBoneEditMode.cpp b/Engine/Source/Editor/AnimGraph/Private/EditModes/ObserveBoneEditMode.cpp
--- a/Engine/Source/Editor/AnimGraph/Private/EditModes/ObserveBoneEditMode.cpp
+++ b/Engine/Source/Editor/AnimGraph/Private/EditModes/ObserveBoneEditMode.cpp
@@ -59,6 +59,11 @@ UE::Widget::EWidgetMode FObserveBoneEditMode::GetWidgetMode() const
 	return UE::Widget::WM_Translate;
 }
 
+bool FObserveBoneEditMode::UsesTransformWidget(UE::Widget::EWidgetMode InWidgetMode) const
+{
+	return InWidgetMode == UE::Widget::WM_Translate;
+}
+
 FName FObserveBoneEditMode::GetSelectedBone() const
 {
 	return GraphNode->Node.BoneToObserve.BoneName;
diff --git a/Engine/Source/Editor/AnimGraph/Private/EditModes/ObserveBoneEditMode.h b/Engine/Source/Editor/AnimGraph/Private/EditModes/ObserveBoneEditMode.h
--- a/Engine/Source/Editor/AnimGraph/Private/EditModes/ObserveBoneEditMode.h
+++ b/Engine/Source/Editor/AnimGraph/Private/EditModes/ObserveBoneEditMode.h
@@ -15,6 +15,7 @@ public:
 	virtual ECoordSystem GetWidgetCoordinateSystem() const override;
 	virtual FVector GetWidgetLocation() const override;
 	virtual UE::Widget::EWidgetMode GetWidgetMode() const override;
+	virtual bool UsesTransformWidget(UE::Widget::EWidgetMode InWidgetMode) const override;
 	virtual FName GetSelectedBone() const override;
 
 private:
diff --git a/Engine/Source/Editor/AnimGraph/Private/EditModes/SplineIKEditMode.cpp b/Engine/Source/Editor/AnimGraph/Private/EditModes/SplineIKEditMode.cpp
--- a/Engine/Source/Editor/AnimGraph/Private/EditModes/SplineIKEditMode.cpp
+++ b/Engine/Source/Editor/AnimGraph/Private/EditModes/SplineIKEditMode.cpp
@@ -164,6 +164,11 @@ bool FSplineIKEditMode::SetWidgetMode(UE::Widget::EWidgetMode InWidgetMode)
 	return true;
 }
 
+bool FSplineIKEditMode::UsesTransformWidget(UE::Widget::EWidgetMode InWidgetMode) const
+{
+	return FindValidWidgetMode(InWidgetMode) == InWidgetMode;
+}
+
 FName FSplineIKEditMode::GetSelectedBone() const
 {
 	return NAME_None;
diff --git a/Engine/Source/Editor/AnimGraph/Private/EditModes/SplineIKEditMode.h b/Engine/Source/Editor/AnimGraph/Private/EditModes/SplineIKEditMode.h
--- a/Engine/Source/Editor/AnimGraph/Private/EditModes/SplineIKEditMode.h
+++ b/Engine/Source/Editor/AnimGraph/Private/EditModes/SplineIKEditMode.h
@@ -23,6 +23,7 @@ public:
 	virtual UE::Widget::EWidgetMode GetWidgetMode() const override;
 	virtual UE::Widget::EWidgetMode ChangeToNextWidgetMode(UE::Widget::EWidgetMode CurWidgetMode) override;
 	virtual bool SetWidgetMode(UE::Widget::EWidgetMode InWidgetMode) override;
+	virtual bool UsesTransformWidget(UE::Widget::EWidgetMode InWidgetMode) const override;
 	virtual FName GetSelectedBone() const override;
 	virtual bool GetCustomDrawingCoordinateSystem(FMatrix& InMatrix, void* InData) override;
 	virtual void DoTranslation(FVector& InTranslation) override;
diff --git a/Engine/Source/Editor/AnimGraph/Private/EditModes/TwoBoneIKEditMode.cpp b/Engine/Source/Editor/AnimGraph/Private/EditModes/TwoBoneIKEditMode.cpp
--- a/Engine/Source/Editor/AnimGraph/Private/EditModes/TwoBoneIKEditMode.cpp
+++ b/Engine/Source/Editor/AnimGraph/Private/EditModes/TwoBoneIKEditMode.cpp
@@ -139,6 +139,11 @@ UE::Widget::EWidgetMode FTwoBoneIKEditMode::GetWidgetMode() const
 	return UE::Widget::WM_None;
 }
 
+bool FTwoBoneIKEditMode::UsesTransformWidget(UE::Widget::EWidgetMode InWidgetMode) const
+{
+	return InWidgetMode == UE::Widget::WM_Translate;
+}
+
 FBoneSocketTarget FTwoBoneIKEditMode::GetSelectedTarget() const
 {
 	// should return mesh bone index
diff --git a/Engine/Source/Editor/AnimGraph/Private/EditModes/TwoBoneIKEditMode.h b/Engine/Source/Editor/AnimGraph/Private/EditModes/TwoBoneIKEditMode.h
--- a/Engine/Source/Editor/AnimGraph/Private/EditModes/TwoBoneIKEditMode.h
+++ b/Engine/Source/Editor/AnimGraph/Private/EditModes/TwoBoneIKEditMode.h
@@ -30,6 +30,7 @@ public:
 	virtual void ExitMode() override;
 	virtual FVector GetWidgetLocation() const override;
 	virtual UE::Widget::EWidgetMode GetWidgetMode() const override;
+	virtual bool UsesTransformWidget(UE::Widget::EWidgetMode InWidgetMode) const override;
 	FBoneSocketTarget GetSelectedTarget() const;
 	virtual void DoTranslation(FVector& InTranslation) override;
 
diff --git a/Engine/Source/Editor/Persona/Private/SAnimationEditorViewport.cpp b/Engine/Source/Editor/Persona/Private/SAnimationEditorViewport.cpp
--- a/Engine/Source/Editor/Persona/Private/SAnimationEditorViewport.cpp
+++ b/Engine/Source/Editor/Persona/Private/SAnimationEditorViewport.cpp
@@ -215,7 +215,18 @@ bool SAnimationEditorViewportTabBody::CanUseGizmos() const
 			return true;
 		}
 	}
-
+	
+	if (LevelViewportClient.IsValid())
+	{
+		if(const FEditorModeTools* ModeTools = LevelViewportClient->GetModeTools())
+		{
+			if(ModeTools->UsesTransformWidget())
+			{
+				return true;
+			}
+		}
+	}
+	
 	return false;
 }
 
