#UE4 Add more info to an assert

diff --git a/Engine/Source/Editor/AnimGraph/Private/AnimBlueprintCompiler.cpp b/Engine/Source/Editor/AnimGraph/Private/AnimBlueprintCompiler.cpp
--- a/Engine/Source/Editor/AnimGraph/Private/AnimBlueprintCompiler.cpp
+++ b/Engine/Source/Editor/AnimGraph/Private/AnimBlueprintCompiler.cpp
@@ -756,7 +756,7 @@ void FAnimBlueprintCompilerContext::CopyTermDefaultsToDefaultObject(UObject* Def
 		// Copy from root sparse class data to our new class
 		if(NewAnimBlueprintClass->GetConstantNodeData() && RootAnimClass->GetConstantNodeData())
 		{
-			check(NewAnimBlueprintClass->GetSparseClassDataStruct()->IsChildOf(RootAnimClass->GetSparseClassDataStruct()));
+			checkf(NewAnimBlueprintClass->GetSparseClassDataStruct()->IsChildOf(RootAnimClass->GetSparseClassDataStruct()), TEXT("SparseClassDataStruct for AnimBPClass '%s' (Struct:%s) is not a child of SparseClassDataStruct for AnimBPClass '%s' (Struct:%s)"), *GetPathNameSafe(NewAnimBlueprintClass), *GetFullNameSafe(NewAnimBlueprintClass->GetSparseClassDataStruct()), *GetPathNameSafe(RootAnimClass), *GetFullNameSafe(RootAnimClass->GetSparseClassDataStruct()));
 
 			for (TFieldIterator<FProperty> PropertyIt(RootAnimClass->GetSparseClassDataStruct()); PropertyIt; ++PropertyIt)
 			{
