Fixed double clicking an aim offset with exposed pins not working

diff --git a/Engine/Source/Editor/AnimGraph/Private/AnimGraphNode_RotationOffsetBlendSpace.cpp b/Engine/Source/Editor/AnimGraph/Private/AnimGraphNode_RotationOffsetBlendSpace.cpp
--- a/Engine/Source/Editor/AnimGraph/Private/AnimGraphNode_RotationOffsetBlendSpace.cpp
+++ b/Engine/Source/Editor/AnimGraph/Private/AnimGraphNode_RotationOffsetBlendSpace.cpp
@@ -329,5 +329,16 @@ void UAnimGraphNode_RotationOffsetBlendSpace::CustomizeDetails(IDetailLayoutBuil
 	}
 }
 
+UAnimationAsset* UAnimGraphNode_RotationOffsetBlendSpace::GetAnimationAsset() const 
+{
+	UBlendSpace* BlendSpace = Node.GetBlendSpace();
+	UEdGraphPin* BlendSpacePin = FindPin(GET_MEMBER_NAME_STRING_CHECKED(FAnimNode_RotationOffsetBlendSpace, BlendSpace));
+	if (BlendSpacePin != nullptr && BlendSpace == nullptr)
+	{
+		BlendSpace = Cast<UBlendSpace>(BlendSpacePin->DefaultObject);
+	}
+
+	return BlendSpace;
+}
 
 #undef LOCTEXT_NAMESPACE
diff --git a/Engine/Source/Editor/AnimGraph/Public/AnimGraphNode_RotationOffsetBlendSpace.h b/Engine/Source/Editor/AnimGraph/Public/AnimGraphNode_RotationOffsetBlendSpace.h
--- a/Engine/Source/Editor/AnimGraph/Public/AnimGraphNode_RotationOffsetBlendSpace.h
+++ b/Engine/Source/Editor/AnimGraph/Public/AnimGraphNode_RotationOffsetBlendSpace.h
@@ -27,7 +27,7 @@ class UAnimGraphNode_RotationOffsetBlendSpace: public UAnimGraphNode_BlendSpaceB
 
 	// UAnimGraphNode_Base interface
 	virtual void ValidateAnimNodeDuringCompilation(class USkeleton* ForSkeleton, class FCompilerResultsLog& MessageLog) override;
-	virtual UAnimationAsset* GetAnimationAsset() const override { return Node.GetBlendSpace(); }
+	virtual UAnimationAsset* GetAnimationAsset() const override;
 	virtual void GetAllAnimationSequencesReferred(TArray<UAnimationAsset*>& AnimationAssets) const override;
 	virtual void ReplaceReferredAnimations(const TMap<UAnimationAsset*, UAnimationAsset*>& AnimAssetReplacementMap) override;
 	virtual EAnimAssetHandlerType SupportsAssetClass(const UClass* AssetClass) const override;
