Added anim composites to anim BP context menu

diff --git a/Engine/Source/Editor/AnimGraph/Private/AnimGraphNode_SequencePlayer.cpp b/Engine/Source/Editor/AnimGraph/Private/AnimGraphNode_SequencePlayer.cpp
--- a/Engine/Source/Editor/AnimGraph/Private/AnimGraphNode_SequencePlayer.cpp
+++ b/Engine/Source/Editor/AnimGraph/Private/AnimGraphNode_SequencePlayer.cpp
@@ -90,7 +90,7 @@ void UAnimGraphNode_SequencePlayer::GetMenuActions(FBlueprintActionDatabaseRegis
 	GetMenuActionsHelper(
 		InActionRegistrar,
 		GetClass(),
-		{ UAnimSequence::StaticClass() },
+		{ UAnimSequence::StaticClass(), UAnimComposite::StaticClass() },
 		{ },
 		[](const FAssetData& InAssetData, UClass* InClass)
 		{
