Fixed linked anim layers showing up multiple times in context menu

Iterate over interfaces rather than ALL anim BPs, as well as accounting for self layers, then filter appropriately according to current BP

diff --git a/Engine/Source/Editor/AnimGraph/Private/AnimGraphNode_LinkedAnimLayer.cpp b/Engine/Source/Editor/AnimGraph/Private/AnimGraphNode_LinkedAnimLayer.cpp
--- a/Engine/Source/Editor/AnimGraph/Private/AnimGraphNode_LinkedAnimLayer.cpp
+++ b/Engine/Source/Editor/AnimGraph/Private/AnimGraphNode_LinkedAnimLayer.cpp
@@ -661,18 +661,18 @@ TSubclassOf<UInterface> UAnimGraphNode_LinkedAnimLayer::GetInterfaceForLayer() c
 {
 	if (UAnimBlueprint* CurrentBlueprint = Cast<UAnimBlueprint>(GetBlueprint()))
 	{
-			// Find layer with this name in interfaces
-			for(FBPInterfaceDescription& InterfaceDesc : CurrentBlueprint->ImplementedInterfaces)
+		// Find layer with this name in interfaces
+		for(FBPInterfaceDescription& InterfaceDesc : CurrentBlueprint->ImplementedInterfaces)
+		{
+			for(UEdGraph* InterfaceGraph : InterfaceDesc.Graphs)
 			{
-				for(UEdGraph* InterfaceGraph : InterfaceDesc.Graphs)
+				if(InterfaceGraph->GetFName() == Node.Layer)
 				{
-					if(InterfaceGraph->GetFName() == Node.Layer)
-					{
-						return InterfaceDesc.Interface;
-					}
+					return InterfaceDesc.Interface;
 				}
 			}
 		}
+	}
 
 	return nullptr;
 }
@@ -829,17 +829,57 @@ void UAnimGraphNode_LinkedAnimLayer::GetMenuActions(FBlueprintActionDatabaseRegi
 			UClass* TargetClass = *TargetAnimBlueprint->SkeletonGeneratedClass;
 			if(TargetClass)
 			{
-				UClass* NodeClass = GetClass();
-				IAnimClassInterface* AnimClassInterface = IAnimClassInterface::GetFromClass(TargetClass);
-				for(const FAnimBlueprintFunction& AnimBlueprintFunction : AnimClassInterface->GetAnimBlueprintFunctions())
+				// Accept interfaces
+				if(TargetAnimBlueprint->BlueprintType == BPTYPE_Interface)
+				{
+					IAnimClassInterface* AnimClassInterface = IAnimClassInterface::GetFromClass(TargetClass);
+					for(const FAnimBlueprintFunction& AnimBlueprintFunction : AnimClassInterface->GetAnimBlueprintFunctions())
+					{
+						if(AnimBlueprintFunction.Name != UEdGraphSchema_K2::GN_AnimGraph)
+						{
+							if(UFunction* Function = TargetClass->FindFunctionByName(AnimBlueprintFunction.Name))
+							{
+								if (UBlueprintNodeSpawner* NodeSpawner = MakeAnimBlueprintAction(GetClass(), AnimBlueprintFunction.Name))
+								{
+									ActionRegistrar.AddBlueprintAction(Function, NodeSpawner);
+								}
+							}
+						}
+					}
+				}
+				else
 				{
-					if(AnimBlueprintFunction.Name != UEdGraphSchema_K2::GN_AnimGraph)
+					// Accept 'self' layers
+					IAnimClassInterface* AnimClassInterface = IAnimClassInterface::GetFromClass(TargetClass);
+					for(const FAnimBlueprintFunction& AnimBlueprintFunction : AnimClassInterface->GetAnimBlueprintFunctions())
 					{
-						if(UFunction* Function = TargetClass->FindFunctionByName(AnimBlueprintFunction.Name))
+						if(AnimBlueprintFunction.Name != UEdGraphSchema_K2::GN_AnimGraph)
 						{
-							if (UBlueprintNodeSpawner* NodeSpawner = MakeAnimBlueprintAction(NodeClass, AnimBlueprintFunction.Name))
+							const bool bIsSelfLayer = [TargetAnimBlueprint, &AnimBlueprintFunction]()
+							{
+								for(const FBPInterfaceDescription& InterfaceDesc : TargetAnimBlueprint->ImplementedInterfaces)
+								{
+									for(UEdGraph* InterfaceGraph : InterfaceDesc.Graphs)
+									{
+										if(InterfaceGraph->GetFName() == AnimBlueprintFunction.Name)
+										{
+											return false;
+										}
+									}
+								}
+
+								return true;
+							}();
+
+							if(bIsSelfLayer)
 							{
-								ActionRegistrar.AddBlueprintAction(Function, NodeSpawner);
+								if(UFunction* Function = TargetClass->FindFunctionByName(AnimBlueprintFunction.Name))
+								{
+									if (UBlueprintNodeSpawner* NodeSpawner = MakeAnimBlueprintAction(GetClass(), AnimBlueprintFunction.Name))
+									{
+										ActionRegistrar.AddBlueprintAction(Function, NodeSpawner);
+									}
+								}
 							}
 						}
 					}
@@ -849,4 +889,39 @@ void UAnimGraphNode_LinkedAnimLayer::GetMenuActions(FBlueprintActionDatabaseRegi
 	}
 }
 
+bool UAnimGraphNode_LinkedAnimLayer::IsActionFilteredOut(class FBlueprintActionFilter const& Filter)
+{
+	bool bIsFilteredOut = false;
+
+	FBlueprintActionContext const& FilterContext = Filter.Context;
+
+	for (UBlueprint* Blueprint : FilterContext.Blueprints)
+	{
+		if (UAnimBlueprint* AnimBlueprint = Cast<UAnimBlueprint>(Blueprint))
+		{
+			if(UClass* TargetClass = *AnimBlueprint->SkeletonGeneratedClass)
+			{
+				// Accept only functions contained in this BP
+				bool bImplemented = false;
+				IAnimClassInterface* AnimClassInterface = IAnimClassInterface::GetFromClass(TargetClass);
+				for(const FAnimBlueprintFunction& AnimBlueprintFunction : AnimClassInterface->GetAnimBlueprintFunctions())
+				{
+					if(Node.Layer == AnimBlueprintFunction.Name)
+					{
+						bImplemented = true;
+						break;
+					}
+				}
+
+				if(!bImplemented)
+				{
+					bIsFilteredOut = true;
+				}
+			}
+		}
+	}
+	
+	return bIsFilteredOut;
+}
+
 #undef LOCTEXT_NAMESPACE
diff --git a/Engine/Source/Editor/AnimGraph/Public/AnimGraphNode_LinkedAnimLayer.h b/Engine/Source/Editor/AnimGraph/Public/AnimGraphNode_LinkedAnimLayer.h
--- a/Engine/Source/Editor/AnimGraph/Public/AnimGraphNode_LinkedAnimLayer.h
+++ b/Engine/Source/Editor/AnimGraph/Public/AnimGraphNode_LinkedAnimLayer.h
@@ -45,6 +45,7 @@ public:
 
 	// UK2Node interface
 	virtual void GetMenuActions(FBlueprintActionDatabaseRegistrar& ActionRegistrar) const override;
+	virtual bool IsActionFilteredOut(class FBlueprintActionFilter const& Filter) override;
 	
 	// UAnimGraphNode_Base interface
 	virtual void CustomizeDetails(IDetailLayoutBuilder& DetailBuilder) override;
