Fixed pose watches not removing on graph deletion

diff --git a/Engine/Source/Editor/AnimGraph/Private/AnimGraphNode_BlendSpaceGraphBase.cpp b/Engine/Source/Editor/AnimGraph/Private/AnimGraphNode_BlendSpaceGraphBase.cpp
--- a/Engine/Source/Editor/AnimGraph/Private/AnimGraphNode_BlendSpaceGraphBase.cpp
+++ b/Engine/Source/Editor/AnimGraph/Private/AnimGraphNode_BlendSpaceGraphBase.cpp
@@ -263,6 +263,13 @@ TArray<UEdGraph*> UAnimGraphNode_BlendSpaceGraphBase::GetSubGraphs() const
 
 void UAnimGraphNode_BlendSpaceGraphBase::DestroyNode()
 {
+	UBlueprint* Blueprint = GetBlueprint();
+	for (UEdGraph* SampleGraph : Graphs)
+	{
+		SampleGraph->Modify();
+		FBlueprintEditorUtils::RemoveGraph(Blueprint, SampleGraph);
+	}
+
 	UEdGraph* GraphToRemove = BlendSpaceGraph;
 	BlendSpaceGraph = nullptr;
 	Graphs.Empty();
@@ -271,7 +278,6 @@ void UAnimGraphNode_BlendSpaceGraphBase::DestroyNode()
 
 	if (GraphToRemove)
 	{
-		UBlueprint* Blueprint = GetBlueprint();
 		GraphToRemove->Modify();
 		FBlueprintEditorUtils::RemoveGraph(Blueprint, GraphToRemove, EGraphRemoveFlags::Recompile);
 	}
diff --git a/Engine/Source/Editor/AnimGraph/Private/AnimationGraphSchema.cpp b/Engine/Source/Editor/AnimGraph/Private/AnimationGraphSchema.cpp
--- a/Engine/Source/Editor/AnimGraph/Private/AnimationGraphSchema.cpp
+++ b/Engine/Source/Editor/AnimGraph/Private/AnimationGraphSchema.cpp
@@ -39,6 +39,7 @@
 #include "AnimGraphNode_RigidBody.h"
 #include "AnimationBlendSpaceSampleGraph.h"
 #include "GraphEditorDragDropAction.h"
+#include "AnimationEditorUtils.h"
 
 #define LOCTEXT_NAMESPACE "AnimationGraphSchema"
 
@@ -234,6 +235,12 @@ void UAnimationGraphSchema::HandleGraphBeingDeleted(UEdGraph& GraphBeingRemoved)
 			// Prevent re-entrancy here
 			NodeToDelete->ClearBoundGraph();
 		}
+
+		// Remove pose watches from nodes in this graph
+		if (UAnimBlueprint* AnimBlueprint = Cast<UAnimBlueprint>(Blueprint))
+		{
+			AnimationEditorUtils::RemovePoseWatchesFromGraph(AnimBlueprint, &GraphBeingRemoved);
+		}
 	}
 }
 
diff --git a/Engine/Source/Editor/UnrealEd/Private/AnimationEditorUtils.cpp b/Engine/Source/Editor/UnrealEd/Private/AnimationEditorUtils.cpp
--- a/Engine/Source/Editor/UnrealEd/Private/AnimationEditorUtils.cpp
+++ b/Engine/Source/Editor/UnrealEd/Private/AnimationEditorUtils.cpp
@@ -861,6 +861,16 @@ namespace AnimationEditorUtils
 #endif
 	}
 
+	void RemovePoseWatchesFromGraph(UAnimBlueprint* AnimBlueprint, class UEdGraph* Graph)
+	{
+#if WITH_EDITORONLY_DATA
+		for (UEdGraphNode* Node : Graph->Nodes)
+		{
+			RemovePoseWatchFromNode(Node, AnimBlueprint);
+		}
+#endif
+	}
+
 	UPoseWatch* FindPoseWatchForNode(const UEdGraphNode* Node, UAnimBlueprint* AnimBlueprintIfKnown)
 	{
 #if WITH_EDITORONLY_DATA
diff --git a/Engine/Source/Editor/UnrealEd/Public/AnimationEditorUtils.h b/Engine/Source/Editor/UnrealEd/Public/AnimationEditorUtils.h
--- a/Engine/Source/Editor/UnrealEd/Public/AnimationEditorUtils.h
+++ b/Engine/Source/Editor/UnrealEd/Public/AnimationEditorUtils.h
@@ -200,6 +200,7 @@ namespace AnimationEditorUtils
 	UNREALED_API UPoseWatch* MakePoseWatchForNode(UAnimBlueprint* AnimBlueprint, UEdGraphNode* Node);
 	UNREALED_API void RemovePoseWatch(UPoseWatch* PoseWatch, UAnimBlueprint* AnimBlueprintIfKnown=nullptr);
 	UNREALED_API void RemovePoseWatchFromNode(UEdGraphNode* Node, UAnimBlueprint* AnimBlueprint);
+	UNREALED_API void RemovePoseWatchesFromGraph(UAnimBlueprint* AnimBlueprint, class UEdGraph* Graph);
 
 	// Delegate fired when a pose watch is added or removed
 	DECLARE_MULTICAST_DELEGATE_TwoParams(FOnPoseWatchesChanged, UAnimBlueprint* /*InAnimBlueprint*/, UEdGraphNode* /*InNode*/);
