Allow child anim blueprints to be used as the default instance class in LinkedAnimLayers

diff --git a/Engine/Source/Editor/AnimGraph/Private/AnimGraphNode_LinkedAnimLayer.cpp b/Engine/Source/Editor/AnimGraph/Private/AnimGraphNode_LinkedAnimLayer.cpp
--- a/Engine/Source/Editor/AnimGraph/Private/AnimGraphNode_LinkedAnimLayer.cpp
+++ b/Engine/Source/Editor/AnimGraph/Private/AnimGraphNode_LinkedAnimLayer.cpp
@@ -444,9 +444,9 @@ void UAnimGraphNode_LinkedAnimLayer::CustomizeDetails(IDetailLayoutBuilder& Deta
 	}
 }
 
-bool UAnimGraphNode_LinkedAnimLayer::OnShouldFilterInstanceBlueprint(const FAssetData& AssetData) const
+bool UAnimGraphNode_LinkedAnimLayer::OnShouldFilterInstanceBlueprint(const FAssetData& InAssetData) const
 {
-	if(Super::OnShouldFilterInstanceBlueprint(AssetData))
+	if(Super::OnShouldFilterInstanceBlueprint(InAssetData))
 	{
 		return true;
 	}
@@ -468,37 +468,53 @@ bool UAnimGraphNode_LinkedAnimLayer::OnShouldFilterInstanceBlueprint(const FAsse
 		// Check interface compatibility
 		if(AnimInterfaces.Num() > 0)
 		{
+			const FAssetRegistryModule& AssetRegistryModule = FModuleManager::LoadModuleChecked<FAssetRegistryModule>("AssetRegistry");
+			FAssetData CurrentAssetData = InAssetData;
 			bool bMatchesInterface = false;
 
-			const FString ImplementedInterfaces = AssetData.GetTagValueRef<FString>(FBlueprintTags::ImplementedInterfaces);
-			if(!ImplementedInterfaces.IsEmpty())
+			do 
 			{
-				FString FullInterface;
-				FString RemainingString;
-				FString InterfacePath;
-				FString CurrentString = *ImplementedInterfaces;
-				while(CurrentString.Split(TEXT(","), &FullInterface, &RemainingString) && !bMatchesInterface)
+				const FString ImplementedInterfaces = CurrentAssetData.GetTagValueRef<FString>(FBlueprintTags::ImplementedInterfaces);
+
+				if(!ImplementedInterfaces.IsEmpty())
 				{
-					if (!CurrentString.StartsWith(TEXT("Graphs=(")))
+					FString FullInterface;
+					FString RemainingString;
+					FString InterfacePath;
+					FString CurrentString = *ImplementedInterfaces;
+					while(CurrentString.Split(TEXT(","), &FullInterface, &RemainingString) && !bMatchesInterface)
 					{
-						if (FullInterface.Split(TEXT("\""), &CurrentString, &InterfacePath, ESearchCase::CaseSensitive))
+						if (!CurrentString.StartsWith(TEXT("Graphs=(")))
 						{
-							// The interface paths in metadata end with "', so remove those
-							InterfacePath.RemoveFromEnd(TEXT("\"'"));
+							if (FullInterface.Split(TEXT("\""), &CurrentString, &InterfacePath, ESearchCase::CaseSensitive))
+							{
+								// The interface paths in metadata end with "', so remove those
+								InterfacePath.RemoveFromEnd(TEXT("\"'"));
 
-							FCoreRedirectObjectName ResolvedInterfaceName = FCoreRedirects::GetRedirectedName(ECoreRedirectFlags::Type_Class, FCoreRedirectObjectName(InterfacePath));
+								FCoreRedirectObjectName ResolvedInterfaceName = FCoreRedirects::GetRedirectedName(ECoreRedirectFlags::Type_Class, FCoreRedirectObjectName(InterfacePath));
 
-							// Verify against all interfaces we currently implement
-							for(TSubclassOf<UInterface> AnimInterface : AnimInterfaces)
-							{
-								bMatchesInterface |= ResolvedInterfaceName.ObjectName == AnimInterface->GetFName();
+								// Verify against all interfaces we currently implement
+								for(TSubclassOf<UInterface> AnimInterface : AnimInterfaces)
+								{
+									bMatchesInterface |= ResolvedInterfaceName.ObjectName == AnimInterface->GetFName();
+								}
 							}
 						}
-					}
 			
-					CurrentString = RemainingString;
+						CurrentString = RemainingString;
+					}
 				}
-			}
+
+				// If we didn't find a matching interface, check the parent class
+				if (!bMatchesInterface)
+				{
+					const FString ParentClassFromData = CurrentAssetData.GetTagValueRef<FString>(FBlueprintTags::ParentClassPath);
+					const FString ClassObjectPath = FPackageName::ExportTextPathToObjectPath(ParentClassFromData);
+					const FName BlueprintPath = FName(*ClassObjectPath.LeftChop(2)); // Chop off _C
+					CurrentAssetData = AssetRegistryModule.Get().GetAssetByObjectPath(BlueprintPath);
+				}
+			// Only continue checking if the parent is an anim blueprint
+			} while (!bMatchesInterface && (CurrentAssetData.AssetClass == UAnimBlueprint::StaticClass()->GetFName()));
 
 			if(!bMatchesInterface)
 			{
