Fixed being unable to override aim offset assets

A previous refactor missed this file

diff --git a/Engine/Source/Editor/AnimGraph/Private/AnimGraphNode_RotationOffsetBlendSpace.cpp b/Engine/Source/Editor/AnimGraph/Private/AnimGraphNode_RotationOffsetBlendSpace.cpp
--- a/Engine/Source/Editor/AnimGraph/Private/AnimGraphNode_RotationOffsetBlendSpace.cpp
+++ b/Engine/Source/Editor/AnimGraph/Private/AnimGraphNode_RotationOffsetBlendSpace.cpp
@@ -17,6 +17,7 @@
 #include "AnimGraphCommands.h"
 #include "BlueprintNodeTemplateCache.h"
 #include "AssetRegistry/AssetRegistryModule.h"
+#include "IAnimBlueprintNodeOverrideAssetsContext.h"
 
 /////////////////////////////////////////////////////
 // UAnimGraphNode_RotationOffsetBlendSpace
@@ -121,6 +122,18 @@ void UAnimGraphNode_RotationOffsetBlendSpace::SetAnimationAsset(UAnimationAsset*
 	}
 }
 
+void UAnimGraphNode_RotationOffsetBlendSpace::OnOverrideAssets(IAnimBlueprintNodeOverrideAssetsContext& InContext) const
+{
+	if(InContext.GetAssets().Num() > 0)
+	{
+		if (UBlendSpace* BlendSpace = Cast<UBlendSpace>(InContext.GetAssets()[0]))
+		{
+			FAnimNode_BlendSpacePlayer& AnimNode = InContext.GetAnimNode<FAnimNode_BlendSpacePlayer>();
+			AnimNode.SetBlendSpace(BlendSpace);
+		}
+	}
+}
+
 void UAnimGraphNode_RotationOffsetBlendSpace::ValidateAnimNodeDuringCompilation(class USkeleton* ForSkeleton, class FCompilerResultsLog& MessageLog)
 {
 	Super::ValidateAnimNodeDuringCompilation(ForSkeleton, MessageLog);
diff --git a/Engine/Source/Editor/AnimGraph/Public/AnimGraphNode_RotationOffsetBlendSpace.h b/Engine/Source/Editor/AnimGraph/Public/AnimGraphNode_RotationOffsetBlendSpace.h
--- a/Engine/Source/Editor/AnimGraph/Public/AnimGraphNode_RotationOffsetBlendSpace.h
+++ b/Engine/Source/Editor/AnimGraph/Public/AnimGraphNode_RotationOffsetBlendSpace.h
@@ -28,6 +28,7 @@ class UAnimGraphNode_RotationOffsetBlendSpace: public UAnimGraphNode_BlendSpaceB
 	// UAnimGraphNode_Base interface
 	virtual void ValidateAnimNodeDuringCompilation(class USkeleton* ForSkeleton, class FCompilerResultsLog& MessageLog) override;
 	virtual UAnimationAsset* GetAnimationAsset() const override;
+	virtual void OnOverrideAssets(IAnimBlueprintNodeOverrideAssetsContext& InContext) const override;
 	virtual void GetAllAnimationSequencesReferred(TArray<UAnimationAsset*>& AnimationAssets) const override;
 	virtual void ReplaceReferredAnimations(const TMap<UAnimationAsset*, UAnimationAsset*>& AnimAssetReplacementMap) override;
 	virtual EAnimAssetHandlerType SupportsAssetClass(const UClass* AssetClass) const override;
