#UE Fix for crash overwriting NewAnimBlueprintConstants object during load

diff --git a/Engine/Source/Editor/AnimGraph/Private/AnimBlueprintCompiler.cpp b/Engine/Source/Editor/AnimGraph/Private/AnimBlueprintCompiler.cpp
--- a/Engine/Source/Editor/AnimGraph/Private/AnimBlueprintCompiler.cpp
+++ b/Engine/Source/Editor/AnimGraph/Private/AnimBlueprintCompiler.cpp
@@ -1267,7 +1267,11 @@ void FAnimBlueprintCompilerContext::RecreateSparseClassData()
 		check(ParentAnimClass->GetSparseClassDataStruct());
 		
 		// Derive sparse class data from parent class
-		NewAnimBlueprintConstants = NewObject<UScriptStruct>(NewAnimBlueprintClass, UAnimBlueprintGeneratedClass::GetConstantsStructName(), RF_Public);
+		NewAnimBlueprintConstants = FindObject<UScriptStruct>(NewAnimBlueprintClass, *UAnimBlueprintGeneratedClass::GetConstantsStructName().ToString());
+		if (!NewAnimBlueprintConstants)
+		{
+			NewAnimBlueprintConstants = NewObject<UScriptStruct>(NewAnimBlueprintClass, UAnimBlueprintGeneratedClass::GetConstantsStructName(), RF_Public);
+		}
 		NewAnimBlueprintConstants->SetSuperStruct(ParentAnimClass->GetSparseClassDataStruct());
 
 		// Just link & assign sparse class data struct here, no additional members are added
