[FIX] Attempt to fix build error on UnrealEditor NonUnity Win64 (Strict, VS2019) by adding missing include.

#rmx

diff --git a/Engine/Source/Editor/AnimGraph/Private/EditModes/AnimDynamicsEditMode.cpp b/Engine/Source/Editor/AnimGraph/Private/EditModes/AnimDynamicsEditMode.cpp
--- a/Engine/Source/Editor/AnimGraph/Private/EditModes/AnimDynamicsEditMode.cpp
+++ b/Engine/Source/Editor/AnimGraph/Private/EditModes/AnimDynamicsEditMode.cpp
@@ -5,6 +5,7 @@
 #include "IPersonaPreviewScene.h"
 #include "Animation/DebugSkelMeshComponent.h"
 #include "EditorViewportClient.h"
+#include "Materials/MaterialInstanceDynamic.h"
 
 // AnimDynamicsEditMode hit proxy
 //
