Merge fix for missing expliict cast.

diff --git a/Engine/Source/Editor/AnimGraph/Private/EditModes/AnimDynamicsEditMode.cpp b/Engine/Source/Editor/AnimGraph/Private/EditModes/AnimDynamicsEditMode.cpp
--- a/Engine/Source/Editor/AnimGraph/Private/EditModes/AnimDynamicsEditMode.cpp
+++ b/Engine/Source/Editor/AnimGraph/Private/EditModes/AnimDynamicsEditMode.cpp
@@ -313,7 +313,7 @@ void FAnimDynamicsEditMode::Render(const FSceneView* View, FViewport* Viewport,
 						LimitPlaneTransform *= PreviewSkelMeshComp->GetComponentSpaceTransforms()[LimitDrivingBoneIdx];
 					}
 
-					const UE::Math::TMatrix< float > PlaneTransform = LimitPlaneTransform.ToMatrixNoScale();
+					const FMatrix PlaneTransform = LimitPlaneTransform.ToMatrixNoScale();
 
 					const FAnimDynamicsViewportObjectReference ViewportObjectRef(EditorAnimDynamicsNode->GetUniqueID(), FAnimDynamicsViewportObjectType::PlaneLimit, LimitIndex);
 					const bool bIsSelected = SelectedViewportObjects.Contains(ViewportObjectRef);
