Fix autowire for non-pose pins on anim nodes

diff --git a/Engine/Source/Editor/AnimGraph/Private/AnimGraphNode_Base.cpp b/Engine/Source/Editor/AnimGraph/Private/AnimGraphNode_Base.cpp
--- a/Engine/Source/Editor/AnimGraph/Private/AnimGraphNode_Base.cpp
+++ b/Engine/Source/Editor/AnimGraph/Private/AnimGraphNode_Base.cpp
@@ -922,6 +922,10 @@ void UAnimGraphNode_Base::AutowireNewNode(UEdGraphPin* FromPin)
 			ensure(GetSchema()->TryCreateConnection(FromPin, PinToConnectTo));
 		}
 	}
+	else
+	{
+		UK2Node::AutowireNewNode(FromPin);
+	}
 }
 
 TSharedRef<SWidget> UAnimGraphNode_Base::MakePropertyBindingWidget(const FAnimPropertyBindingWidgetArgs& InArgs)
