[UE][FIX] Static analysis warning on missing null check

#rmx

diff --git a/Engine/Source/Editor/AnimGraph/Private/EditModes/AnimDynamicsEditMode.cpp b/Engine/Source/Editor/AnimGraph/Private/EditModes/AnimDynamicsEditMode.cpp
--- a/Engine/Source/Editor/AnimGraph/Private/EditModes/AnimDynamicsEditMode.cpp
+++ b/Engine/Source/Editor/AnimGraph/Private/EditModes/AnimDynamicsEditMode.cpp
@@ -661,39 +661,40 @@ void FAnimDynamicsEditMode::DoTranslation(FVector& InTranslation)
 		for (const FAnimDynamicsViewportObjectReference& SelectedObjectRef : SelectedViewportObjects)
 		{
 			UAnimGraphNode_AnimDynamics* const EditorAnimDynamicsNode = FindSelectedEditorAnimNode(SelectedObjectRef.EditorNodeUniqueId);
-			FAnimNode_AnimDynamics* const RuntimeAnimDynamicsNode = EditorAnimDynamicsNode->GetPreviewDynamicsNode();
-
-			check(EditorAnimDynamicsNode && RuntimeAnimDynamicsNode);
-
-			if (EditorAnimDynamicsNode && RuntimeAnimDynamicsNode)
+			if (EditorAnimDynamicsNode)
 			{
-				if ((SelectedObjectRef.Type == FAnimDynamicsViewportObjectType::BoxExtents) || (SelectedObjectRef.Type == FAnimDynamicsViewportObjectType::SphericalColisionVolume))
+				FAnimNode_AnimDynamics* const RuntimeAnimDynamicsNode = EditorAnimDynamicsNode->GetPreviewDynamicsNode();
+				if (RuntimeAnimDynamicsNode)
 				{
-					if (RuntimeAnimDynamicsNode->PhysicsBodyDefinitions.IsValidIndex(SelectedObjectRef.Index))
+
+					if ((SelectedObjectRef.Type == FAnimDynamicsViewportObjectType::BoxExtents) || (SelectedObjectRef.Type == FAnimDynamicsViewportObjectType::SphericalColisionVolume))
 					{
-						FVector& PhysicsBodyJointOffset = RuntimeAnimDynamicsNode->PhysicsBodyDefinitions[SelectedObjectRef.Index].LocalJointOffset;
+						if (RuntimeAnimDynamicsNode->PhysicsBodyDefinitions.IsValidIndex(SelectedObjectRef.Index))
+						{
+							FVector& PhysicsBodyJointOffset = RuntimeAnimDynamicsNode->PhysicsBodyDefinitions[SelectedObjectRef.Index].LocalJointOffset;
 
-						PhysicsBodyJointOffset += LocalSpaceTranslation;
-						EditorAnimDynamicsNode->Node.PhysicsBodyDefinitions[SelectedObjectRef.Index].LocalJointOffset = PhysicsBodyJointOffset;
+							PhysicsBodyJointOffset += LocalSpaceTranslation;
+							EditorAnimDynamicsNode->Node.PhysicsBodyDefinitions[SelectedObjectRef.Index].LocalJointOffset = PhysicsBodyJointOffset;
+						}
 					}
-				}
-				else if (SelectedObjectRef.Type == FAnimDynamicsViewportObjectType::PlaneLimit)
-				{
-					if (RuntimeAnimDynamicsNode->PlanarLimits.IsValidIndex(SelectedObjectRef.Index))
+					else if (SelectedObjectRef.Type == FAnimDynamicsViewportObjectType::PlaneLimit)
 					{
-						FTransform& LimitTransform = RuntimeAnimDynamicsNode->PlanarLimits[SelectedObjectRef.Index].PlaneTransform;
-						LimitTransform.SetTranslation(LimitTransform.GetTranslation() + LocalSpaceTranslation);
-						EditorAnimDynamicsNode->Node.PlanarLimits[SelectedObjectRef.Index].PlaneTransform = LimitTransform;
+						if (RuntimeAnimDynamicsNode->PlanarLimits.IsValidIndex(SelectedObjectRef.Index))
+						{
+							FTransform& LimitTransform = RuntimeAnimDynamicsNode->PlanarLimits[SelectedObjectRef.Index].PlaneTransform;
+							LimitTransform.SetTranslation(LimitTransform.GetTranslation() + LocalSpaceTranslation);
+							EditorAnimDynamicsNode->Node.PlanarLimits[SelectedObjectRef.Index].PlaneTransform = LimitTransform;
+						}
 					}
-				}
-				else if (SelectedObjectRef.Type == FAnimDynamicsViewportObjectType::SphericalLimit)
-				{
-					if (RuntimeAnimDynamicsNode->SphericalLimits.IsValidIndex(SelectedObjectRef.Index))
+					else if (SelectedObjectRef.Type == FAnimDynamicsViewportObjectType::SphericalLimit)
 					{
+						if (RuntimeAnimDynamicsNode->SphericalLimits.IsValidIndex(SelectedObjectRef.Index))
+						{
 
-						FVector& LimitOffset = RuntimeAnimDynamicsNode->SphericalLimits[SelectedObjectRef.Index].SphereLocalOffset;
-						LimitOffset += LocalSpaceTranslation;
-						EditorAnimDynamicsNode->Node.SphericalLimits[SelectedObjectRef.Index].SphereLocalOffset = LimitOffset;
+							FVector& LimitOffset = RuntimeAnimDynamicsNode->SphericalLimits[SelectedObjectRef.Index].SphereLocalOffset;
+							LimitOffset += LocalSpaceTranslation;
+							EditorAnimDynamicsNode->Node.SphericalLimits[SelectedObjectRef.Index].SphereLocalOffset = LimitOffset;
+						}
 					}
 				}
 			}
@@ -727,18 +728,19 @@ void FAnimDynamicsEditMode::DoRotation(FRotator& InRotation)
 		for (const FAnimDynamicsViewportObjectReference& SelectedObjectRef : SelectedViewportObjects)
 		{
 			UAnimGraphNode_AnimDynamics* const EditorAnimDynamicsNode = FindSelectedEditorAnimNode(SelectedObjectRef.EditorNodeUniqueId);
-			FAnimNode_AnimDynamics* const RuntimeAnimDynamicsNode = EditorAnimDynamicsNode->GetPreviewDynamicsNode();
-
-			check(EditorAnimDynamicsNode && RuntimeAnimDynamicsNode);
-			if (EditorAnimDynamicsNode && RuntimeAnimDynamicsNode)
+			if (EditorAnimDynamicsNode)
 			{
-				if (SelectedObjectRef.Type == FAnimDynamicsViewportObjectType::PlaneLimit)
+				FAnimNode_AnimDynamics* const RuntimeAnimDynamicsNode = EditorAnimDynamicsNode->GetPreviewDynamicsNode();
+				if (RuntimeAnimDynamicsNode)
 				{
-					if (RuntimeAnimDynamicsNode->PlanarLimits.IsValidIndex(SelectedObjectRef.Index))
+					if (SelectedObjectRef.Type == FAnimDynamicsViewportObjectType::PlaneLimit)
 					{
-						FTransform& LimitTransform = RuntimeAnimDynamicsNode->PlanarLimits[SelectedObjectRef.Index].PlaneTransform;
-						LimitTransform.SetRotation(LocalSpaceRotation * LimitTransform.GetRotation());
-						EditorAnimDynamicsNode->Node.PlanarLimits[SelectedObjectRef.Index].PlaneTransform = LimitTransform;
+						if (RuntimeAnimDynamicsNode->PlanarLimits.IsValidIndex(SelectedObjectRef.Index))
+						{
+							FTransform& LimitTransform = RuntimeAnimDynamicsNode->PlanarLimits[SelectedObjectRef.Index].PlaneTransform;
+							LimitTransform.SetRotation(LocalSpaceRotation * LimitTransform.GetRotation());
+							EditorAnimDynamicsNode->Node.PlanarLimits[SelectedObjectRef.Index].PlaneTransform = LimitTransform;
+						}
 					}
 				}
 			}
@@ -753,39 +755,39 @@ void FAnimDynamicsEditMode::DoScale(FVector& InScale)
 	for (const FAnimDynamicsViewportObjectReference& SelectedObjectRef : SelectedViewportObjects)
 	{
 		UAnimGraphNode_AnimDynamics* const EditorAnimDynamicsNode = FindSelectedEditorAnimNode(SelectedObjectRef.EditorNodeUniqueId);
-		FAnimNode_AnimDynamics* const RuntimeAnimDynamicsNode = EditorAnimDynamicsNode->GetPreviewDynamicsNode();
-
-		check(EditorAnimDynamicsNode && RuntimeAnimDynamicsNode);
-
-		if (EditorAnimDynamicsNode && RuntimeAnimDynamicsNode)
+		if (EditorAnimDynamicsNode)
 		{
-			if (SelectedObjectRef.Type == FAnimDynamicsViewportObjectType::BoxExtents)
+			FAnimNode_AnimDynamics* const RuntimeAnimDynamicsNode = EditorAnimDynamicsNode->GetPreviewDynamicsNode();
+			if (RuntimeAnimDynamicsNode)
 			{
-				if (RuntimeAnimDynamicsNode->PhysicsBodyDefinitions.IsValidIndex(SelectedObjectRef.Index))
+				if (SelectedObjectRef.Type == FAnimDynamicsViewportObjectType::BoxExtents)
 				{
-					FVector& PhysicsBodyBoxExtents = RuntimeAnimDynamicsNode->PhysicsBodyDefinitions[SelectedObjectRef.Index].BoxExtents;
-					PhysicsBodyBoxExtents += InScale;
-					EditorAnimDynamicsNode->Node.PhysicsBodyDefinitions[SelectedObjectRef.Index].BoxExtents = PhysicsBodyBoxExtents;
+					if (RuntimeAnimDynamicsNode->PhysicsBodyDefinitions.IsValidIndex(SelectedObjectRef.Index))
+					{
+						FVector& PhysicsBodyBoxExtents = RuntimeAnimDynamicsNode->PhysicsBodyDefinitions[SelectedObjectRef.Index].BoxExtents;
+						PhysicsBodyBoxExtents += InScale;
+						EditorAnimDynamicsNode->Node.PhysicsBodyDefinitions[SelectedObjectRef.Index].BoxExtents = PhysicsBodyBoxExtents;
+					}
 				}
-			}
-			else if (SelectedObjectRef.Type == FAnimDynamicsViewportObjectType::SphericalColisionVolume)
-			{
-				if (RuntimeAnimDynamicsNode->PhysicsBodyDefinitions.IsValidIndex(SelectedObjectRef.Index))
+				else if (SelectedObjectRef.Type == FAnimDynamicsViewportObjectType::SphericalColisionVolume)
 				{
-					float& PhysicsBodySphereCollisionRadius = RuntimeAnimDynamicsNode->PhysicsBodyDefinitions[SelectedObjectRef.Index].SphereCollisionRadius;
+					if (RuntimeAnimDynamicsNode->PhysicsBodyDefinitions.IsValidIndex(SelectedObjectRef.Index))
+					{
+						float& PhysicsBodySphereCollisionRadius = RuntimeAnimDynamicsNode->PhysicsBodyDefinitions[SelectedObjectRef.Index].SphereCollisionRadius;
 
-					PhysicsBodySphereCollisionRadius += Scale;
-					EditorAnimDynamicsNode->Node.PhysicsBodyDefinitions[SelectedObjectRef.Index].SphereCollisionRadius = PhysicsBodySphereCollisionRadius;
+						PhysicsBodySphereCollisionRadius += Scale;
+						EditorAnimDynamicsNode->Node.PhysicsBodyDefinitions[SelectedObjectRef.Index].SphereCollisionRadius = PhysicsBodySphereCollisionRadius;
+					}
 				}
-			}
-			else if (SelectedObjectRef.Type == FAnimDynamicsViewportObjectType::SphericalLimit)
-			{
-				if (RuntimeAnimDynamicsNode->SphericalLimits.IsValidIndex(SelectedObjectRef.Index))
+				else if (SelectedObjectRef.Type == FAnimDynamicsViewportObjectType::SphericalLimit)
 				{
-					const float LimitRadius = RuntimeAnimDynamicsNode->SphericalLimits[SelectedObjectRef.Index].LimitRadius + Scale;
+					if (RuntimeAnimDynamicsNode->SphericalLimits.IsValidIndex(SelectedObjectRef.Index))
+					{
+						const float LimitRadius = RuntimeAnimDynamicsNode->SphericalLimits[SelectedObjectRef.Index].LimitRadius + Scale;
 
-					RuntimeAnimDynamicsNode->SphericalLimits[SelectedObjectRef.Index].LimitRadius = LimitRadius;
-					EditorAnimDynamicsNode->Node.SphericalLimits[SelectedObjectRef.Index].LimitRadius = LimitRadius;
+						RuntimeAnimDynamicsNode->SphericalLimits[SelectedObjectRef.Index].LimitRadius = LimitRadius;
+						EditorAnimDynamicsNode->Node.SphericalLimits[SelectedObjectRef.Index].LimitRadius = LimitRadius;
+					}
 				}
 			}
 		}
