Fix issues with child anim BP compilation ordering & sparse class data fixup

This partially ports a fix from CL 18740288, but improves on when and how sparse class data gets fixed up.
This allows this fix to cover more parent->child load/compile orderings (the previous fix only worked if they were all compiled in the same batch, which isnt always guaranteed.

Original CL desc:

Ensures that cleaned-up sparse class data structs are moved out of the package they are in, thus not requiring a GC step to ensure that batch compilation works correctly.
Without this, it is possible for the NewObject call in FAnimBlueprintCompilerContext::RecreateSparseClassData to attempt to re-use (as it has a fixed name) the currently-existing (but not referenced) object rather than creating a new one.
This also reverts some other changes that were made to try to address related issues.

diff --git a/Engine/Source/Editor/AnimGraph/Private/AnimBlueprintCompiler.cpp b/Engine/Source/Editor/AnimGraph/Private/AnimBlueprintCompiler.cpp
--- a/Engine/Source/Editor/AnimGraph/Private/AnimBlueprintCompiler.cpp
+++ b/Engine/Source/Editor/AnimGraph/Private/AnimBlueprintCompiler.cpp
@@ -1338,10 +1338,7 @@ void FAnimBlueprintCompilerContext::SpawnNewClass(const FString& NewClassName)
 void FAnimBlueprintCompilerContext::OnPostCDOCompiled()
 {
 	UAnimBlueprintGeneratedClass* NewAnimBlueprintClass = GetNewAnimBlueprintClass();
-	for (UAnimBlueprintGeneratedClass* ClassWithInputHandlers = NewAnimBlueprintClass; ClassWithInputHandlers != nullptr; ClassWithInputHandlers = Cast<UAnimBlueprintGeneratedClass>(ClassWithInputHandlers->GetSuperClass()))
-	{
-		ClassWithInputHandlers->OnPostLoadDefaults(NewAnimBlueprintClass->ClassDefaultObject);
-	}
+	NewAnimBlueprintClass->OnPostLoadDefaults(NewAnimBlueprintClass->ClassDefaultObject);
 }
 
 void FAnimBlueprintCompilerContext::OnNewClassSet(UBlueprintGeneratedClass* ClassToUse)
@@ -1351,10 +1348,30 @@ void FAnimBlueprintCompilerContext::OnNewClassSet(UBlueprintGeneratedClass* Clas
 
 void FAnimBlueprintCompilerContext::CleanAndSanitizeClass(UBlueprintGeneratedClass* ClassToClean, UObject*& InOldCDO)
 {
-	Super::CleanAndSanitizeClass(ClassToClean, InOldCDO);
+	UAnimBlueprintGeneratedClass* AnimBlueprintClassToClean = CastChecked<UAnimBlueprintGeneratedClass>(ClassToClean);
 
-	UAnimBlueprintGeneratedClass* AnimBlueprintClassToClean= CastChecked<UAnimBlueprintGeneratedClass>(ClassToClean);
+	UScriptStruct* CurrentSparseClassDataStruct = AnimBlueprintClassToClean->GetSparseClassDataStruct();
+	if(CurrentSparseClassDataStruct && CurrentSparseClassDataStruct->GetOuter() == AnimBlueprintClassToClean)
+	{
+		// Stash a reference to patch the linker later in RecreateSparseClassData
+		OldSparseClassDataStruct = CurrentSparseClassDataStruct;
 
+		// Only 'clear' (which renames the struct aside) if we own the sparse class data
+		// We do this because this could be a parent classes struct and we dont want to be 
+		// altering other classes during compilation
+		AnimBlueprintClassToClean->ClearSparseClassDataStruct(Blueprint->bIsRegeneratingOnLoad);
+	}
+	else
+	{
+		AnimBlueprintClassToClean->SetSparseClassDataStruct(nullptr);
+	}
+	
+	// Calling super will set this classes superclass, which will reset its sparse class data to the parent (archetype)
+	Super::CleanAndSanitizeClass(ClassToClean, InOldCDO);
+
+	// Clear reference to archetype sparse class data (set in the super call above), we will be regenerating it
+	AnimBlueprintClassToClean->SetSparseClassDataStruct(nullptr);
+	
 	AnimBlueprintClassToClean->AnimBlueprintDebugData = FAnimBlueprintDebugData();
 
 	// Reset the baked data
@@ -1380,10 +1397,6 @@ void FAnimBlueprintCompilerContext::CleanAndSanitizeClass(UBlueprintGeneratedCla
 
 	UAnimBlueprint* RootAnimBP = UAnimBlueprint::FindRootAnimBlueprint(AnimBlueprint);
 	bIsDerivedAnimBlueprint = RootAnimBP != NULL;
-
-	// Cleanup sparse class data & stash a reference to patch the linker later in RecreateSparseClassData
-	OldSparseClassDataStruct = AnimBlueprintClassToClean->GetSparseClassDataStruct();
-	AnimBlueprintClassToClean->ClearSparseClassDataStruct();
 	
 	FAnimBlueprintGeneratedClassCompiledData CompiledData(AnimBlueprintClassToClean);
 	FAnimBlueprintCompilationBracketContext CompilerContext(this);
diff --git a/Engine/Source/Runtime/CoreUObject/Private/UObject/Class.cpp b/Engine/Source/Runtime/CoreUObject/Private/UObject/Class.cpp
--- a/Engine/Source/Runtime/CoreUObject/Private/UObject/Class.cpp
+++ b/Engine/Source/Runtime/CoreUObject/Private/UObject/Class.cpp
@@ -5260,10 +5260,6 @@ void UClass::SetSparseClassDataStruct(UScriptStruct* InSparseClassDataStruct)
 					SubClass->CleanupSparseClassData();
 					SubClassSparseClassDataStruct->SetSuperStruct(InSparseClassDataStruct);
 				}
-				else if (SubClassSparseClassDataStruct == SparseClassDataStruct)
-				{
-					SubClass->SetSparseClassDataStruct(InSparseClassDataStruct);
-				}
 			}
 		}
 		// the old type and new type may not match when we do a reload so get rid of the old data
@@ -5273,7 +5269,7 @@ void UClass::SetSparseClassDataStruct(UScriptStruct* InSparseClassDataStruct)
 	}
 }
 
-void UClass::ClearSparseClassDataStruct()
+void UClass::ClearSparseClassDataStruct(bool bInRecomplingOnLoad)
 { 
 	if (SparseClassDataStruct != nullptr)
 	{
@@ -5290,17 +5286,22 @@ void UClass::ClearSparseClassDataStruct()
 			}
 		}
 
+		auto ClearSparseClassDataStructInner = [bInRecomplingOnLoad](UClass* InClassToClear)
+		{
+			UScriptStruct* CurrentSparseClassDataStruct = InClassToClear->GetSparseClassDataStruct();
+			InClassToClear->CleanupSparseClassData();
+			InClassToClear->SparseClassDataStruct = nullptr;
+			CurrentSparseClassDataStruct->SetSuperStruct(nullptr);
+			const ERenameFlags RenameFlags = REN_DontCreateRedirectors | ((bInRecomplingOnLoad) ? REN_ForceNoResetLoaders : 0) | REN_NonTransactional | REN_DoNotDirty;
+			CurrentSparseClassDataStruct->Rename(nullptr, GetTransientPackage(), RenameFlags);
+		};
+
 		for(UClass* SubClassToClear : SubClassesToClear)
 		{
-			UScriptStruct* SubClassSparseClassDataStruct = SubClassToClear->GetSparseClassDataStruct();
-			SubClassToClear->CleanupSparseClassData();
-			SubClassToClear->SparseClassDataStruct = nullptr;
-			SubClassSparseClassDataStruct->SetSuperStruct(nullptr);	
+			ClearSparseClassDataStructInner(SubClassToClear);
 		}
 
-		CleanupSparseClassData();
-
-		SparseClassDataStruct = nullptr;
+		ClearSparseClassDataStructInner(this);
 	}
 }
 
diff --git a/Engine/Source/Runtime/CoreUObject/Public/UObject/Class.h b/Engine/Source/Runtime/CoreUObject/Public/UObject/Class.h
--- a/Engine/Source/Runtime/CoreUObject/Public/UObject/Class.h
+++ b/Engine/Source/Runtime/CoreUObject/Public/UObject/Class.h
@@ -2835,8 +2835,11 @@ public:
 
 	void SetSparseClassDataStruct(UScriptStruct* InSparseClassDataStruct);
 
-	/** Clears the sparse class data struct for this and all child classes that directly reference it as a super-struct */
-	void ClearSparseClassDataStruct();
+	/** 
+	 * Clears the sparse class data struct for this and all child classes that directly reference it as a super-struct 
+	 * This will rename the current sparse class data struct aside into the transient package
+	 */
+	void ClearSparseClassDataStruct(bool bInRecomplingOnLoad);
 
 	/** Assemble reference token streams for all classes if they haven't had it assembled already */
 	static void AssembleReferenceTokenStreams();
