Fix "additional meshes" in persona preview not animating

diff --git a/Engine/Source/Editor/AnimGraph/Private/AnimPreviewAttacheInstance.cpp b/Engine/Source/Editor/AnimGraph/Private/AnimPreviewAttacheInstance.cpp
--- a/Engine/Source/Editor/AnimGraph/Private/AnimPreviewAttacheInstance.cpp
+++ b/Engine/Source/Editor/AnimGraph/Private/AnimPreviewAttacheInstance.cpp
@@ -18,20 +18,14 @@ void FAnimPreviewAttacheInstanceProxy::PreUpdate(UAnimInstance* InAnimInstance,
 {
 	FAnimInstanceProxy::PreUpdate(InAnimInstance, DeltaSeconds);
 
-	if (CopyPoseFromMesh.SourceMeshComponent.IsValid())
-	{
-		CopyPoseFromMesh.PreUpdate(InAnimInstance);
-	}
+	CopyPoseFromMesh.PreUpdate(InAnimInstance);
 }
 
 void FAnimPreviewAttacheInstanceProxy::UpdateAnimationNode(const FAnimationUpdateContext& InContext)
 {
 	UpdateCounter.Increment();
 	
-	if (CopyPoseFromMesh.SourceMeshComponent.IsValid())
-	{
-		CopyPoseFromMesh.Update_AnyThread(InContext);
-	}
+	CopyPoseFromMesh.Update_AnyThread(InContext);
 }
 
 bool FAnimPreviewAttacheInstanceProxy::Evaluate(FPoseContext& Output)
