[UE][FIX] AnimDynamics UI - Ensure bone name updates correctly in details panel array element headers.

#rmx

diff --git a/Engine/Source/Editor/AnimGraph/Private/AnimGraphNode_AnimDynamics.cpp b/Engine/Source/Editor/AnimGraph/Private/AnimGraphNode_AnimDynamics.cpp
--- a/Engine/Source/Editor/AnimGraph/Private/AnimGraphNode_AnimDynamics.cpp
+++ b/Engine/Source/Editor/AnimGraph/Private/AnimGraphNode_AnimDynamics.cpp
@@ -22,21 +22,6 @@
 
 #define LOCTEXT_NAMESPACE "AnimDynamicsNode"
 
-////////////////////////////////////////
-// class FAnimGraphNode_AnimDynamics_DetailCustomization
-//
-// Required to customise the chain body array rendered in the details panel.
-//
-class FAnimGraphNode_AnimDynamics_DetailCustomization : public IDetailCustomization
-{
-public:
-	static TSharedRef< IDetailCustomization > MakeInstance();
-	virtual void CustomizeDetails(IDetailLayoutBuilder& DetailBuilder) override;
-	void OnPhysicsBodyDefCustomizeDetails(TSharedRef<class IPropertyHandle> ElementProperty, int32 ElementIndex, class IDetailChildrenBuilder& ChildrenBuilder, class IDetailLayoutBuilder* DetailLayout);
-
-	UClass* OwningClass;
-};
-
 ////////////////////////////////////////
 // class UAnimGraphNode_AnimDynamics
 
@@ -90,13 +75,15 @@ void UAnimGraphNode_AnimDynamics::CustomizeDetails(IDetailLayoutBuilder& DetailB
 			.OnClicked(FOnClicked::CreateStatic(&UAnimGraphNode_AnimDynamics::ResetButtonClicked, &DetailBuilder))
 		];
 
-	if (!DetailCustomization)
+	// Customise the physics body definition array rendered in the details panel.
+	IDetailCategoryBuilder& PhysicsParametersCategory = DetailBuilder.EditCategory(TEXT("PhysicsParameters"));
+	TSharedRef< IPropertyHandle > PhysicsBodyDefinitionsProperty = DetailBuilder.GetProperty("Node.PhysicsBodyDefinitions", GetClass());
+	if (PhysicsBodyDefinitionsProperty->AsArray().IsValid())
 	{
-		DetailCustomization = MakeShared<FAnimGraphNode_AnimDynamics_DetailCustomization>();
-	}
-
-	DetailCustomization->OwningClass = GetClass();
-	DetailCustomization->CustomizeDetails(DetailBuilder);
+		TSharedRef<FDetailArrayBuilder> PropertyBuilder = MakeShared<FDetailArrayBuilder>(PhysicsBodyDefinitionsProperty, /*InGenerateHeader*/ true, /*InDisplayResetToDefault*/ true, /*InDisplayElementNum*/ true);
+		PropertyBuilder->OnGenerateArrayElementWidget(FOnGenerateArrayElementWidget::CreateUObject(this, &UAnimGraphNode_AnimDynamics::OnPhysicsBodyDefCustomizeDetails, &DetailBuilder));
+		PhysicsParametersCategory.AddCustomBuilder(PropertyBuilder, false);
+	}	
 
 	// Force order of details panel catagories - Must set order for all of them as any that are edited automatically move to the top.
 	{
@@ -289,24 +276,7 @@ void UAnimGraphNode_AnimDynamics::Serialize(FArchive& Ar)
 	}
 }
 
-TSharedRef< IDetailCustomization > FAnimGraphNode_AnimDynamics_DetailCustomization::MakeInstance()
-{
-	return MakeShareable(new FAnimGraphNode_AnimDynamics_DetailCustomization);
-}
-
-void FAnimGraphNode_AnimDynamics_DetailCustomization::CustomizeDetails(IDetailLayoutBuilder& DetailLayout)
-{
-	IDetailCategoryBuilder& SetupCategory = DetailLayout.EditCategory(TEXT("PhysicsParameters"));
-	TSharedRef< IPropertyHandle > PhysicsBodyDefinitionsProperty = DetailLayout.GetProperty("Node.PhysicsBodyDefinitions", OwningClass);
-	if (PhysicsBodyDefinitionsProperty->AsArray().IsValid())
-	{
-		TSharedRef<FDetailArrayBuilder> PropertyBuilder = MakeShareable(new FDetailArrayBuilder(PhysicsBodyDefinitionsProperty, /*InGenerateHeader*/ true, /*InDisplayResetToDefault*/ true, /*InDisplayElementNum*/ true));
-		PropertyBuilder->OnGenerateArrayElementWidget(FOnGenerateArrayElementWidget::CreateSP(this, &FAnimGraphNode_AnimDynamics_DetailCustomization::OnPhysicsBodyDefCustomizeDetails, &DetailLayout));
-		SetupCategory.AddCustomBuilder(PropertyBuilder, false);
-	}
-}
-
-void FAnimGraphNode_AnimDynamics_DetailCustomization::OnPhysicsBodyDefCustomizeDetails(TSharedRef<IPropertyHandle> ElementProperty, int32 ElementIndex, IDetailChildrenBuilder& ChildrenBuilder, IDetailLayoutBuilder* DetailLayout)
+void UAnimGraphNode_AnimDynamics::OnPhysicsBodyDefCustomizeDetails(TSharedRef<IPropertyHandle> ElementProperty, int32 ElementIndex, IDetailChildrenBuilder& ChildrenBuilder, IDetailLayoutBuilder* DetailLayout)
 {
 	TSharedPtr<IPropertyHandle> BoundBoneProperty = ElementProperty->GetChildHandle(GET_MEMBER_NAME_CHECKED(FAnimPhysBodyDefinition, BoundBone));
 
@@ -323,16 +293,26 @@ void FAnimGraphNode_AnimDynamics_DetailCustomization::OnPhysicsBodyDefCustomizeD
 
 		DetailLayout->HideProperty(BoundBoneProperty);
 		IDetailPropertyRow& Row = ChildrenBuilder.AddProperty(ElementProperty);
-		
-		// Set a custom widget show a more useful item name and remove the 'n items' text that would otherwise appear on the body def group header.
-		const FString BodyDefNameString = "[" + FString::FromInt(ElementIndex) + "] " + BoneName.ToString();
-		Row.CustomWidget(true) 
+
+		// Set a custom widget to show a more useful item name and remove the 'n items' text that would otherwise appear on the body def group header		
+		Row.CustomWidget(true)
 			.NameContent()
 			[
 				SNew(STextBlock)
-				.Text(FText::FromString(BodyDefNameString))
+				.Text_Lambda([this, ElementIndex]() -> FText { return this->BodyDefinitionUINameString(ElementIndex); })
 			];
 	}
 }
 
+FText UAnimGraphNode_AnimDynamics::BodyDefinitionUINameString(const uint32 BodyIndex) const
+{
+	if (Node.PhysicsBodyDefinitions.IsValidIndex(BodyIndex))
+	{
+		const FName BoneName = Node.PhysicsBodyDefinitions[BodyIndex].BoundBone.BoneName;
+		return FText::Format(INVTEXT("[{0}] {1}"), FText::AsNumber(BodyIndex), FText::FromName(BoneName));
+	}
+
+	return FText::GetEmpty();
+}
+
 #undef LOCTEXT_NAMESPACE
diff --git a/Engine/Source/Editor/AnimGraph/Public/AnimGraphNode_AnimDynamics.h b/Engine/Source/Editor/AnimGraph/Public/AnimGraphNode_AnimDynamics.h
--- a/Engine/Source/Editor/AnimGraph/Public/AnimGraphNode_AnimDynamics.h
+++ b/Engine/Source/Editor/AnimGraph/Public/AnimGraphNode_AnimDynamics.h
@@ -86,6 +86,10 @@ public:
 	FAnimNode_AnimDynamics* GetPreviewDynamicsNode() const;
 	bool IsPreviewLiveActive() const { return bPreviewLive;  }
 
+	// UI callbacks for customising the physics bodies array in the details panel.
+	void OnPhysicsBodyDefCustomizeDetails(TSharedRef<IPropertyHandle> ElementProperty, int32 ElementIndex, IDetailChildrenBuilder& ChildrenBuilder, IDetailLayoutBuilder* DetailLayout);
+	FText BodyDefinitionUINameString(const uint32 BodyIndex) const;
+
 protected:
 
 	virtual FText GetControllerDescription() const override;
@@ -102,7 +106,6 @@ protected:
 	USkeleton* GetSkeleton() const;
 
 private:
-	TSharedPtr<class FAnimGraphNode_AnimDynamics_DetailCustomization> DetailCustomization;
 
 	FNodeTitleTextTable CachedNodeTitles;
 };
