fixed crash when validating nodes: template anim BPs don't have any skeleton assigned

diff --git a/Engine/Plugins/Animation/AnimationWarping/Source/Editor/Private/AnimGraph/AnimGraphNode_OrientationWarping.cpp b/Engine/Plugins/Animation/AnimationWarping/Source/Editor/Private/AnimGraph/AnimGraphNode_OrientationWarping.cpp
--- a/Engine/Plugins/Animation/AnimationWarping/Source/Editor/Private/AnimGraph/AnimGraphNode_OrientationWarping.cpp
+++ b/Engine/Plugins/Animation/AnimationWarping/Source/Editor/Private/AnimGraph/AnimGraphNode_OrientationWarping.cpp
@@ -188,7 +188,7 @@ void UAnimGraphNode_OrientationWarping::ValidateAnimNodeDuringCompilation(USkele
 
 	auto HasInvalidBoneIndex = [&](const FName& BoneName)
 	{
-		return ForSkeleton->GetReferenceSkeleton().FindBoneIndex(BoneName) == INDEX_NONE;
+		return ForSkeleton && ForSkeleton->GetReferenceSkeleton().FindBoneIndex(BoneName) == INDEX_NONE;
 	};
 
 	auto InvalidBoneNameMessage = [&](const FName& BoneName)
diff --git a/Engine/Plugins/Animation/AnimationWarping/Source/Editor/Private/AnimGraph/AnimGraphNode_StrideWarping.cpp b/Engine/Plugins/Animation/AnimationWarping/Source/Editor/Private/AnimGraph/AnimGraphNode_StrideWarping.cpp
--- a/Engine/Plugins/Animation/AnimationWarping/Source/Editor/Private/AnimGraph/AnimGraphNode_StrideWarping.cpp
+++ b/Engine/Plugins/Animation/AnimationWarping/Source/Editor/Private/AnimGraph/AnimGraphNode_StrideWarping.cpp
@@ -220,7 +220,7 @@ void UAnimGraphNode_StrideWarping::ValidateAnimNodeDuringCompilation(USkeleton*
 
 	auto HasInvalidBoneIndex = [&] (const FName& BoneName) 
 	{ 
-		return ForSkeleton->GetReferenceSkeleton().FindBoneIndex(BoneName) == INDEX_NONE; 
+		return ForSkeleton && ForSkeleton->GetReferenceSkeleton().FindBoneIndex(BoneName) == INDEX_NONE; 
 	};
 
 	auto InvalidBoneNameMessage = [&](const FName& BoneName) 
diff --git a/Engine/Plugins/Animation/ControlRig/Source/ControlRigDeveloper/Private/AnimGraphNode_ControlRig.cpp b/Engine/Plugins/Animation/ControlRig/Source/ControlRigDeveloper/Private/AnimGraphNode_ControlRig.cpp
--- a/Engine/Plugins/Animation/ControlRig/Source/ControlRigDeveloper/Private/AnimGraphNode_ControlRig.cpp
+++ b/Engine/Plugins/Animation/ControlRig/Source/ControlRigDeveloper/Private/AnimGraphNode_ControlRig.cpp
@@ -257,26 +257,29 @@ void UAnimGraphNode_ControlRig::ValidateAnimNodeDuringCompilation(USkeleton* For
 		if (UControlRigBlueprint* Blueprint = Cast<UControlRigBlueprint>(TargetClass->ClassGeneratedBy))
 		{
 			URigHierarchy* Hierarchy = Blueprint->Hierarchy;
-			const FReferenceSkeleton& ReferenceSkeleton = ForSkeleton->GetReferenceSkeleton();
-			const TArray<FMeshBoneInfo>& BoneInfos = ReferenceSkeleton.GetRefBoneInfo();
-
-			for (const FMeshBoneInfo& BoneInfo : BoneInfos)
+			if (ForSkeleton)
 			{
-				const FRigElementKey BoneKey(BoneInfo.Name, ERigElementType::Bone);
-				if (Hierarchy->Contains(BoneKey))
-				{
-					const FRigElementKey ParentKey = Hierarchy->GetFirstParent(BoneKey);
+				const FReferenceSkeleton& ReferenceSkeleton = ForSkeleton->GetReferenceSkeleton();
+				const TArray<FMeshBoneInfo>& BoneInfos = ReferenceSkeleton.GetRefBoneInfo();
 
-					FName DesiredParentName = NAME_None;
-					if (BoneInfo.ParentIndex != INDEX_NONE)
+				for (const FMeshBoneInfo& BoneInfo : BoneInfos)
+				{
+					const FRigElementKey BoneKey(BoneInfo.Name, ERigElementType::Bone);
+					if (Hierarchy->Contains(BoneKey))
 					{
-						DesiredParentName = BoneInfos[BoneInfo.ParentIndex].Name;
-					}
+						const FRigElementKey ParentKey = Hierarchy->GetFirstParent(BoneKey);
 
-					if (DesiredParentName != ParentKey.Name)
-					{
-						FString Message = FString::Printf(TEXT("@@ - Hierarchy discrepancy for bone '%s' - different parents on Control Rig vs SkeletalMesh."), *BoneInfo.Name.ToString());
-						MessageLog.Warning(*Message, this);
+						FName DesiredParentName = NAME_None;
+						if (BoneInfo.ParentIndex != INDEX_NONE)
+						{
+							DesiredParentName = BoneInfos[BoneInfo.ParentIndex].Name;
+						}
+
+						if (DesiredParentName != ParentKey.Name)
+						{
+							FString Message = FString::Printf(TEXT("@@ - Hierarchy discrepancy for bone '%s' - different parents on Control Rig vs SkeletalMesh."), *BoneInfo.Name.ToString());
+							MessageLog.Warning(*Message, this);
+						}
 					}
 				}
 			}
diff --git a/Engine/Plugins/Animation/IKRig/Source/IKRigDeveloper/Private/AnimGraphNode_RetargetPoseFromMesh.cpp b/Engine/Plugins/Animation/IKRig/Source/IKRigDeveloper/Private/AnimGraphNode_RetargetPoseFromMesh.cpp
--- a/Engine/Plugins/Animation/IKRig/Source/IKRigDeveloper/Private/AnimGraphNode_RetargetPoseFromMesh.cpp
+++ b/Engine/Plugins/Animation/IKRig/Source/IKRigDeveloper/Private/AnimGraphNode_RetargetPoseFromMesh.cpp
@@ -77,22 +77,25 @@ void UAnimGraphNode_RetargetPoseFromMesh::ValidateAnimNodeDuringCompilation(USke
 	{
 		return;
 	}
-	
-	// validate that target bone chains exist on this skeleton
-	const FReferenceSkeleton &RefSkel = ForSkeleton->GetReferenceSkeleton();
-	const TArray<FBoneChain> &TargetBoneChains = Node.IKRetargeterAsset->GetTargetIKRig()->GetRetargetChains();
-    for (const FBoneChain &Chain : TargetBoneChains)
-    {
-        if (RefSkel.FindBoneIndex(Chain.StartBone.BoneName) == INDEX_NONE)
-        {
-        	MessageLog.Warning(*LOCTEXT("StartBoneNotFound", "@@ - Start Bone in target IK Rig Bone Chain not found.").ToString(), this);
-        }
-
-    	if (RefSkel.FindBoneIndex(Chain.EndBone.BoneName) == INDEX_NONE)
-    	{
-    		MessageLog.Warning(*LOCTEXT("EndBoneNotFound", "@@ - End Bone in target IK Rig Bone Chain not found.").ToString(), this);
-    	}
-    }
+
+	if (ForSkeleton)
+	{
+		// validate that target bone chains exist on this skeleton
+		const FReferenceSkeleton &RefSkel = ForSkeleton->GetReferenceSkeleton();
+		const TArray<FBoneChain> &TargetBoneChains = Node.IKRetargeterAsset->GetTargetIKRig()->GetRetargetChains();
+		for (const FBoneChain &Chain : TargetBoneChains)
+		{
+			if (RefSkel.FindBoneIndex(Chain.StartBone.BoneName) == INDEX_NONE)
+			{
+				MessageLog.Warning(*LOCTEXT("StartBoneNotFound", "@@ - Start Bone in target IK Rig Bone Chain not found.").ToString(), this);
+			}
+
+			if (RefSkel.FindBoneIndex(Chain.EndBone.BoneName) == INDEX_NONE)
+			{
+				MessageLog.Warning(*LOCTEXT("EndBoneNotFound", "@@ - End Bone in target IK Rig Bone Chain not found.").ToString(), this);
+			}
+		}
+	}
 }
 
 void UAnimGraphNode_RetargetPoseFromMesh::PreloadRequiredAssets()
diff --git a/Engine/Source/Editor/AnimGraph/Private/AnimGraphNode_AimOffsetLookAt.cpp b/Engine/Source/Editor/AnimGraph/Private/AnimGraphNode_AimOffsetLookAt.cpp
--- a/Engine/Source/Editor/AnimGraph/Private/AnimGraphNode_AimOffsetLookAt.cpp
+++ b/Engine/Source/Editor/AnimGraph/Private/AnimGraphNode_AimOffsetLookAt.cpp
@@ -173,30 +173,32 @@ void UAnimGraphNode_AimOffsetLookAt::ValidateAnimNodeDuringCompilation(class USk
 	else
 	{
 		const USkeleton* BlendSpaceSkeleton = BlendSpaceToCheck->GetSkeleton();
-		if (BlendSpaceSkeleton && // if blend space doesn't have skeleton, it might be due to blend space not loaded yet, @todo: wait with anim blueprint compilation until all assets are loaded?
-			!ForSkeleton->IsCompatible(BlendSpaceSkeleton))
+		if (BlendSpaceSkeleton) // if blend space doesn't have skeleton, it might be due to blend space not loaded yet, @todo: wait with anim blueprint compilation until all assets are loaded?
 		{
-			MessageLog.Error(TEXT("@@ references blendspace that uses an incompatible skeleton @@"), this, BlendSpaceSkeleton);
-		}
+			if (ForSkeleton && !ForSkeleton->IsCompatible(BlendSpaceSkeleton))
+			{
+				MessageLog.Error(TEXT("@@ references blendspace that uses an incompatible skeleton @@"), this, BlendSpaceSkeleton);
+			}
 
-		// Make sure that the source socket name is a valid one for the skeleton
-		UEdGraphPin* SocketNamePin = FindPin(GET_MEMBER_NAME_STRING_CHECKED(FAnimNode_AimOffsetLookAt, SourceSocketName));
-		FName SocketNameToCheck = (SocketNamePin != nullptr) ? FName(*SocketNamePin->DefaultValue) : Node.SourceSocketName;
+			// Make sure that the source socket name is a valid one for the skeleton
+			UEdGraphPin* SocketNamePin = FindPin(GET_MEMBER_NAME_STRING_CHECKED(FAnimNode_AimOffsetLookAt, SourceSocketName));
+			FName SocketNameToCheck = (SocketNamePin != nullptr) ? FName(*SocketNamePin->DefaultValue) : Node.SourceSocketName;
 
-		// Temporary fix where skeleton is not fully loaded during AnimBP compilation and thus the socket name check is invalid UE-39499 (NEED FIX) 
-		if (BlendSpaceSkeleton && !BlendSpaceSkeleton->HasAnyFlags(RF_NeedPostLoad))
-		{
-			const bool bValidValue = SocketNamePin == nullptr && BlendSpaceSkeleton->FindSocket(Node.SourceSocketName);
-			const bool bValidPinValue = SocketNamePin != nullptr && BlendSpaceSkeleton->FindSocket(FName(*SocketNamePin->DefaultValue));
-			const bool bValidConnectedPin = SocketNamePin != nullptr && SocketNamePin->LinkedTo.Num();
+			// Temporary fix where skeleton is not fully loaded during AnimBP compilation and thus the socket name check is invalid UE-39499 (NEED FIX) 
+			if (!BlendSpaceSkeleton->HasAnyFlags(RF_NeedPostLoad))
+			{
+				const bool bValidValue = SocketNamePin == nullptr && BlendSpaceSkeleton->FindSocket(Node.SourceSocketName);
+				const bool bValidPinValue = SocketNamePin != nullptr && BlendSpaceSkeleton->FindSocket(FName(*SocketNamePin->DefaultValue));
+				const bool bValidConnectedPin = SocketNamePin != nullptr && SocketNamePin->LinkedTo.Num();
 
-			if (!bValidValue && !bValidPinValue && !bValidConnectedPin)
-			{ 
-				FFormatNamedArguments Args;
-				Args.Add(TEXT("SocketName"), FText::FromName(SocketNameToCheck));
+				if (!bValidValue && !bValidPinValue && !bValidConnectedPin)
+				{ 
+					FFormatNamedArguments Args;
+					Args.Add(TEXT("SocketName"), FText::FromName(SocketNameToCheck));
 
-				const FText Msg = FText::Format(LOCTEXT("SocketNameNotFound", "@@ - Socket {SocketName} not found in Skeleton"), Args);
-				MessageLog.Error(*Msg.ToString(), this);
+					const FText Msg = FText::Format(LOCTEXT("SocketNameNotFound", "@@ - Socket {SocketName} not found in Skeleton"), Args);
+					MessageLog.Error(*Msg.ToString(), this);
+				}
 			}
 		}
 	}
diff --git a/Engine/Source/Editor/AnimGraph/Private/AnimGraphNode_Constraint.cpp b/Engine/Source/Editor/AnimGraph/Private/AnimGraphNode_Constraint.cpp
--- a/Engine/Source/Editor/AnimGraph/Private/AnimGraphNode_Constraint.cpp
+++ b/Engine/Source/Editor/AnimGraph/Private/AnimGraphNode_Constraint.cpp
@@ -38,7 +38,7 @@ UAnimGraphNode_Constraint::UAnimGraphNode_Constraint(const FObjectInitializer& O
 
 void UAnimGraphNode_Constraint::ValidateAnimNodeDuringCompilation(USkeleton* ForSkeleton, FCompilerResultsLog& MessageLog)
 {
-	if (ForSkeleton->GetReferenceSkeleton().FindBoneIndex(Node.BoneToModify.BoneName) == INDEX_NONE)
+	if (ForSkeleton && ForSkeleton->GetReferenceSkeleton().FindBoneIndex(Node.BoneToModify.BoneName) == INDEX_NONE)
 	{
 		if (Node.BoneToModify.BoneName == NAME_None)
 		{
diff --git a/Engine/Source/Editor/AnimGraph/Private/AnimGraphNode_ObserveBone.cpp b/Engine/Source/Editor/AnimGraph/Private/AnimGraphNode_ObserveBone.cpp
--- a/Engine/Source/Editor/AnimGraph/Private/AnimGraphNode_ObserveBone.cpp
+++ b/Engine/Source/Editor/AnimGraph/Private/AnimGraphNode_ObserveBone.cpp
@@ -103,7 +103,7 @@ UAnimGraphNode_ObserveBone::UAnimGraphNode_ObserveBone(const FObjectInitializer&
 
 void UAnimGraphNode_ObserveBone::ValidateAnimNodeDuringCompilation(USkeleton* ForSkeleton, FCompilerResultsLog& MessageLog)
 {
-	if (ForSkeleton->GetReferenceSkeleton().FindBoneIndex(Node.BoneToObserve.BoneName) == INDEX_NONE)
+	if (ForSkeleton && ForSkeleton->GetReferenceSkeleton().FindBoneIndex(Node.BoneToObserve.BoneName) == INDEX_NONE)
 	{
 		MessageLog.Warning(*LOCTEXT("NoBoneToObserve", "@@ - You must pick a bone to observe").ToString(), this);
 	}
diff --git a/Engine/Source/Editor/AnimGraph/Private/AnimGraphNode_PoseByName.cpp b/Engine/Source/Editor/AnimGraph/Private/AnimGraphNode_PoseByName.cpp
--- a/Engine/Source/Editor/AnimGraph/Private/AnimGraphNode_PoseByName.cpp
+++ b/Engine/Source/Editor/AnimGraph/Private/AnimGraphNode_PoseByName.cpp
@@ -120,7 +120,7 @@ void UAnimGraphNode_PoseByName::ValidateAnimNodeDuringCompilation(class USkeleto
 			MessageLog.Error(TEXT("@@ references an unknown pose asset"), this);
 		}
 	}
-	else
+	else if (ForSkeleton)
 	{
 		USkeleton* SeqSkeleton = PoseAssetToCheck->GetSkeleton();
 		if (SeqSkeleton && // if anim sequence doesn't have skeleton, it might be due to anim sequence not loaded yet, @todo: wait with anim blueprint compilation until all assets are loaded?
diff --git a/Engine/Source/Editor/AnimGraph/Private/AnimGraphNode_PoseDriver.cpp b/Engine/Source/Editor/AnimGraph/Private/AnimGraphNode_PoseDriver.cpp
--- a/Engine/Source/Editor/AnimGraph/Private/AnimGraphNode_PoseDriver.cpp
+++ b/Engine/Source/Editor/AnimGraph/Private/AnimGraphNode_PoseDriver.cpp
@@ -114,12 +114,15 @@ void UAnimGraphNode_PoseDriver::ValidateAnimNodeDuringCompilation(USkeleton* For
 	}
 
 	FName MissingBoneName = NAME_None;
-	for (const FBoneReference& BoneRef : Node.SourceBones)
+	if (ForSkeleton)
 	{
-		if (ForSkeleton->GetReferenceSkeleton().FindBoneIndex(BoneRef.BoneName) == INDEX_NONE)
+		for (const FBoneReference& BoneRef : Node.SourceBones)
 		{
-			MissingBoneName = BoneRef.BoneName;
-			break;
+			if (ForSkeleton->GetReferenceSkeleton().FindBoneIndex(BoneRef.BoneName) == INDEX_NONE)
+			{
+				MissingBoneName = BoneRef.BoneName;
+				break;
+			}
 		}
 	}
 
diff --git a/Engine/Source/Editor/AnimGraph/Private/AnimGraphNode_RotationOffsetBlendSpace.cpp b/Engine/Source/Editor/AnimGraph/Private/AnimGraphNode_RotationOffsetBlendSpace.cpp
--- a/Engine/Source/Editor/AnimGraph/Private/AnimGraphNode_RotationOffsetBlendSpace.cpp
+++ b/Engine/Source/Editor/AnimGraph/Private/AnimGraphNode_RotationOffsetBlendSpace.cpp
@@ -158,7 +158,7 @@ void UAnimGraphNode_RotationOffsetBlendSpace::ValidateAnimNodeDuringCompilation(
 	{
 		MessageLog.Error(TEXT("@@ references an invalid blend space (one that is not an aim offset)"), this);
 	}
-	else
+	else if (ForSkeleton)
 	{
 		USkeleton* BlendSpaceSkeleton = BlendSpaceToCheck->GetSkeleton();
 		if (BlendSpaceSkeleton && // if blend space doesn't have skeleton, it might be due to blend space not loaded yet, @todo: wait with anim blueprint compilation until all assets are loaded?
