Fixed tagged-property-serialization causing warnings with generated sparse class data structs in animation blueprints

This changes the generated property names for object properties to ensure that object properties of differing class types dont collide. This avoids the warning in FObjectPropertyBase::CheckValidObject (Serialized Class X for a property of Class Y. Reference will be nullptred.).

diff --git a/Engine/Source/Editor/AnimGraph/Private/AnimBlueprintCompiler.cpp b/Engine/Source/Editor/AnimGraph/Private/AnimBlueprintCompiler.cpp
--- a/Engine/Source/Editor/AnimGraph/Private/AnimBlueprintCompiler.cpp
+++ b/Engine/Source/Editor/AnimGraph/Private/AnimBlueprintCompiler.cpp
@@ -1915,7 +1915,21 @@ void FAnimBlueprintCompilerContext::ProcessFoldedPropertyRecords()
 					if(AnimationGraphSchema->ConvertPropertyToPinType(Record->Property, VariableType))
 					{
 						// Patch into sparse class data
-						const FName PropertyName = FName(Record->Property->GetClass()->GetFName(), InRecords.Num() - 1 - PropertyIndex);
+						FString PropertyNameString;
+
+						// If we are an object property, we should be named according to the underlying class type to avoid
+						// warnings if we tagged-property-serialize a colliding name with an updated struct layout
+						if(const FObjectPropertyBase* ObjectProperty = CastField<FObjectPropertyBase>(Record->Property))
+						{
+							PropertyNameString = ObjectProperty->PropertyClass->GetName();
+						}
+						else
+						{
+							PropertyNameString = Record->Property->GetClass()->GetName();
+						}
+
+						const FName PropertyNameBase(*PropertyNameString);
+						const FName PropertyName = FName(PropertyNameBase, InRecords.Num() - 1 - PropertyIndex);
 						Record->GeneratedProperty = CreateStructVariable(InStruct, PropertyName, VariableType);
 						if (Record->GeneratedProperty == nullptr)
 						{
