Made  LinkedAnimLayers node purple to distinguish them from LinkedAnimGraph

#tests anim blueprint editor

diff --git a/Engine/Source/Editor/AnimGraph/Private/AnimGraphNode_LinkedAnimGraphBase.cpp b/Engine/Source/Editor/AnimGraph/Private/AnimGraphNode_LinkedAnimGraphBase.cpp
--- a/Engine/Source/Editor/AnimGraph/Private/AnimGraphNode_LinkedAnimGraphBase.cpp
+++ b/Engine/Source/Editor/AnimGraph/Private/AnimGraphNode_LinkedAnimGraphBase.cpp
@@ -66,7 +66,7 @@ FLinearColor UAnimGraphNode_LinkedAnimGraphBase::GetNodeTitleColor() const
 	
 	if(!Color.IsSet())
 	{
-		return LinkedAnimGraphGraphNodeConstants::TitleColor;
+		return GetDefaultNodeTitleColor();
 	}
 
 	return Color.GetValue();
@@ -463,6 +463,11 @@ void UAnimGraphNode_LinkedAnimGraphBase::OnSetInstanceBlueprint(const FAssetData
 	}
 }
 
+FLinearColor UAnimGraphNode_LinkedAnimGraphBase::GetDefaultNodeTitleColor() const
+{
+	return LinkedAnimGraphGraphNodeConstants::TitleColor;
+}
+
 FPoseLinkMappingRecord UAnimGraphNode_LinkedAnimGraphBase::GetLinkIDLocation(const UScriptStruct* NodeType, UEdGraphPin* SourcePin)
 {
 	FPoseLinkMappingRecord Record = Super::GetLinkIDLocation(NodeType, SourcePin);
diff --git a/Engine/Source/Editor/AnimGraph/Private/AnimGraphNode_LinkedAnimLayer.cpp b/Engine/Source/Editor/AnimGraph/Private/AnimGraphNode_LinkedAnimLayer.cpp
--- a/Engine/Source/Editor/AnimGraph/Private/AnimGraphNode_LinkedAnimLayer.cpp
+++ b/Engine/Source/Editor/AnimGraph/Private/AnimGraphNode_LinkedAnimLayer.cpp
@@ -25,6 +25,12 @@
 
 #define LOCTEXT_NAMESPACE "LinkedAnimLayerNode"
 
+namespace LinkedAnimLayerGraphNodeConstants
+{
+	FLinearColor TitleColor(0.4f, 0.1f, 1.f);
+}
+
+
 void UAnimGraphNode_LinkedAnimLayer::Serialize(FArchive& Ar)
 {
 	Super::Serialize(Ar);
@@ -659,6 +665,11 @@ bool UAnimGraphNode_LinkedAnimLayer::IsStructuralProperty(FProperty* InProperty)
 		InProperty->GetFName() == GET_MEMBER_NAME_CHECKED(FAnimNode_LinkedAnimLayer, Layer);
 }
 
+FLinearColor UAnimGraphNode_LinkedAnimLayer::GetDefaultNodeTitleColor() const
+{
+	return LinkedAnimLayerGraphNodeConstants::TitleColor;
+}
+
 UClass* UAnimGraphNode_LinkedAnimLayer::GetTargetSkeletonClass() const
 {
 	UClass* SuperTargetSkeletonClass = Super::GetTargetSkeletonClass();
diff --git a/Engine/Source/Editor/AnimGraph/Public/AnimGraphNode_LinkedAnimGraphBase.h b/Engine/Source/Editor/AnimGraph/Public/AnimGraphNode_LinkedAnimGraphBase.h
--- a/Engine/Source/Editor/AnimGraph/Public/AnimGraphNode_LinkedAnimGraphBase.h
+++ b/Engine/Source/Editor/AnimGraph/Public/AnimGraphNode_LinkedAnimGraphBase.h
@@ -82,6 +82,8 @@ protected:
 	void OnSetInstanceBlueprint(const FAssetData& AssetData, IDetailLayoutBuilder* InDetailBuilder);
 	// ----- END UI CALLBACKS ----- //
 
+	virtual FLinearColor GetDefaultNodeTitleColor() const;
+
 	// Skeleton name used for filtering unloaded classes 
 	FString SkeletonName;
 
diff --git a/Engine/Source/Editor/AnimGraph/Public/AnimGraphNode_LinkedAnimLayer.h b/Engine/Source/Editor/AnimGraph/Public/AnimGraphNode_LinkedAnimLayer.h
--- a/Engine/Source/Editor/AnimGraph/Public/AnimGraphNode_LinkedAnimLayer.h
+++ b/Engine/Source/Editor/AnimGraph/Public/AnimGraphNode_LinkedAnimLayer.h
@@ -80,6 +80,7 @@ protected:
 	virtual bool OnShouldFilterInstanceBlueprint(const FAssetData& AssetData) const override;
 	virtual FString GetCurrentInstanceBlueprintPath() const override;
 	virtual bool IsStructuralProperty(FProperty* InProperty) const override;
+	virtual FLinearColor GetDefaultNodeTitleColor() const override;
 
 	friend class FAnimationLayerDragDropAction;
 	
