[UE][FIX] cyclic module dependency introduced with CL-19283727
- Revert change to AnimGraph.Build.cs
- Introduce new abstract interface for Physics Assert Render functions called from the AnimGraph module

diff --git a/Engine/Source/Editor/AnimGraph/AnimGraph.Build.cs b/Engine/Source/Editor/AnimGraph/AnimGraph.Build.cs
--- a/Engine/Source/Editor/AnimGraph/AnimGraph.Build.cs
+++ b/Engine/Source/Editor/AnimGraph/AnimGraph.Build.cs
@@ -23,7 +23,6 @@ public class AnimGraph : ModuleRules
 				"Slate",
 				"AnimGraphRuntime",
 				"BlueprintGraph",
-				"Persona"
 			}
 		);
 
@@ -55,7 +54,8 @@ public class AnimGraph : ModuleRules
 
         PrivateIncludePathModuleNames.AddRange(
             new string[] {
-                "SkeletonEditor",
+				"Persona",
+				"SkeletonEditor",
                 "AdvancedPreviewScene",
                 "AnimationBlueprintEditor",
             }
diff --git a/Engine/Source/Editor/AnimGraph/Private/AnimGraphNode_RigidBody.cpp b/Engine/Source/Editor/AnimGraph/Private/AnimGraphNode_RigidBody.cpp
--- a/Engine/Source/Editor/AnimGraph/Private/AnimGraphNode_RigidBody.cpp
+++ b/Engine/Source/Editor/AnimGraph/Private/AnimGraphNode_RigidBody.cpp
@@ -4,7 +4,7 @@
 #include "AnimNodeEditModes.h"
 #include "Kismet2/CompilerResultsLog.h"
 #include "BoneControllers/AnimNode_RigidBody.h"
-#include "PhysicsAssetRenderUtils.h"
+#include "IPhysicsAssetRenderInterface.h"
 
 // Details includes
 #include "PropertyHandle.h"
@@ -56,7 +56,8 @@ void UAnimGraphNode_RigidBody::Draw(FPrimitiveDrawInterface* PDI, USkeletalMeshC
 	{
 		if (UPhysicsAsset* const PhysicsAsset = RuntimeRigidBodyNode->GetPhysicsAsset())
 		{
-			PhysicsAssetRender::DebugDraw(PreviewSkelMeshComp, PhysicsAsset, PDI);
+			IPhysicsAssetRenderInterface& PhysicsAssetRenderInterface = IModularFeatures::Get().GetModularFeature<IPhysicsAssetRenderInterface>("PhysicsAssetRenderInterface");
+			PhysicsAssetRenderInterface.DebugDraw(PreviewSkelMeshComp, PhysicsAsset, PDI);
 		}
 	}
 }
@@ -109,49 +110,41 @@ void UAnimGraphNode_RigidBody::PostEditChangeProperty(struct FPropertyChangedEve
 {
 	Super::PostEditChangeProperty(PropertyChangedEvent);
 
-	if (UPhysicsAssetRenderUtilities* PhysicsAssetRenderUtilities = GetMutableDefault<UPhysicsAssetRenderUtilities>())
-	{
-		PhysicsAssetRenderUtilities->SaveConfig();
-	}
+	IPhysicsAssetRenderInterface& PhysicsAssetRenderInterface = IModularFeatures::Get().GetModularFeature<IPhysicsAssetRenderInterface>("PhysicsAssetRenderInterface");
+	PhysicsAssetRenderInterface.SaveConfig();
 }
 
 void UAnimGraphNode_RigidBody::ToggleBodyVisibility()
 {
 	FAnimNode_RigidBody* const RigidBodyNode = static_cast<FAnimNode_RigidBody*>(GetDebuggedAnimNode());
-	FPhysicsAssetRenderSettings* const PhysicsAssetRenderSettings = GetRenderSettings();
+	
+	IPhysicsAssetRenderInterface& PhysicsAssetRenderInterface = IModularFeatures::Get().GetModularFeature<IPhysicsAssetRenderInterface>("PhysicsAssetRenderInterface");
 
-	if (PhysicsAssetRenderSettings && RigidBodyNode)
+	if (RigidBodyNode)
 	{
-		PhysicsAssetRenderSettings->ToggleShowAllBodies(RigidBodyNode->GetPhysicsAsset());
+		PhysicsAssetRenderInterface.ToggleShowAllBodies(RigidBodyNode->GetPhysicsAsset());
 	}
 }
 
 void UAnimGraphNode_RigidBody::ToggleConstraintVisibility()
 {
 	FAnimNode_RigidBody* const RigidBodyNode = static_cast<FAnimNode_RigidBody*>(GetDebuggedAnimNode());
-	FPhysicsAssetRenderSettings* const PhysicsAssetRenderSettings = GetRenderSettings();
+	IPhysicsAssetRenderInterface& PhysicsAssetRenderInterface = IModularFeatures::Get().GetModularFeature<IPhysicsAssetRenderInterface>("PhysicsAssetRenderInterface");
 
-	if (PhysicsAssetRenderSettings && RigidBodyNode)
+	if (RigidBodyNode)
 	{
-		PhysicsAssetRenderSettings->ToggleShowAllConstraints(RigidBodyNode->GetPhysicsAsset());
+		PhysicsAssetRenderInterface.ToggleShowAllConstraints(RigidBodyNode->GetPhysicsAsset());
 	}
 }
 
-FPhysicsAssetRenderSettings* UAnimGraphNode_RigidBody::GetRenderSettings() const
-{
-	if (FAnimNode_RigidBody* const RigidBodyNode = static_cast<FAnimNode_RigidBody*>(GetDebuggedAnimNode()))
-	{
-		return UPhysicsAssetRenderUtilities::GetSettings(RigidBodyNode->GetPhysicsAsset());
-	}
-
-	return nullptr;
-}
-
 bool UAnimGraphNode_RigidBody::AreAnyBodiesHidden() const
 {
-	if (FPhysicsAssetRenderSettings* const PhysicsAssetRenderSettings = GetRenderSettings())
+	FAnimNode_RigidBody* const RigidBodyNode = static_cast<FAnimNode_RigidBody*>(GetDebuggedAnimNode());
+	IPhysicsAssetRenderInterface& PhysicsAssetRenderInterface = IModularFeatures::Get().GetModularFeature<IPhysicsAssetRenderInterface>("PhysicsAssetRenderInterface");
+
+	if (RigidBodyNode)
 	{
-		return PhysicsAssetRenderSettings->AreAnyBodiesHidden();
+		PhysicsAssetRenderInterface.AreAnyBodiesHidden(RigidBodyNode->GetPhysicsAsset());
 	}
 
 	return false;
@@ -159,9 +152,12 @@ bool UAnimGraphNode_RigidBody::AreAnyBodiesHidden() const
 
 bool UAnimGraphNode_RigidBody::AreAnyConstraintsHidden() const
 {
-	if (FPhysicsAssetRenderSettings* const PhysicsAssetRenderSettings = GetRenderSettings())
+	FAnimNode_RigidBody* const RigidBodyNode = static_cast<FAnimNode_RigidBody*>(GetDebuggedAnimNode());
+	IPhysicsAssetRenderInterface& PhysicsAssetRenderInterface = IModularFeatures::Get().GetModularFeature<IPhysicsAssetRenderInterface>("PhysicsAssetRenderInterface");
+
+	if (RigidBodyNode)
 	{
-		return PhysicsAssetRenderSettings->AreAnyConstraintsHidden();
+		PhysicsAssetRenderInterface.AreAnyConstraintsHidden(RigidBodyNode->GetPhysicsAsset());
 	}
 
 	return false;
diff --git a/Engine/Source/Editor/AnimGraph/Public/AnimGraphNode_RigidBody.h b/Engine/Source/Editor/AnimGraph/Public/AnimGraphNode_RigidBody.h
--- a/Engine/Source/Editor/AnimGraph/Public/AnimGraphNode_RigidBody.h
+++ b/Engine/Source/Editor/AnimGraph/Public/AnimGraphNode_RigidBody.h
@@ -36,7 +36,6 @@ protected:
 	virtual const FAnimNode_SkeletalControlBase* GetNode() const override { return &Node; }
 	// End of UAnimGraphNode_SkeletalControlBase interface
 
-	struct FPhysicsAssetRenderSettings* GetRenderSettings() const;
 	bool AreAnyBodiesHidden() const;
 	bool AreAnyConstraintsHidden() const;
 };
diff --git a/Engine/Source/Editor/Persona/Private/PhysicsAssetRenderUtils.cpp b/Engine/Source/Editor/Persona/Private/PhysicsAssetRenderUtils.cpp
--- a/Engine/Source/Editor/Persona/Private/PhysicsAssetRenderUtils.cpp
+++ b/Engine/Source/Editor/Persona/Private/PhysicsAssetRenderUtils.cpp
@@ -511,11 +511,28 @@ namespace PhysicsAssetRender
 	}
 }
 
+UPhysicsAssetRenderUtilities::UPhysicsAssetRenderUtilities()
+: PhysicsAssetRenderInterface(nullptr)
+{}
+
+UPhysicsAssetRenderUtilities::~UPhysicsAssetRenderUtilities()
+{
+	delete PhysicsAssetRenderInterface;
+}
+
+
 void UPhysicsAssetRenderUtilities::Initialise()
 {
 	if (UPhysicsAssetRenderUtilities* PhysicsAssetRenderUtilities = GetMutableDefault<UPhysicsAssetRenderUtilities>())
 	{
 		PhysicsAssetRenderUtilities->InitialiseImpl();
+
+		if (!PhysicsAssetRenderUtilities->PhysicsAssetRenderInterface)
+		{
+			PhysicsAssetRenderUtilities->PhysicsAssetRenderInterface = new FPhysicsAssetRenderInterface;
+		}
+
+		IModularFeatures::Get().RegisterModularFeature("PhysicsAssetRenderInterface", PhysicsAssetRenderUtilities->PhysicsAssetRenderInterface);
 	}
 }
 
@@ -598,7 +615,6 @@ FPhysicsAssetRenderSettings* UPhysicsAssetRenderUtilities::GetSettingsImpl(const
 	return Settings;
 }
 
-
 uint32 UPhysicsAssetRenderUtilities::GetPathNameHash(const UPhysicsAsset* InPhysicsAsset)
 {
 	if (InPhysicsAsset)
@@ -614,4 +630,54 @@ uint32 UPhysicsAssetRenderUtilities::GetPathNameHash(const FString& InPhysicsAss
 	return TextKeyUtil::HashString(InPhysicsAssetPathName);
 }
 
+// class FPhysicsAssetRenderInterface
+void FPhysicsAssetRenderInterface::DebugDraw(class USkeletalMeshComponent* const SkeletalMeshComponent, class UPhysicsAsset* const PhysicsAsset, FPrimitiveDrawInterface* PDI)
+{
+	PhysicsAssetRender::DebugDraw(SkeletalMeshComponent, PhysicsAsset, PDI);
+}
+
+void FPhysicsAssetRenderInterface::SaveConfig()
+{
+	if (UPhysicsAssetRenderUtilities* PhysicsAssetRenderUtilities = GetMutableDefault<UPhysicsAssetRenderUtilities>())
+	{
+		PhysicsAssetRenderUtilities->SaveConfig();
+	}
+}
+
+void FPhysicsAssetRenderInterface::ToggleShowAllBodies(class UPhysicsAsset* const PhysicsAsset)
+{
+	if (FPhysicsAssetRenderSettings* PhysicsAssetRenderSettings = UPhysicsAssetRenderUtilities::GetSettings(PhysicsAsset))
+	{
+		PhysicsAssetRenderSettings->ToggleShowAllBodies(PhysicsAsset);
+	}
+}
+
+void FPhysicsAssetRenderInterface::ToggleShowAllConstraints(class UPhysicsAsset* const PhysicsAsset)
+{
+	if (FPhysicsAssetRenderSettings* PhysicsAssetRenderSettings = UPhysicsAssetRenderUtilities::GetSettings(PhysicsAsset))
+	{
+		PhysicsAssetRenderSettings->ToggleShowAllConstraints(PhysicsAsset);
+	}
+}
+
+bool FPhysicsAssetRenderInterface::AreAnyBodiesHidden(class UPhysicsAsset* const PhysicsAsset)
+{
+	if (FPhysicsAssetRenderSettings* PhysicsAssetRenderSettings = UPhysicsAssetRenderUtilities::GetSettings(PhysicsAsset))
+	{
+		return PhysicsAssetRenderSettings->AreAnyBodiesHidden();
+	}
+
+	return false;
+}
+
+bool FPhysicsAssetRenderInterface::AreAnyConstraintsHidden(class UPhysicsAsset* const PhysicsAsset)
+{
+	if (FPhysicsAssetRenderSettings* PhysicsAssetRenderSettings = UPhysicsAssetRenderUtilities::GetSettings(PhysicsAsset))
+	{
+		return PhysicsAssetRenderSettings->AreAnyConstraintsHidden();
+	}
+
+	return false;
+}
+
 #undef LOCTEXT_NAMESPACE
diff --git a/Engine/Source/Editor/Persona/Public/IPhysicsAssetRenderInterface.h b/Engine/Source/Editor/Persona/Public/IPhysicsAssetRenderInterface.h
new file mode 100644
--- /dev/null
+++ b/Engine/Source/Editor/Persona/Public/IPhysicsAssetRenderInterface.h
@@ -0,0 +1,23 @@
+// Copyright Epic Games, Inc. All Rights Reserved.
+
+#pragma once
+
+#include "CoreMinimal.h"
+#include "Features/IModularFeature.h"
+#include "Modules/ModuleInterface.h"
+
+class IPhysicsAssetRenderInterface : public IModuleInterface, public IModularFeature
+{
+	// virtual ~IPhysicsAssetRenderInterface() = default;
+
+public:
+
+	virtual void DebugDraw(class USkeletalMeshComponent* const SkeletalMeshComponent, class UPhysicsAsset* const PhysicsAsset, class FPrimitiveDrawInterface* PDI) = 0;
+
+	virtual void SaveConfig() = 0;
+
+	virtual void ToggleShowAllBodies(class UPhysicsAsset* const PhysicsAsset) = 0;
+	virtual void ToggleShowAllConstraints(class UPhysicsAsset* const PhysicsAsset) = 0;
+	virtual bool AreAnyBodiesHidden(class UPhysicsAsset* const PhysicsAsset) = 0;
+	virtual bool AreAnyConstraintsHidden(class UPhysicsAsset* const PhysicsAsset) = 0;
+};
diff --git a/Engine/Source/Editor/Persona/Public/PhysicsAssetRenderUtils.h b/Engine/Source/Editor/Persona/Public/PhysicsAssetRenderUtils.h
--- a/Engine/Source/Editor/Persona/Public/PhysicsAssetRenderUtils.h
+++ b/Engine/Source/Editor/Persona/Public/PhysicsAssetRenderUtils.h
@@ -3,6 +3,7 @@
 #pragma once
 
 #include "CoreMinimal.h"
+#include "IPhysicsAssetRenderInterface.h"
 #include "UObject/ObjectMacros.h"
 #include "Math/Color.h"
 #include "PhysicsEngine/ShapeElem.h"
@@ -130,6 +131,9 @@ class PERSONA_API UPhysicsAssetRenderUtilities : public UObject
 
 public:
 
+	UPhysicsAssetRenderUtilities();
+	virtual ~UPhysicsAssetRenderUtilities();
+
 	static void Initialise();
 
 	/** Returns an existing render settings object or creates and returns a new one if none exists. */
@@ -150,6 +154,8 @@ private:
 
 	UPROPERTY()
 	TMap< uint32, FPhysicsAssetRenderSettings > IdToSettingsMap;
+
+	IPhysicsAssetRenderInterface* PhysicsAssetRenderInterface;
 };
 
 //////////////////////////////////////////////////////////////////////////
@@ -178,4 +184,17 @@ namespace PhysicsAssetRender
 	PERSONA_API FTransform GetPrimitiveTransform(const UPhysicsAsset* PhysicsAsset, const FTransform& BoneTM, const int32 BodyIndex, const EAggCollisionShape::Type PrimType, const int32 PrimIndex, const float Scale);
 	PERSONA_API FColor GetPrimitiveColor(const int32 BodyIndex, const EAggCollisionShape::Type PrimitiveType, const int32 PrimitiveIndex, const FPhysicsAssetRenderSettings& Settings);
 	PERSONA_API class UMaterialInterface* GetPrimitiveMaterial(const int32 BodyIndex, const EAggCollisionShape::Type PrimitiveType, const int32 PrimitiveIndex, const FPhysicsAssetRenderSettings& Settings);
-}
\ No newline at end of file
+}
+
+class FPhysicsAssetRenderInterface : public IPhysicsAssetRenderInterface
+{
+public:
+	virtual void DebugDraw(class USkeletalMeshComponent* const SkeletalMeshComponent, class UPhysicsAsset* const PhysicsAsset, FPrimitiveDrawInterface* PDI) override;
+
+	virtual void SaveConfig() override;
+
+	virtual void ToggleShowAllBodies(class UPhysicsAsset* const PhysicsAsset) override;
+	virtual void ToggleShowAllConstraints(class UPhysicsAsset* const PhysicsAsset) override;
+	virtual bool AreAnyBodiesHidden(class UPhysicsAsset* const PhysicsAsset) override;
+	virtual bool AreAnyConstraintsHidden(class UPhysicsAsset* const PhysicsAsset) override;
+};
\ No newline at end of file
