Remove the Pose Driver's 'experimental' label. It's been around long enough.

diff --git a/Engine/Source/Editor/AnimGraph/Public/AnimGraphNode_PoseDriver.h b/Engine/Source/Editor/AnimGraph/Public/AnimGraphNode_PoseDriver.h
--- a/Engine/Source/Editor/AnimGraph/Public/AnimGraphNode_PoseDriver.h
+++ b/Engine/Source/Editor/AnimGraph/Public/AnimGraphNode_PoseDriver.h
@@ -11,7 +11,7 @@
 
 class FCompilerResultsLog;
 
-UCLASS(Experimental, BlueprintType)
+UCLASS(BlueprintType)
 class ANIMGRAPH_API UAnimGraphNode_PoseDriver : public UAnimGraphNode_PoseHandler
 {
 	GENERATED_UCLASS_BODY()
@@ -133,7 +133,7 @@ public:
 
 protected:
 	// UAnimGraphNode_PoseHandler interface
-	virtual bool IsPoseAssetRequired() { return false; }
+	virtual bool IsPoseAssetRequired() override { return false; }
 	virtual FAnimNode_PoseHandler* GetPoseHandlerNode() override { return &Node; }
 	virtual const FAnimNode_PoseHandler* GetPoseHandlerNode() const override { return &Node; }
 	// End of UAnimGraphNode_PoseHandler interface
