Non-looping Montage animations pop and don't finish their playback duration in previewer
#fix ensure that blending (in/out) multiplier for montage instance is zero-ed (matching old behaviour)

diff --git a/Engine/Source/Editor/AnimGraph/Private/AnimPreviewInstance.cpp b/Engine/Source/Editor/AnimGraph/Private/AnimPreviewInstance.cpp
--- a/Engine/Source/Editor/AnimGraph/Private/AnimPreviewInstance.cpp
+++ b/Engine/Source/Editor/AnimGraph/Private/AnimPreviewInstance.cpp
@@ -574,6 +574,7 @@ void UAnimPreviewInstance::RestartMontage(UAnimMontage* Montage, FName FromSecti
 		// just hard stop here
 		Montage_Stop(0.0f, Montage);
 		Montage_Play(Montage, Proxy.GetPlayRate());
+		MontagePreview_RemoveBlendOut();
 		if (FromSection != NAME_None)
 		{
 			Montage_JumpToSection(FromSection);
