Fix crash in cooked builds where references to editor-only properties were left in anim BP bytecode

This was inadvertenly caused by CL 19250839 that renamed some generated properties to that of the same name as an anim node property (BlendSpace - they used to be called ObjectProperty). Then the code that expanded anim nodes to link up BP logic left a 'set struct member to null' behind because the generated member name matched the pin on the anim graph node.
This fix switches generated members to have a "__" prefix to avoid collisions with user properties, as well as fixing the code on expansion to check the correct property name, which will cause the 'set struct member to null' node to not be left in the exec path.

diff --git a/Engine/Source/Editor/AnimGraph/Private/AnimBlueprintCompiler.cpp b/Engine/Source/Editor/AnimGraph/Private/AnimBlueprintCompiler.cpp
--- a/Engine/Source/Editor/AnimGraph/Private/AnimBlueprintCompiler.cpp
+++ b/Engine/Source/Editor/AnimGraph/Private/AnimBlueprintCompiler.cpp
@@ -1915,21 +1915,23 @@ void FAnimBlueprintCompilerContext::ProcessFoldedPropertyRecords()
 					if(AnimationGraphSchema->ConvertPropertyToPinType(Record->Property, VariableType))
 					{
 						// Patch into sparse class data
-						FString PropertyNameString;
+						TStringBuilder<64> PropertyNameStringBuilder;
+
+						// Add prefix for internal use
+						PropertyNameStringBuilder.Append(TEXT("__"));
 
 						// If we are an object property, we should be named according to the underlying class type to avoid
 						// warnings if we tagged-property-serialize a colliding name with an updated struct layout
 						if(const FObjectPropertyBase* ObjectProperty = CastField<FObjectPropertyBase>(Record->Property))
 						{
-							PropertyNameString = ObjectProperty->PropertyClass->GetName();
+							PropertyNameStringBuilder.Append(ObjectProperty->PropertyClass->GetName());
 						}
 						else
 						{
-							PropertyNameString = Record->Property->GetClass()->GetName();
+							PropertyNameStringBuilder.Append(Record->Property->GetClass()->GetName());
 						}
 
-						const FName PropertyNameBase(*PropertyNameString);
-						const FName PropertyName = FName(PropertyNameBase, InRecords.Num() - 1 - PropertyIndex);
+						const FName PropertyName = FName(PropertyNameStringBuilder.ToString(), InRecords.Num() - 1 - PropertyIndex);
 						Record->GeneratedProperty = CreateStructVariable(InStruct, PropertyName, VariableType);
 						if (Record->GeneratedProperty == nullptr)
 						{
diff --git a/Engine/Source/Editor/AnimGraph/Private/AnimBlueprintExtension_Base.cpp b/Engine/Source/Editor/AnimGraph/Private/AnimBlueprintExtension_Base.cpp
--- a/Engine/Source/Editor/AnimGraph/Private/AnimBlueprintExtension_Base.cpp
+++ b/Engine/Source/Editor/AnimGraph/Private/AnimBlueprintExtension_Base.cpp
@@ -698,7 +698,7 @@ void UAnimBlueprintExtension_Base::CreateEvaluationHandler(IAnimBlueprintCompila
 					{
 						UEdGraphPin* DestPin = SourceInfo->CopyRecords[0].DestPin;
 
-						PropertiesBeingSet.Add(TargetPin->PinName);
+						PropertiesBeingSet.Add(PropertyName);
 						TargetPin->CopyPersistentDataFromOldPin(*DestPin);
 						InCompilationContext.GetMessageLog().NotifyIntermediatePinCreation(TargetPin, DestPin);
 					}
