Added a distinction in color and title between linked anim layer & self anim layer anim nodes.
Self anim layer are now titled Anim Layer (self) instaed of Linked anim layer

#tests editor

diff --git a/Engine/Source/Editor/AnimGraph/Private/AnimGraphNode_LinkedAnimLayer.cpp b/Engine/Source/Editor/AnimGraph/Private/AnimGraphNode_LinkedAnimLayer.cpp
--- a/Engine/Source/Editor/AnimGraph/Private/AnimGraphNode_LinkedAnimLayer.cpp
+++ b/Engine/Source/Editor/AnimGraph/Private/AnimGraphNode_LinkedAnimLayer.cpp
@@ -27,7 +27,8 @@
 
 namespace LinkedAnimLayerGraphNodeConstants
 {
-	FLinearColor TitleColor(0.4f, 0.1f, 1.f);
+	FLinearColor TitleColorSelfLayer(0.2f, 0.07f, 0.6f);
+	FLinearColor TitleColorLinkedLayer(0.45f, 0.f, 0.7f);
 }
 
 
@@ -110,7 +111,7 @@ FText UAnimGraphNode_LinkedAnimLayer::GetNodeTitle(ENodeTitleType::Type TitleTyp
 		bool bIsSelf = TargetAnimBlueprintInterface == nullptr; 
 		
 		FFormatNamedArguments Args;
-		Args.Add(TEXT("NodeType"), LOCTEXT("NodeTitle", "Linked Anim Layer"));
+		Args.Add(TEXT("NodeType"), bIsSelf ? LOCTEXT("NodeTitle", "Anim Layer (self)") : LOCTEXT("NodeTitle", "Linked Anim Layer"));
 		Args.Add(TEXT("TargetClass"), bIsSelf ? LOCTEXT("ClassSelf", "Self") : FText::FromString(TargetAnimBlueprintInterface->GetName()));
 		Args.Add(TEXT("Layer"), (Node.Layer == NAME_None) ? LOCTEXT("LayerNone", "None") : FText::FromName(Node.Layer));
 
@@ -667,7 +668,14 @@ bool UAnimGraphNode_LinkedAnimLayer::IsStructuralProperty(FProperty* InProperty)
 
 FLinearColor UAnimGraphNode_LinkedAnimLayer::GetDefaultNodeTitleColor() const
 {
-	return LinkedAnimLayerGraphNodeConstants::TitleColor;
+	if (HasValidNonSelfLayer())
+	{
+		return LinkedAnimLayerGraphNodeConstants::TitleColorLinkedLayer;		
+	}
+	else
+	{
+		return LinkedAnimLayerGraphNodeConstants::TitleColorSelfLayer;
+	}
 }
 
 UClass* UAnimGraphNode_LinkedAnimLayer::GetTargetSkeletonClass() const
