Fixed crash opening an anim BP with an aim offset graph connected

diff --git a/Engine/Source/Editor/AnimGraph/Private/AnimGraphNode_RotationOffsetBlendSpaceGraph.cpp b/Engine/Source/Editor/AnimGraph/Private/AnimGraphNode_RotationOffsetBlendSpaceGraph.cpp
--- a/Engine/Source/Editor/AnimGraph/Private/AnimGraphNode_RotationOffsetBlendSpaceGraph.cpp
+++ b/Engine/Source/Editor/AnimGraph/Private/AnimGraphNode_RotationOffsetBlendSpaceGraph.cpp
@@ -69,6 +69,10 @@ void UAnimGraphNode_RotationOffsetBlendSpaceGraph::GetMenuActions(FBlueprintActi
 			{
 				return FText::Format(LOCTEXT("MenuDescFormat", "AimOffset '{0}'"), FText::FromName(InAssetData.AssetName));
 			}
+			else if(InClass != nullptr)
+			{
+				return InClass->GetDisplayNameText();
+			}
 			else
 			{
 				return LOCTEXT("MenuDesc", "AimOffset");
@@ -80,6 +84,10 @@ void UAnimGraphNode_RotationOffsetBlendSpaceGraph::GetMenuActions(FBlueprintActi
 			{
 				return FText::Format(LOCTEXT("MenuDescTooltipFormat", "AimOffset\n'{0}'"), FText::FromName(InAssetData.ObjectPath));
 			}
+			else if(InClass != nullptr)
+			{
+				return InClass->GetDisplayNameText();
+			}
 			else
 			{
 				return LOCTEXT("MenuDescTooltip", "AimOffset");
@@ -89,6 +97,11 @@ void UAnimGraphNode_RotationOffsetBlendSpaceGraph::GetMenuActions(FBlueprintActi
 		{
 			UAnimGraphNode_RotationOffsetBlendSpaceGraph* GraphNode = CastChecked<UAnimGraphNode_RotationOffsetBlendSpaceGraph>(InNewNode);
 			GraphNode->SetupFromAsset(InAssetData, bInIsTemplateNode);
+		},
+		[](UEdGraphNode* InNewNode, bool bInIsTemplateNode, TSubclassOf<UObject> InClass)
+		{
+			UAnimGraphNode_RotationOffsetBlendSpaceGraph* GraphNode = CastChecked<UAnimGraphNode_RotationOffsetBlendSpaceGraph>(InNewNode);
+			GraphNode->SetupFromClass(TSubclassOf<UBlendSpace>(InClass.Get()), bInIsTemplateNode);
 		});
 }
 
