Preview Additional Meshes doesn't copy curves from source skelmesh
#fix ensure that the native copy-pose node inside of AnimPreviewAttacheInstance has the 'Copy Curves' flag set

diff --git a/Engine/Source/Editor/AnimGraph/Private/AnimPreviewAttacheInstance.cpp b/Engine/Source/Editor/AnimGraph/Private/AnimPreviewAttacheInstance.cpp
--- a/Engine/Source/Editor/AnimGraph/Private/AnimPreviewAttacheInstance.cpp
+++ b/Engine/Source/Editor/AnimGraph/Private/AnimPreviewAttacheInstance.cpp
@@ -11,6 +11,7 @@ void FAnimPreviewAttacheInstanceProxy::Initialize(UAnimInstance* InAnimInstance)
 
 	FAnimationInitializeContext InitContext(this);
 	CopyPoseFromMesh.bUseAttachedParent = true;
+	CopyPoseFromMesh.bCopyCurves = true;
 	CopyPoseFromMesh.Initialize_AnyThread(InitContext);
 }
 
