[UE][FIX] RBAN UI buttons for debug draw don't update on debug draw state change.

#rmx

diff --git a/Engine/Source/Editor/AnimGraph/Private/AnimGraphNode_RigidBody.cpp b/Engine/Source/Editor/AnimGraph/Private/AnimGraphNode_RigidBody.cpp
--- a/Engine/Source/Editor/AnimGraph/Private/AnimGraphNode_RigidBody.cpp
+++ b/Engine/Source/Editor/AnimGraph/Private/AnimGraphNode_RigidBody.cpp
@@ -144,7 +144,7 @@ bool UAnimGraphNode_RigidBody::AreAnyBodiesHidden() const
 
 	if (RigidBodyNode)
 	{
-		PhysicsAssetRenderInterface.AreAnyBodiesHidden(RigidBodyNode->GetPhysicsAsset());
+		return PhysicsAssetRenderInterface.AreAnyBodiesHidden(RigidBodyNode->GetPhysicsAsset());
 	}
 
 	return false;
@@ -157,7 +157,7 @@ bool UAnimGraphNode_RigidBody::AreAnyConstraintsHidden() const
 
 	if (RigidBodyNode)
 	{
-		PhysicsAssetRenderInterface.AreAnyConstraintsHidden(RigidBodyNode->GetPhysicsAsset());
+		return PhysicsAssetRenderInterface.AreAnyConstraintsHidden(RigidBodyNode->GetPhysicsAsset());
 	}
 
 	return false;
