Fixed ensure when expanding asset override tree for anim BPs that contain pose handlers

diff --git a/Engine/Source/Editor/AnimGraph/Private/AnimGraphNode_PoseHandler.cpp b/Engine/Source/Editor/AnimGraph/Private/AnimGraphNode_PoseHandler.cpp
--- a/Engine/Source/Editor/AnimGraph/Private/AnimGraphNode_PoseHandler.cpp
+++ b/Engine/Source/Editor/AnimGraph/Private/AnimGraphNode_PoseHandler.cpp
@@ -57,3 +57,8 @@ UAnimationAsset* UAnimGraphNode_PoseHandler::GetAnimationAsset() const
 
 	return PoseAsset;
 }
+
+TSubclassOf<UAnimationAsset> UAnimGraphNode_PoseHandler::GetAnimationAssetClass() const
+{
+	return UPoseAsset::StaticClass();
+}
\ No newline at end of file
diff --git a/Engine/Source/Editor/AnimGraph/Public/AnimGraphNode_PoseHandler.h b/Engine/Source/Editor/AnimGraph/Public/AnimGraphNode_PoseHandler.h
--- a/Engine/Source/Editor/AnimGraph/Public/AnimGraphNode_PoseHandler.h
+++ b/Engine/Source/Editor/AnimGraph/Public/AnimGraphNode_PoseHandler.h
@@ -18,6 +18,7 @@ class UAnimGraphNode_PoseHandler : public UAnimGraphNode_AssetPlayerBase
 	virtual void ValidateAnimNodeDuringCompilation(class USkeleton* ForSkeleton, class FCompilerResultsLog& MessageLog) override;
 	virtual void PreloadRequiredAssets() override;
 	virtual UAnimationAsset* GetAnimationAsset() const override;
+	virtual TSubclassOf<UAnimationAsset> GetAnimationAssetClass() const override;
 	virtual void OnOverrideAssets(IAnimBlueprintNodeOverrideAssetsContext& InContext) const override;
 	// End of UAnimGraphNode_Base
 
