Animation blueprints: Allow linked animation layer nodes to be placed inside of self layers

#tests Anim Blueprint Editor

diff --git a/Engine/Source/Editor/AnimGraph/Private/AnimGraphNode_LinkedAnimLayer.cpp b/Engine/Source/Editor/AnimGraph/Private/AnimGraphNode_LinkedAnimLayer.cpp
--- a/Engine/Source/Editor/AnimGraph/Private/AnimGraphNode_LinkedAnimLayer.cpp
+++ b/Engine/Source/Editor/AnimGraph/Private/AnimGraphNode_LinkedAnimLayer.cpp
@@ -259,9 +259,9 @@ void UAnimGraphNode_LinkedAnimLayer::ValidateAnimNodeDuringCompilation(USkeleton
 				static const FName DefaultAnimGraphName("AnimGraph");
 				if (InGraph->Nodes.Contains(OriginalThis))
 				{
-					if (!InGraph->IsA<UAnimationStateGraph>() && InGraph->GetFName() != DefaultAnimGraphName)
+					if (!InGraph->IsA<UAnimationStateGraph>() && InGraph->GetFName() != DefaultAnimGraphName && InGraph->InterfaceGuid.IsValid())
 					{
-						MessageLog.Error(*FText::Format(LOCTEXT("NestedLayer", "Linked anim layer node @@ is part of Animation Layer Graph '{0}', layers cannot be nested."), FText::FromName(InGraph->GetFName())).ToString(), this);
+						MessageLog.Error(*FText::Format(LOCTEXT("NestedLayer", "Linked anim layer node @@ is part of Animation Layer Graph '{0}', linked layers cannot be nested."), FText::FromName(InGraph->GetFName())).ToString(), this);
 					}
 				}
 			};
