Fix for pin type colors only appearing white in state machines

diff --git a/Engine/Source/Editor/AnimGraph/Private/AnimationStateMachineSchema.cpp b/Engine/Source/Editor/AnimGraph/Private/AnimationStateMachineSchema.cpp
--- a/Engine/Source/Editor/AnimGraph/Private/AnimationStateMachineSchema.cpp
+++ b/Engine/Source/Editor/AnimGraph/Private/AnimationStateMachineSchema.cpp
@@ -359,7 +359,7 @@ void UAnimationStateMachineSchema::GetContextMenuActions(UToolMenu* Menu, UGraph
 
 FLinearColor UAnimationStateMachineSchema::GetPinTypeColor(const FEdGraphPinType& PinType) const
 {
-	return FLinearColor::White;
+	return UEdGraphSchema::GetPinTypeColor(PinType);
 }
 
 void UAnimationStateMachineSchema::GetGraphDisplayInformation(const UEdGraph& Graph, /*out*/ FGraphDisplayInfo& DisplayInfo) const
