Updated fix for state machine variable pin colors

diff --git a/Engine/Source/Editor/AnimGraph/Private/AnimationStateMachineSchema.cpp b/Engine/Source/Editor/AnimGraph/Private/AnimationStateMachineSchema.cpp
--- a/Engine/Source/Editor/AnimGraph/Private/AnimationStateMachineSchema.cpp
+++ b/Engine/Source/Editor/AnimGraph/Private/AnimationStateMachineSchema.cpp
@@ -359,7 +359,12 @@ void UAnimationStateMachineSchema::GetContextMenuActions(UToolMenu* Menu, UGraph
 
 FLinearColor UAnimationStateMachineSchema::GetPinTypeColor(const FEdGraphPinType& PinType) const
 {
-	return UEdGraphSchema::GetPinTypeColor(PinType);
+	if (PinType.PinCategory == TEXT("Transition"))
+	{
+		return FLinearColor::White;
+	}
+
+	return GetDefault<UEdGraphSchema_K2>()->GetPinTypeColor(PinType);
 }
 
 void UAnimationStateMachineSchema::GetGraphDisplayInformation(const UEdGraph& Graph, /*out*/ FGraphDisplayInfo& DisplayInfo) const
