Fix PIE crash - more fallout from this->IsChildOf() check where this == nullptr.

diff --git a/Engine/Source/Editor/AnimGraph/Private/AnimBlueprintCompiler.cpp b/Engine/Source/Editor/AnimGraph/Private/AnimBlueprintCompiler.cpp
--- a/Engine/Source/Editor/AnimGraph/Private/AnimBlueprintCompiler.cpp
+++ b/Engine/Source/Editor/AnimGraph/Private/AnimBlueprintCompiler.cpp
@@ -1442,7 +1442,7 @@ void FAnimBlueprintCompilerContext::FinishCompilingClass(UClass* Class)
 
 	for (FBPInterfaceDescription& InterfaceDesc : Blueprint->ImplementedInterfaces)
 	{
-		if (InterfaceDesc.Interface->IsChildOf<UAnimLayerInterface>())
+		if (InterfaceDesc.Interface != nullptr && InterfaceDesc.Interface->IsChildOf<UAnimLayerInterface>())
 		{
 			for (UEdGraph* Graph : InterfaceDesc.Graphs)
 			{
