Fix crash customizing an anim node that no longer has a valid blueprint

diff --git a/Engine/Source/Editor/AnimGraph/Private/AnimGraphNode_Base.cpp b/Engine/Source/Editor/AnimGraph/Private/AnimGraphNode_Base.cpp
--- a/Engine/Source/Editor/AnimGraph/Private/AnimGraphNode_Base.cpp
+++ b/Engine/Source/Editor/AnimGraph/Private/AnimGraphNode_Base.cpp
@@ -965,11 +965,12 @@ void UAnimGraphNode_Base::AutowireNewNode(UEdGraphPin* FromPin)
 TSharedRef<SWidget> UAnimGraphNode_Base::MakePropertyBindingWidget(const FAnimPropertyBindingWidgetArgs& InArgs)
 {
 	UAnimGraphNode_Base* FirstAnimGraphNode = InArgs.Nodes[0];
-	UBlueprint* Blueprint = FirstAnimGraphNode->GetAnimBlueprint();
 	const bool bMultiSelect = InArgs.Nodes.Num() > 1;
 	
-	if(IModularFeatures::Get().IsModularFeatureAvailable("PropertyAccessEditor"))
+	if(FirstAnimGraphNode->HasValidBlueprint() && IModularFeatures::Get().IsModularFeatureAvailable("PropertyAccessEditor"))
 	{
+		UBlueprint* Blueprint = FirstAnimGraphNode->GetAnimBlueprint();
+		
 		int32 PinArrayIndex = InArgs.PinName.GetNumber() - 1;
 		const bool bIsArrayOrArrayElement = InArgs.PinProperty->IsA<FArrayProperty>();
 		const bool bIsArrayElement = bIsArrayOrArrayElement && PinArrayIndex != INDEX_NONE && InArgs.bPropertyIsOnFNode;
