fix double click on linked layer node to navigate the the actual linked instance when you are attached and debugging

diff --git a/Engine/Source/Editor/AnimGraph/Private/AnimGraphNode_LinkedAnimLayer.cpp b/Engine/Source/Editor/AnimGraph/Private/AnimGraphNode_LinkedAnimLayer.cpp
--- a/Engine/Source/Editor/AnimGraph/Private/AnimGraphNode_LinkedAnimLayer.cpp
+++ b/Engine/Source/Editor/AnimGraph/Private/AnimGraphNode_LinkedAnimLayer.cpp
@@ -282,40 +282,53 @@ UObject* UAnimGraphNode_LinkedAnimLayer::GetJumpTargetForDoubleClick() const
 		UObject* JumpTargetObject = nullptr;
 
 		UAnimBlueprint* TargetAnimBlueprint = InClass ? CastChecked<UAnimBlueprint>(InClass->ClassGeneratedBy) : nullptr;
-		if(TargetAnimBlueprint == nullptr || TargetAnimBlueprint == Cast<UAnimBlueprint>(GetBlueprint()))
-		{
-			// jump to graph in self
-			TArray<UEdGraph*> Graphs;
-			GetBlueprint()->GetAllGraphs(Graphs);
 
-			UEdGraph** FoundGraph = Graphs.FindByPredicate([this](UEdGraph* InGraph){ return InGraph->GetFName() == Node.Layer; });
-			if(FoundGraph)
-			{
-				JumpTargetObject = *FoundGraph;
-			}
-		}
-		else if(TargetAnimBlueprint)
+		while (TargetAnimBlueprint != nullptr)
 		{
-			// jump to graph in other BP
+			// jump to graph in other BP, going up the parent BP hierarchy until we find it
 			TArray<UEdGraph*> Graphs;
 			TargetAnimBlueprint->GetAllGraphs(Graphs);
 
-			UEdGraph** FoundGraph = Graphs.FindByPredicate([this](UEdGraph* InGraph){ return InGraph->GetFName() == Node.Layer; });
-			if(FoundGraph)
+			UEdGraph** FoundGraph = Graphs.FindByPredicate([this](UEdGraph* InGraph) { return InGraph->GetFName() == Node.Layer; });
+			if (FoundGraph)
 			{
 				JumpTargetObject = *FoundGraph;
+				return JumpTargetObject;
 			}
 			else
 			{
-				JumpTargetObject = TargetAnimBlueprint;
+				TargetAnimBlueprint = UAnimBlueprint::GetParentAnimBlueprint(TargetAnimBlueprint);
 			}
 		}
 
+		// jump to graph in self
+		TArray<UEdGraph*> Graphs;
+		GetBlueprint()->GetAllGraphs(Graphs);
+
+		UEdGraph** FoundGraph = Graphs.FindByPredicate([this](UEdGraph* InGraph) { return InGraph->GetFName() == Node.Layer; });
+		if (FoundGraph)
+		{
+			JumpTargetObject = *FoundGraph;
+		}
+
 		return JumpTargetObject;
 	};
 
 	// First try a concrete class, if any
-	UObject* JumpTargetObject = JumpTargetFromClass(*Node.InstanceClass);
+
+	const FAnimNode_LinkedAnimGraph* RuntimeNode = GetLinkedAnimGraphNode();
+
+
+	UObject* JumpTargetObject = nullptr;
+
+	if (UObject* TargetInstance = RuntimeNode->GetTargetInstance<UObject>())
+	{
+		JumpTargetObject = JumpTargetFromClass(TargetInstance->GetClass());
+	}
+	if (JumpTargetObject == nullptr)
+	{
+		JumpTargetObject = JumpTargetFromClass(RuntimeNode->InstanceClass);
+	}
 	if(JumpTargetObject == nullptr)
 	{
 		// then try the interface
@@ -960,4 +973,17 @@ bool UAnimGraphNode_LinkedAnimLayer::IsActionFilteredOut(class FBlueprintActionF
 	return bIsFilteredOut;
 }
 
+FAnimNode_LinkedAnimGraph* UAnimGraphNode_LinkedAnimLayer::GetLinkedAnimGraphNode()
+{
+	FAnimNode_LinkedAnimLayer* const RuntimeLinkedAnimGraphNode = GetDebuggedAnimNode<FAnimNode_LinkedAnimLayer>();
+	return RuntimeLinkedAnimGraphNode ? RuntimeLinkedAnimGraphNode : &Node;
+}
+
+const FAnimNode_LinkedAnimGraph* UAnimGraphNode_LinkedAnimLayer::GetLinkedAnimGraphNode() const
+{
+	const FAnimNode_LinkedAnimGraph* const RuntimeLinkedAnimGraphNode = GetDebuggedAnimNode<FAnimNode_LinkedAnimLayer>();
+	return RuntimeLinkedAnimGraphNode ? RuntimeLinkedAnimGraphNode : &Node;
+}
+
+
 #undef LOCTEXT_NAMESPACE
diff --git a/Engine/Source/Editor/AnimGraph/Public/AnimGraphNode_LinkedAnimLayer.h b/Engine/Source/Editor/AnimGraph/Public/AnimGraphNode_LinkedAnimLayer.h
--- a/Engine/Source/Editor/AnimGraph/Public/AnimGraphNode_LinkedAnimLayer.h
+++ b/Engine/Source/Editor/AnimGraph/Public/AnimGraphNode_LinkedAnimLayer.h
@@ -75,8 +75,8 @@ protected:
 	virtual UClass* GetTargetSkeletonClass() const override;
 
 	// Begin UAnimGraphNode_LinkedAnimGraphBase
-	virtual FAnimNode_LinkedAnimLayer* GetLinkedAnimGraphNode() override { return &Node;  }
-	virtual const FAnimNode_LinkedAnimLayer* GetLinkedAnimGraphNode() const override { return &Node; }
+	virtual FAnimNode_LinkedAnimGraph* GetLinkedAnimGraphNode() override;
+	virtual const FAnimNode_LinkedAnimGraph* GetLinkedAnimGraphNode() const override;
 	virtual bool OnShouldFilterInstanceBlueprint(const FAssetData& AssetData) const override;
 	virtual FString GetCurrentInstanceBlueprintPath() const override;
 	virtual bool IsStructuralProperty(FProperty* InProperty) const override;
