Make double clicking Linked Anim Layer nodes also attach the debugger to the correct anim instance

diff --git a/Engine/Source/Editor/AnimGraph/Private/AnimGraphNode_LinkedAnimLayer.cpp b/Engine/Source/Editor/AnimGraph/Private/AnimGraphNode_LinkedAnimLayer.cpp
--- a/Engine/Source/Editor/AnimGraph/Private/AnimGraphNode_LinkedAnimLayer.cpp
+++ b/Engine/Source/Editor/AnimGraph/Private/AnimGraphNode_LinkedAnimLayer.cpp
@@ -275,12 +275,14 @@ void UAnimGraphNode_LinkedAnimLayer::ValidateAnimNodeDuringCompilation(USkeleton
 	}
 }
 
-UObject* UAnimGraphNode_LinkedAnimLayer::GetJumpTargetForDoubleClick() const
+
+void UAnimGraphNode_LinkedAnimLayer::GetLinkTarget(UObject*& OutTargetGraph, UBlueprint*& OutTargetBlueprint) const
 {
-	auto JumpTargetFromClass = [this](UClass* InClass)
-	{
-		UObject* JumpTargetObject = nullptr;
+	OutTargetGraph = nullptr;
+	OutTargetBlueprint = nullptr;
 
+	auto JumpTargetFromClass = [this](UClass* InClass, UObject*& OutTargetGraph, UBlueprint*& OutTargetBlueprint)
+	{
 		UAnimBlueprint* TargetAnimBlueprint = InClass ? CastChecked<UAnimBlueprint>(InClass->ClassGeneratedBy) : nullptr;
 
 		while (TargetAnimBlueprint != nullptr)
@@ -292,8 +294,9 @@ UObject* UAnimGraphNode_LinkedAnimLayer::GetJumpTargetForDoubleClick() const
 			UEdGraph** FoundGraph = Graphs.FindByPredicate([this](UEdGraph* InGraph) { return InGraph->GetFName() == Node.Layer; });
 			if (FoundGraph)
 			{
-				JumpTargetObject = *FoundGraph;
-				return JumpTargetObject;
+				OutTargetBlueprint = TargetAnimBlueprint;
+				OutTargetGraph = *FoundGraph;
+				return;
 			}
 			else
 			{
@@ -308,41 +311,63 @@ UObject* UAnimGraphNode_LinkedAnimLayer::GetJumpTargetForDoubleClick() const
 		UEdGraph** FoundGraph = Graphs.FindByPredicate([this](UEdGraph* InGraph) { return InGraph->GetFName() == Node.Layer; });
 		if (FoundGraph)
 		{
-			JumpTargetObject = *FoundGraph;
+			OutTargetBlueprint = nullptr;
+			OutTargetGraph = *FoundGraph;
+			return;
 		}
-
-		return JumpTargetObject;
 	};
 
 	// First try a concrete class, if any
 
 	const FAnimNode_LinkedAnimGraph* RuntimeNode = GetLinkedAnimGraphNode();
 
-
-	UObject* JumpTargetObject = nullptr;
-
 	if (UObject* TargetInstance = RuntimeNode->GetTargetInstance<UObject>())
 	{
-		JumpTargetObject = JumpTargetFromClass(TargetInstance->GetClass());
+		JumpTargetFromClass(TargetInstance->GetClass(), OutTargetGraph, OutTargetBlueprint);
 	}
-	if (JumpTargetObject == nullptr)
+	if (OutTargetGraph == nullptr)
 	{
-		JumpTargetObject = JumpTargetFromClass(RuntimeNode->InstanceClass);
+		JumpTargetFromClass(RuntimeNode->InstanceClass, OutTargetGraph, OutTargetBlueprint);
 	}
-	if(JumpTargetObject == nullptr)
+	if (OutTargetGraph == nullptr)
 	{
 		// then try the interface
-		JumpTargetObject = JumpTargetFromClass(*Node.Interface);
+		JumpTargetFromClass(*Node.Interface, OutTargetGraph, OutTargetBlueprint);
 	}
 
-	return JumpTargetObject;
+}
+
+
+UObject* UAnimGraphNode_LinkedAnimLayer::GetJumpTargetForDoubleClick() const
+{
+	UObject* TargetGraph;
+	UBlueprint* TargetBlueprint;
+
+	GetLinkTarget(TargetGraph, TargetBlueprint);
+	return TargetGraph;
 }
 
 void UAnimGraphNode_LinkedAnimLayer::JumpToDefinition() const
 {
-	if (UEdGraph* HyperlinkTarget = Cast<UEdGraph>(GetJumpTargetForDoubleClick()))
+
+	UObject* TargetGraph;
+	UBlueprint* TargetBlueprint;
+
+	GetLinkTarget(TargetGraph, TargetBlueprint);
+
+	if (UAnimationGraph* HyperlinkTarget = Cast<UAnimationGraph>(TargetGraph))
 	{
 		FKismetEditorUtilities::BringKismetToFocusAttentionOnObject(HyperlinkTarget);
+
+		if (TargetBlueprint != nullptr)
+		{
+			const FAnimNode_LinkedAnimGraph* RuntimeNode = GetLinkedAnimGraphNode();
+
+			if (UObject* TargetInstance = RuntimeNode->GetTargetInstance<UObject>())
+			{
+				TargetBlueprint->SetObjectBeingDebugged(TargetInstance);
+			}
+		}
 	}
 	else
 	{
diff --git a/Engine/Source/Editor/AnimGraph/Public/AnimGraphNode_LinkedAnimLayer.h b/Engine/Source/Editor/AnimGraph/Public/AnimGraphNode_LinkedAnimLayer.h
--- a/Engine/Source/Editor/AnimGraph/Public/AnimGraphNode_LinkedAnimLayer.h
+++ b/Engine/Source/Editor/AnimGraph/Public/AnimGraphNode_LinkedAnimLayer.h
@@ -54,6 +54,9 @@ public:
 	void UpdateGuidForLayer();
 
 protected:
+
+	void GetLinkTarget(UObject* &OutTargetGraph, UBlueprint* &OutTargetBlueprint) const;
+
 	// ----- UI CALLBACKS ----- //
 	// Handlers for layer combo
 	void GetLayerNames(TArray<TSharedPtr<FString>>& OutStrings, TArray<TSharedPtr<SToolTip>>& OutToolTips, TArray<bool>& OutRestrictedItems);
