Animation State Machine UX update (renamed "Add New State Machine" as "State Machine" and removed trailing "..." from all the context options, to be coherent with other UX options

diff --git a/Engine/Source/Editor/AnimGraph/Private/AnimGraphNode_StateMachineBase.cpp b/Engine/Source/Editor/AnimGraph/Private/AnimGraphNode_StateMachineBase.cpp
--- a/Engine/Source/Editor/AnimGraph/Private/AnimGraphNode_StateMachineBase.cpp
+++ b/Engine/Source/Editor/AnimGraph/Private/AnimGraphNode_StateMachineBase.cpp
@@ -98,7 +98,7 @@ FText UAnimGraphNode_StateMachineBase::GetNodeTitle(ENodeTitleType::Type TitleTy
 {
 	if ((TitleType == ENodeTitleType::MenuTitle || TitleType == ENodeTitleType::ListView) && (EditorStateMachineGraph == nullptr))
 	{
-		return LOCTEXT("AddNewStateMachine", "Add New State Machine...");
+		return LOCTEXT("AddNewStateMachine", "State Machine");
 	}
 	else if (EditorStateMachineGraph == nullptr)
 	{
diff --git a/Engine/Source/Editor/AnimGraph/Private/AnimationStateMachineSchema.cpp b/Engine/Source/Editor/AnimGraph/Private/AnimationStateMachineSchema.cpp
--- a/Engine/Source/Editor/AnimGraph/Private/AnimationStateMachineSchema.cpp
+++ b/Engine/Source/Editor/AnimGraph/Private/AnimationStateMachineSchema.cpp
@@ -274,19 +274,19 @@ void UAnimationStateMachineSchema::GetGraphContextActions(FGraphContextMenuBuild
 {
 	// Add state node
 	{
-		TSharedPtr<FEdGraphSchemaAction_NewStateNode> Action = AddNewStateNodeAction(ContextMenuBuilder, FText::GetEmpty(), LOCTEXT("AddState", "Add State..."), LOCTEXT("AddStateTooltip", "A new state"));
+		TSharedPtr<FEdGraphSchemaAction_NewStateNode> Action = AddNewStateNodeAction(ContextMenuBuilder, FText::GetEmpty(), LOCTEXT("AddState", "Add State"), LOCTEXT("AddStateTooltip", "A new state"));
 		Action->NodeTemplate = NewObject<UAnimStateNode>(ContextMenuBuilder.OwnerOfTemporaries);
 	}
 
 	// Add state alias node
 	{
-		TSharedPtr<FEdGraphSchemaAction_NewStateNode> Action = AddNewStateNodeAction(ContextMenuBuilder, FText::GetEmpty(), LOCTEXT("AddStateAlias", "Add State Alias..."), LOCTEXT("AddStateAliasTooltip", "A new state alias"));
+		TSharedPtr<FEdGraphSchemaAction_NewStateNode> Action = AddNewStateNodeAction(ContextMenuBuilder, FText::GetEmpty(), LOCTEXT("AddStateAlias", "Add State Alias"), LOCTEXT("AddStateAliasTooltip", "A new state alias"));
 		Action->NodeTemplate = NewObject<UAnimStateAliasNode>(ContextMenuBuilder.OwnerOfTemporaries);
 	}
 
 	// Add conduit node
 	{
-		TSharedPtr<FEdGraphSchemaAction_NewStateNode> Action = AddNewStateNodeAction(ContextMenuBuilder, FText::GetEmpty(), LOCTEXT("AddConduit", "Add Conduit..."), LOCTEXT("AddConduitTooltip", "A new conduit state"));
+		TSharedPtr<FEdGraphSchemaAction_NewStateNode> Action = AddNewStateNodeAction(ContextMenuBuilder, FText::GetEmpty(), LOCTEXT("AddConduit", "Add Conduit"), LOCTEXT("AddConduitTooltip", "A new conduit state"));
 		Action->NodeTemplate = NewObject<UAnimStateConduitNode>(ContextMenuBuilder.OwnerOfTemporaries);
 	}
 
@@ -305,7 +305,7 @@ void UAnimationStateMachineSchema::GetGraphContextActions(FGraphContextMenuBuild
 
 		if (!bHasEntry)
 		{
-			TSharedPtr<FEdGraphSchemaAction_NewStateNode> Action = AddNewStateNodeAction(ContextMenuBuilder, FText::GetEmpty(), LOCTEXT("AddEntryPoint", "Add Entry Point..."), LOCTEXT("AddEntryPointTooltip", "Define State Machine's Entry Point"));
+			TSharedPtr<FEdGraphSchemaAction_NewStateNode> Action = AddNewStateNodeAction(ContextMenuBuilder, FText::GetEmpty(), LOCTEXT("AddEntryPoint", "Add Entry Point"), LOCTEXT("AddEntryPointTooltip", "Define State Machine's Entry Point"));
 			Action->NodeTemplate = NewObject<UAnimStateEntryNode>(ContextMenuBuilder.OwnerOfTemporaries);
 		}
 	}
@@ -315,7 +315,7 @@ void UAnimationStateMachineSchema::GetGraphContextActions(FGraphContextMenuBuild
 	{
 		UBlueprint* OwnerBlueprint = FBlueprintEditorUtils::FindBlueprintForGraphChecked(ContextMenuBuilder.CurrentGraph);
 		const bool bIsManyNodesSelected = (FKismetEditorUtilities::GetNumberOfSelectedNodes(OwnerBlueprint) > 0);
-		const FText MenuDescription = bIsManyNodesSelected ? LOCTEXT("CreateCommentSelection", "Create Comment from Selection") : LOCTEXT("AddComment", "Add Comment...");
+		const FText MenuDescription = bIsManyNodesSelected ? LOCTEXT("CreateCommentSelection", "Create Comment from Selection") : LOCTEXT("AddComment", "Add Comment");
 		const FText ToolTip = LOCTEXT("CreateCommentSelectionTooltip", "Create a resizeable comment box around selected nodes.");
 
 		TSharedPtr<FEdGraphSchemaAction_NewStateComment> NewComment( new FEdGraphSchemaAction_NewStateComment(FText::GetEmpty(), MenuDescription, ToolTip, 0) );
