FNames containing asset paths are deprecated. FSoftObjectPath or FTopLevelAssetPath should be used instead.
Fixups for animation.

diff --git a/Engine/Plugins/Animation/GameplayInsights/Source/GameplayInsights/Private/GameplayProvider.cpp b/Engine/Plugins/Animation/GameplayInsights/Source/GameplayInsights/Private/GameplayProvider.cpp
--- a/Engine/Plugins/Animation/GameplayInsights/Source/GameplayInsights/Private/GameplayProvider.cpp
+++ b/Engine/Plugins/Animation/GameplayInsights/Source/GameplayInsights/Private/GameplayProvider.cpp
@@ -36,7 +36,7 @@ static UClass* FindBlueprintClass(const FString& TargetNameRaw)
 	UClass* FoundClass = nullptr;
 	AssetRegistry.EnumerateAssets(Filter, [&FoundClass, TargetName](const FAssetData& AssetData)
 	{
-		if ((AssetData.AssetName.ToString() == TargetName) || (AssetData.ObjectPath.ToString() == TargetName))
+		if ((AssetData.AssetName.ToString() == TargetName) || (AssetData.GetObjectPathString() == TargetName))
 		{
 			if (UBlueprint* BP = Cast<UBlueprint>(AssetData.GetAsset()))
 			{
diff --git a/Engine/Plugins/Animation/GameplayInsights/Source/GameplayInsights/Private/SkeletalMeshPoseTrack.cpp b/Engine/Plugins/Animation/GameplayInsights/Source/GameplayInsights/Private/SkeletalMeshPoseTrack.cpp
--- a/Engine/Plugins/Animation/GameplayInsights/Source/GameplayInsights/Private/SkeletalMeshPoseTrack.cpp
+++ b/Engine/Plugins/Animation/GameplayInsights/Source/GameplayInsights/Private/SkeletalMeshPoseTrack.cpp
@@ -112,7 +112,7 @@ void FSkeletalMeshPoseTrack::InitTooltip(FTooltipDrawState& Tooltip, const ITimi
 				Tooltip.AddNameValueTextLine(LOCTEXT("Mesh", "Mesh").ToString(), SkeletalMeshObjectInfo->PathName);
 
 				FAssetRegistryModule& AssetRegistryModule = FModuleManager::LoadModuleChecked<FAssetRegistryModule>(TEXT("AssetRegistry"));
-				if(!AssetRegistryModule.Get().GetAssetByObjectPath(SkeletalMeshObjectInfo->PathName).IsValid())
+				if(!AssetRegistryModule.Get().GetAssetByObjectPath(FSoftObjectPath(SkeletalMeshObjectInfo->PathName)).IsValid())
 				{
 					Tooltip.AddTextLine(LOCTEXT("MeshNotFound", "Mesh not found").ToString(), FLinearColor::Red);
 				}
diff --git a/Engine/Source/Editor/AnimGraph/Private/AnimGraphNode_BlendSpaceEvaluator.cpp b/Engine/Source/Editor/AnimGraph/Private/AnimGraphNode_BlendSpaceEvaluator.cpp
--- a/Engine/Source/Editor/AnimGraph/Private/AnimGraphNode_BlendSpaceEvaluator.cpp
+++ b/Engine/Source/Editor/AnimGraph/Private/AnimGraphNode_BlendSpaceEvaluator.cpp
@@ -48,7 +48,7 @@ void UAnimGraphNode_BlendSpaceEvaluator::GetMenuActions(FBlueprintActionDatabase
 		{
 			if(InAssetData.IsValid())
 			{
-				return FText::Format(LOCTEXT("MenuDescTooltipFormat", "Blendspace Evaluator\n'{0}'"), FText::FromName(InAssetData.ObjectPath));
+				return FText::Format(LOCTEXT("MenuDescTooltipFormat", "Blendspace Evaluator\n'{0}'"), FText::FromString(InAssetData.GetObjectPathString()));
 			}
 			else
 			{
diff --git a/Engine/Source/Editor/AnimGraph/Private/AnimGraphNode_BlendSpaceGraph.cpp b/Engine/Source/Editor/AnimGraph/Private/AnimGraphNode_BlendSpaceGraph.cpp
--- a/Engine/Source/Editor/AnimGraph/Private/AnimGraphNode_BlendSpaceGraph.cpp
+++ b/Engine/Source/Editor/AnimGraph/Private/AnimGraphNode_BlendSpaceGraph.cpp
@@ -80,7 +80,7 @@ void UAnimGraphNode_BlendSpaceGraph::GetMenuActions(FBlueprintActionDatabaseRegi
 		{
 			if(InAssetData.IsValid())
 			{
-				return FText::Format(LOCTEXT("MenuDescTooltipFormat", "Blendspace\n'{0}'"), FText::FromName(InAssetData.ObjectPath));
+				return FText::Format(LOCTEXT("MenuDescTooltipFormat", "Blendspace\n'{0}'"), FText::FromString(InAssetData.GetObjectPathString()));
 			}
 			else if(InClass != nullptr)
 			{
diff --git a/Engine/Source/Editor/AnimGraph/Private/AnimGraphNode_BlendSpacePlayer.cpp b/Engine/Source/Editor/AnimGraph/Private/AnimGraphNode_BlendSpacePlayer.cpp
--- a/Engine/Source/Editor/AnimGraph/Private/AnimGraphNode_BlendSpacePlayer.cpp
+++ b/Engine/Source/Editor/AnimGraph/Private/AnimGraphNode_BlendSpacePlayer.cpp
@@ -80,7 +80,7 @@ void UAnimGraphNode_BlendSpacePlayer::GetMenuActions(FBlueprintActionDatabaseReg
 		{
 			if(InAssetData.IsValid())
 			{
-				return FText::Format(LOCTEXT("MenuDescTooltipFormat", "Blendspace Player\n'{0}'"), FText::FromName(InAssetData.ObjectPath));
+				return FText::Format(LOCTEXT("MenuDescTooltipFormat", "Blendspace Player\n'{0}'"), FText::FromString(InAssetData.GetObjectPathString()));
 			}
 			else
 			{
diff --git a/Engine/Source/Editor/AnimGraph/Private/AnimGraphNode_LinkedAnimGraph.cpp b/Engine/Source/Editor/AnimGraph/Private/AnimGraphNode_LinkedAnimGraph.cpp
--- a/Engine/Source/Editor/AnimGraph/Private/AnimGraphNode_LinkedAnimGraph.cpp
+++ b/Engine/Source/Editor/AnimGraph/Private/AnimGraphNode_LinkedAnimGraph.cpp
@@ -115,7 +115,7 @@ void UAnimGraphNode_LinkedAnimGraph::GetMenuActions(FBlueprintActionDatabaseRegi
 				}
 				else
 				{
-					return FText::Format(LOCTEXT("MenuDescTooltipFormat", "Linked Anim Graph - Runs a linked anim graph in another instance to process animation\n'{0}'"), FText::FromName(InAssetData.ObjectPath));
+					return FText::Format(LOCTEXT("MenuDescTooltipFormat", "Linked Anim Graph - Runs a linked anim graph in another instance to process animation\n'{0}'"), FText::FromString(InAssetData.GetObjectPathString()));
 				}
 			}
 			else
diff --git a/Engine/Source/Editor/AnimGraph/Private/AnimGraphNode_LinkedAnimLayer.cpp b/Engine/Source/Editor/AnimGraph/Private/AnimGraphNode_LinkedAnimLayer.cpp
--- a/Engine/Source/Editor/AnimGraph/Private/AnimGraphNode_LinkedAnimLayer.cpp
+++ b/Engine/Source/Editor/AnimGraph/Private/AnimGraphNode_LinkedAnimLayer.cpp
@@ -556,8 +556,8 @@ bool UAnimGraphNode_LinkedAnimLayer::OnShouldFilterInstanceBlueprint(const FAsse
 				{
 					const FString ParentClassFromData = CurrentAssetData.GetTagValueRef<FString>(FBlueprintTags::ParentClassPath);
 					const FString ClassObjectPath = FPackageName::ExportTextPathToObjectPath(ParentClassFromData);
-					const FName BlueprintPath = FName(*ClassObjectPath.LeftChop(2)); // Chop off _C
-					CurrentAssetData = AssetRegistryModule.Get().GetAssetByObjectPath(BlueprintPath);
+					const FString BlueprintPath = ClassObjectPath.LeftChop(2); // Chop off _C
+					CurrentAssetData = AssetRegistryModule.Get().GetAssetByObjectPath(FSoftObjectPath(BlueprintPath));
 				}
 			// Only continue checking if the parent is an anim blueprint
 			} while (!bMatchesInterface && (CurrentAssetData.AssetClassPath == UAnimBlueprint::StaticClass()->GetClassPathName()));
diff --git a/Engine/Source/Editor/AnimGraph/Private/AnimGraphNode_Mirror.cpp b/Engine/Source/Editor/AnimGraph/Private/AnimGraphNode_Mirror.cpp
--- a/Engine/Source/Editor/AnimGraph/Private/AnimGraphNode_Mirror.cpp
+++ b/Engine/Source/Editor/AnimGraph/Private/AnimGraphNode_Mirror.cpp
@@ -227,7 +227,7 @@ void UAnimGraphNode_Mirror::GetMenuActions(FBlueprintActionDatabaseRegistrar& Ac
 			NodeSpawner->CustomizeNodeDelegate = UBlueprintNodeSpawner::FCustomizeNodeDelegate::CreateStatic(UnloadedAssetSetup, Asset);
 
 			NodeSpawner->DefaultMenuSignature.MenuName = GetTitleGivenAssetInfo(FText::FromName(Asset.AssetName));
-			NodeSpawner->DefaultMenuSignature.Tooltip = GetTitleGivenAssetInfo(FText::FromName(Asset.ObjectPath));
+			NodeSpawner->DefaultMenuSignature.Tooltip = GetTitleGivenAssetInfo(FText::FromString(Asset.GetObjectPathString()));
 			ActionRegistrar.AddBlueprintAction(Asset, NodeSpawner);
 		}
 	}
diff --git a/Engine/Source/Editor/AnimGraph/Private/AnimGraphNode_PoseBlendNode.cpp b/Engine/Source/Editor/AnimGraph/Private/AnimGraphNode_PoseBlendNode.cpp
--- a/Engine/Source/Editor/AnimGraph/Private/AnimGraphNode_PoseBlendNode.cpp
+++ b/Engine/Source/Editor/AnimGraph/Private/AnimGraphNode_PoseBlendNode.cpp
@@ -21,7 +21,7 @@ protected:
 	FAssetData AssetInfo;
 public:
 	FNewPoseBlendNodeAction(const FAssetData& InAssetInfo, FText Title)
-		: FEdGraphSchemaAction_K2NewNode(LOCTEXT("PoseAsset", "PoseAssets"), Title, LOCTEXT("EvalCurvesToMakePose", "Evaluates curves to produce a pose from pose asset"), 0, FText::FromName(InAssetInfo.ObjectPath))
+		: FEdGraphSchemaAction_K2NewNode(LOCTEXT("PoseAsset", "PoseAssets"), Title, LOCTEXT("EvalCurvesToMakePose", "Evaluates curves to produce a pose from pose asset"), 0, FText::FromString(InAssetInfo.GetObjectPathString()))
 	{
 		AssetInfo = InAssetInfo;
 
diff --git a/Engine/Source/Editor/AnimGraph/Private/AnimGraphNode_RotationOffsetBlendSpace.cpp b/Engine/Source/Editor/AnimGraph/Private/AnimGraphNode_RotationOffsetBlendSpace.cpp
--- a/Engine/Source/Editor/AnimGraph/Private/AnimGraphNode_RotationOffsetBlendSpace.cpp
+++ b/Engine/Source/Editor/AnimGraph/Private/AnimGraphNode_RotationOffsetBlendSpace.cpp
@@ -93,7 +93,7 @@ void UAnimGraphNode_RotationOffsetBlendSpace::GetMenuActions(FBlueprintActionDat
 		{
 			if(InAssetData.IsValid())
 			{
-				return FText::Format(LOCTEXT("MenuDescTooltipFormat", "AimOffset Player\n'{0}'"), FText::FromName(InAssetData.ObjectPath));
+				return FText::Format(LOCTEXT("MenuDescTooltipFormat", "AimOffset Player\n'{0}'"), FText::FromString(InAssetData.GetObjectPathString()));
 			}
 			else
 			{
diff --git a/Engine/Source/Editor/AnimGraph/Private/AnimGraphNode_RotationOffsetBlendSpaceGraph.cpp b/Engine/Source/Editor/AnimGraph/Private/AnimGraphNode_RotationOffsetBlendSpaceGraph.cpp
--- a/Engine/Source/Editor/AnimGraph/Private/AnimGraphNode_RotationOffsetBlendSpaceGraph.cpp
+++ b/Engine/Source/Editor/AnimGraph/Private/AnimGraphNode_RotationOffsetBlendSpaceGraph.cpp
@@ -82,7 +82,7 @@ void UAnimGraphNode_RotationOffsetBlendSpaceGraph::GetMenuActions(FBlueprintActi
 		{
 			if(InAssetData.IsValid())
 			{
-				return FText::Format(LOCTEXT("MenuDescTooltipFormat", "AimOffset\n'{0}'"), FText::FromName(InAssetData.ObjectPath));
+				return FText::Format(LOCTEXT("MenuDescTooltipFormat", "AimOffset\n'{0}'"), FText::FromString(InAssetData.GetObjectPathString()));
 			}
 			else if(InClass != nullptr)
 			{
diff --git a/Engine/Source/Editor/AnimGraph/Private/AnimGraphNode_SequenceEvaluator.cpp b/Engine/Source/Editor/AnimGraph/Private/AnimGraphNode_SequenceEvaluator.cpp
--- a/Engine/Source/Editor/AnimGraph/Private/AnimGraphNode_SequenceEvaluator.cpp
+++ b/Engine/Source/Editor/AnimGraph/Private/AnimGraphNode_SequenceEvaluator.cpp
@@ -102,11 +102,11 @@ void UAnimGraphNode_SequenceEvaluator::GetMenuActions(FBlueprintActionDatabaseRe
 				const FString TagValue = InAssetData.GetTagValueRef<FString>(GET_MEMBER_NAME_CHECKED(UAnimSequence, AdditiveAnimType));
 				if(const bool bKnownToBeAdditive = (!TagValue.IsEmpty() && !TagValue.Equals(TEXT("AAT_None"))))
 				{
-					return FText::Format(LOCTEXT("MenuDescTooltipFormat_EvaluateAdditive", "Evaluate (additive)\n'{0}'"), FText::FromName(InAssetData.ObjectPath));
+					return FText::Format(LOCTEXT("MenuDescTooltipFormat_EvaluateAdditive", "Evaluate (additive)\n'{0}'"), FText::FromString(InAssetData.GetObjectPathString()));
 				}
 				else
 				{
-					return FText::Format(LOCTEXT("MenuDescTooltipFormat_Evaluate", "Evaluate\n'{0}'"), FText::FromName(InAssetData.ObjectPath));
+					return FText::Format(LOCTEXT("MenuDescTooltipFormat_Evaluate", "Evaluate\n'{0}'"), FText::FromString(InAssetData.GetObjectPathString()));
 				}
 			}
 			else
diff --git a/Engine/Source/Editor/AnimGraph/Private/AnimGraphNode_SequencePlayer.cpp b/Engine/Source/Editor/AnimGraph/Private/AnimGraphNode_SequencePlayer.cpp
--- a/Engine/Source/Editor/AnimGraph/Private/AnimGraphNode_SequencePlayer.cpp
+++ b/Engine/Source/Editor/AnimGraph/Private/AnimGraphNode_SequencePlayer.cpp
@@ -118,11 +118,11 @@ void UAnimGraphNode_SequencePlayer::GetMenuActions(FBlueprintActionDatabaseRegis
 				const FString TagValue = InAssetData.GetTagValueRef<FString>(GET_MEMBER_NAME_CHECKED(UAnimSequence, AdditiveAnimType));
 				if(const bool bKnownToBeAdditive = (!TagValue.IsEmpty() && !TagValue.Equals(TEXT("AAT_None"))))
 				{
-					return FText::Format(LOCTEXT("MenuDescTooltipFormat_PlayAdded", "Play (additive)\n'{0}'"), FText::FromName(InAssetData.ObjectPath));
+					return FText::Format(LOCTEXT("MenuDescTooltipFormat_PlayAdded", "Play (additive)\n'{0}'"), FText::FromString(InAssetData.GetObjectPathString()));
 				}
 				else
 				{
-					return FText::Format(LOCTEXT("MenuDescTooltipFormat_Play", "Play\n'{0}'"), FText::FromName(InAssetData.ObjectPath));
+					return FText::Format(LOCTEXT("MenuDescTooltipFormat_Play", "Play\n'{0}'"), FText::FromString(InAssetData.GetObjectPathString()));
 				}
 			}
 			else
