Fix graph template compilation crash when referencing animation assets (sequences, blend spaces, etc)

#tests Graph template compilation

diff --git a/Engine/Source/Editor/AnimGraph/Private/AnimGraphNode_AssetPlayerBase.cpp b/Engine/Source/Editor/AnimGraph/Private/AnimGraphNode_AssetPlayerBase.cpp
--- a/Engine/Source/Editor/AnimGraph/Private/AnimGraphNode_AssetPlayerBase.cpp
+++ b/Engine/Source/Editor/AnimGraph/Private/AnimGraphNode_AssetPlayerBase.cpp
@@ -534,7 +534,7 @@ void UAnimGraphNode_AssetPlayerBase::ValidateAnimNodeDuringCompilationHelper(USk
 		USkeleton* AssetSkeleton = AssetToCheck->GetSkeleton();
 		
 		// if asset doesn't have a skeleton, it might be due to the asset no being not loaded yet
-		if(AssetSkeleton && !ForSkeleton->IsCompatible(AssetSkeleton))
+		if(AssetSkeleton && (!ForSkeleton || !ForSkeleton->IsCompatible(AssetSkeleton)))
 		{
 			MessageLog.Error(*FText::Format(LOCTEXT("IncompatibleSkeletonFormat", "@@ references {0} that uses an incompatible skeleton @@"), InAssetType->GetDisplayNameText()).ToString(), this, AssetSkeleton);
 		}
