Fix duplication of animation blueprints with custom property bindings producing internal compiler errors

diff --git a/Engine/Source/Editor/AnimGraph/Private/AnimGraphNode_CustomProperty.cpp b/Engine/Source/Editor/AnimGraph/Private/AnimGraphNode_CustomProperty.cpp
--- a/Engine/Source/Editor/AnimGraph/Private/AnimGraphNode_CustomProperty.cpp
+++ b/Engine/Source/Editor/AnimGraph/Private/AnimGraphNode_CustomProperty.cpp
@@ -172,7 +172,12 @@ FString UAnimGraphNode_CustomProperty::GetPinTargetVariableName(const UEdGraphPi
 
 FString UAnimGraphNode_CustomProperty::GetPinTargetVariableName(FName InPinName) const
 {
-	return TEXT("__CustomProperty_") + InPinName.ToString() + TEXT("_") + NodeGuid.ToString();
+	return GetPinTargetVariableNameBase(InPinName) + NodeGuid.ToString();
+}
+
+FString UAnimGraphNode_CustomProperty::GetPinTargetVariableNameBase(FName InPinName) const
+{
+	return TEXT("__CustomProperty_") + InPinName.ToString() + TEXT("_");
 }
 
 FText UAnimGraphNode_CustomProperty::GetPropertyTypeText(FProperty* Property)
@@ -384,4 +389,43 @@ void UAnimGraphNode_CustomProperty::PinConnectionListChanged(UEdGraphPin* Pin)
 	}
 }
 
+void UAnimGraphNode_CustomProperty::PostDuplicate(bool bDuplicateForPIE)
+{
+	Super::PostDuplicate(bDuplicateForPIE);
+	
+	const FString NewGUID = NodeGuid.ToString();
+	
+	TMap<FName, FAnimGraphNodePropertyBinding> NewBindings;
+	
+	// Update any bindings to point at the newly-generated property name
+	for(const TPair<FName, FAnimGraphNodePropertyBinding>& BindingPair : PropertyBindings)
+	{
+		FString BindingNameString = BindingPair.Key.ToString();
+		FAnimGraphNodePropertyBinding Binding = BindingPair.Value;
+
+		// Check pins
+		for (UEdGraphPin* Pin : Pins)
+		{
+			FString PinTargetVariableNameBase = GetPinTargetVariableNameBase(Pin->GetFName());
+			if(BindingNameString.StartsWith(PinTargetVariableNameBase))
+			{
+				// Replace GUID with this node's if it differs
+				FString OldGUID = BindingNameString.RightChop(PinTargetVariableNameBase.Len());
+				if(OldGUID != NewGUID)
+				{
+					FName PropertyName = *(BindingNameString.Replace(*OldGUID, *NewGUID));
+					Binding.PropertyName = PropertyName;
+					NewBindings.Add(PropertyName, Binding);
+				}
+				break;
+			}
+		}
+	}
+
+	if(NewBindings.Num() > 0)
+	{
+		PropertyBindings = NewBindings;
+	}
+}
+
 #undef LOCTEXT_NAMESPACE
\ No newline at end of file
diff --git a/Engine/Source/Editor/AnimGraph/Public/AnimGraphNode_CustomProperty.h b/Engine/Source/Editor/AnimGraph/Public/AnimGraphNode_CustomProperty.h
--- a/Engine/Source/Editor/AnimGraph/Public/AnimGraphNode_CustomProperty.h
+++ b/Engine/Source/Editor/AnimGraph/Public/AnimGraphNode_CustomProperty.h
@@ -25,6 +25,7 @@ public:
 	// UObject interface
 	virtual void Serialize(FArchive& Ar) override;
 	virtual void PostEditChangeProperty(FPropertyChangedEvent& PropertyChangedEvent) override;
+	virtual void PostDuplicate(bool bDuplicateForPIE) override;
 	
 	// IClassVariableCreator interface
 	virtual void CreateClassVariablesFromBlueprint(IAnimBlueprintVariableCreationContext& InCreationContext) override;
@@ -99,4 +100,7 @@ protected:
 	
 	// Sets the visibility of the specified pin, reconstructs the node if it changes
 	void SetCustomPinVisibility(bool bInVisible, int32 InOptionalPinIndex);
+
+	// Helper function for GetPinTargetVariableName
+	FString GetPinTargetVariableNameBase(FName InPinName) const;
 };
