Fixed FAnimNode_LayeredBoneBlend BlendMasks getting out of sync with BlendPoses (fixes older assets with incorrect BlendMasks array size)
Provided by [REDACTED_PERSON]

diff --git a/Engine/Source/Editor/AnimGraph/Private/AnimGraphNode_LayeredBoneBlend.cpp b/Engine/Source/Editor/AnimGraph/Private/AnimGraphNode_LayeredBoneBlend.cpp
--- a/Engine/Source/Editor/AnimGraph/Private/AnimGraphNode_LayeredBoneBlend.cpp
+++ b/Engine/Source/Editor/AnimGraph/Private/AnimGraphNode_LayeredBoneBlend.cpp
@@ -9,6 +9,8 @@
 
 #include "DetailLayoutBuilder.h"
 #include "Kismet2/CompilerResultsLog.h"
+#include "UObject/UE5ReleaseStreamObjectVersion.h"
+
 /////////////////////////////////////////////////////
 // UAnimGraphNode_LayeredBoneBlend
 
@@ -193,4 +195,20 @@ void UAnimGraphNode_LayeredBoneBlend::ValidateAnimNodeDuringCompilation(class US
  		Node.RebuildPerBoneBlendWeights(ForSkeleton);
  	}
 }
+
+void UAnimGraphNode_LayeredBoneBlend::Serialize(FArchive& Ar)
+{
+	Super::Serialize(Ar);
+
+	Ar.UsingCustomVersion(FUE5ReleaseStreamObjectVersion::GUID);
+
+	if (Ar.IsLoading() && Ar.CustomVer(FUE5ReleaseStreamObjectVersion::GUID) < FUE5ReleaseStreamObjectVersion::AnimLayeredBoneBlendMasks)
+	{
+		if (Node.BlendMode == ELayeredBoneBlendMode::BlendMask && Node.BlendMasks.Num() != Node.BlendPoses.Num())
+		{
+			Node.BlendMasks.SetNum(Node.BlendPoses.Num());
+		}
+	}
+}
+
 #undef LOCTEXT_NAMESPACE
diff --git a/Engine/Source/Editor/AnimGraph/Public/AnimGraphNode_LayeredBoneBlend.h b/Engine/Source/Editor/AnimGraph/Public/AnimGraphNode_LayeredBoneBlend.h
--- a/Engine/Source/Editor/AnimGraph/Public/AnimGraphNode_LayeredBoneBlend.h
+++ b/Engine/Source/Editor/AnimGraph/Public/AnimGraphNode_LayeredBoneBlend.h
@@ -21,6 +21,10 @@ class UAnimGraphNode_LayeredBoneBlend : public UAnimGraphNode_BlendListBase
 	ANIMGRAPH_API virtual void AddPinToBlendByFilter();
 	ANIMGRAPH_API virtual void RemovePinFromBlendByFilter(UEdGraphPin* Pin);
 
+	// UObject interface
+	virtual void Serialize(FArchive& Ar) override;
+	// End of UObject interface
+
 	// UEdGraphNode interface
 	virtual FLinearColor GetNodeTitleColor() const override;
 	virtual FText GetTooltipText() const override;
diff --git a/Engine/Source/Runtime/Core/Public/UObject/UE5ReleaseStreamObjectVersion.h b/Engine/Source/Runtime/Core/Public/UObject/UE5ReleaseStreamObjectVersion.h
--- a/Engine/Source/Runtime/Core/Public/UObject/UE5ReleaseStreamObjectVersion.h
+++ b/Engine/Source/Runtime/Core/Public/UObject/UE5ReleaseStreamObjectVersion.h
@@ -121,6 +121,9 @@ struct CORE_API FUE5ReleaseStreamObjectVersion
 		// For Blueprint real numbers, ensure that legacy float data is serialized as single-precision
 		SerializeFloatPinDefaultValuesAsSinglePrecision,
 
+		// Upgrade the BlendMasks array in existing LayeredBoneBlend nodes
+		AnimLayeredBoneBlendMasks,
+
 		// -----<new versions can be added above this line>-------------------------------------------------
 		VersionPlusOne,
 		LatestVersion = VersionPlusOne - 1
