Compile error fix for missing pose watch include

diff --git a/Engine/Source/Editor/AnimGraph/Public/AnimGraphNode_RigidBody.h b/Engine/Source/Editor/AnimGraph/Public/AnimGraphNode_RigidBody.h
--- a/Engine/Source/Editor/AnimGraph/Public/AnimGraphNode_RigidBody.h
+++ b/Engine/Source/Editor/AnimGraph/Public/AnimGraphNode_RigidBody.h
@@ -4,6 +4,7 @@
 
 #include "AnimGraphNode_SkeletalControlBase.h"
 #include "BoneControllers/AnimNode_RigidBody.h"
+#include "Engine/PoseWatch.h"
 #include "AnimGraphNode_RigidBody.generated.h"
 
 UCLASS(MinimalAPI, meta=(Keywords = "Simulate Rigid Body Physics Ragdoll"))
