Prevent editing of input parameters on input pose nodes in the main anim graph

Only other layers are allowed pin-based inputs. The legacy linked graph system uses class properties.

diff --git a/Engine/Source/Editor/AnimGraph/Private/AnimGraphNode_LinkedInputPose.cpp b/Engine/Source/Editor/AnimGraph/Private/AnimGraphNode_LinkedInputPose.cpp
--- a/Engine/Source/Editor/AnimGraph/Private/AnimGraphNode_LinkedInputPose.cpp
+++ b/Engine/Source/Editor/AnimGraph/Private/AnimGraphNode_LinkedInputPose.cpp
@@ -506,8 +506,17 @@ void UAnimGraphNode_LinkedInputPose::CustomizeDetails(IDetailLayoutBuilder& Deta
 		MakeNameWidget(DetailBuilder)
 	];
 
-	InputsCategoryBuilder.AddProperty(GET_MEMBER_NAME_CHECKED(UAnimGraphNode_LinkedInputPose, Inputs), GetClass())
-		.ShouldAutoExpand(true);
+	UEdGraph* Graph = GetGraph();
+	if(Graph && Graph->GetFName() != UEdGraphSchema_K2::GN_AnimGraph)
+	{
+		InputsCategoryBuilder.AddProperty(GET_MEMBER_NAME_CHECKED(UAnimGraphNode_LinkedInputPose, Inputs), GetClass())
+			.ShouldAutoExpand(true);
+	}
+	else
+	{
+		TSharedPtr<IPropertyHandle> InputsPropertyHandle = DetailBuilder.GetProperty(GET_MEMBER_NAME_CHECKED(UAnimGraphNode_LinkedInputPose, Inputs), GetClass());
+		InputsPropertyHandle->MarkHiddenByCustomization();
+	}
 }
 
 void UAnimGraphNode_LinkedInputPose::OnCopyTermDefaultsToDefaultObject(IAnimBlueprintCopyTermDefaultsContext& InCompilationContext, IAnimBlueprintNodeCopyTermDefaultsContext& InPerNodeContext, IAnimBlueprintGeneratedClassCompiledData& OutCompiledData)
