Fix for copy-pasting nodes dependant upon source state machine state.

diff --git a/Engine/Source/Editor/AnimGraph/Private/K2Node_AnimGetter.cpp b/Engine/Source/Editor/AnimGraph/Private/K2Node_AnimGetter.cpp
--- a/Engine/Source/Editor/AnimGraph/Private/K2Node_AnimGetter.cpp
+++ b/Engine/Source/Editor/AnimGraph/Private/K2Node_AnimGetter.cpp
@@ -13,6 +13,7 @@
 #include "AnimationCustomTransitionSchema.h"
 #include "AnimationGraphSchema.h"
 #include "AnimationStateMachineGraph.h"
+#include "AnimationTransitionGraph.h"
 #include "AnimationTransitionSchema.h"
 #include "BlueprintActionDatabaseRegistrar.h"
 #include "BlueprintNodeBinder.h"
@@ -55,7 +56,9 @@ void UK2Node_AnimGetter::PostPasteNode()
 {
 	Super::PostPasteNode();
 
+	RestoreStateMachineState();
 	RestoreStateMachineNode();
+	UpdateCachedTitle();
 }
 
 void UK2Node_AnimGetter::AllocateDefaultPins()
@@ -153,12 +156,11 @@ void UK2Node_AnimGetter::GetMenuActions(FBlueprintActionDatabaseRegistrar& Actio
 						// Should always succeed
 						if(UAnimationAsset* NodeAsset = AssetNode->GetAnimationAsset())
 						{
-							FText Title = FText::Format(LOCTEXT("NodeTitle", "{0} ({1})"), Getter->GetDisplayNameText(), FText::FromString(*NodeAsset->GetName()));
 							Params.SourceNode = AssetNode;
-							Params.CachedTitle = Title;
+							UpdateCachedTitle(Params);
 
 							UBlueprintNodeSpawner* Spawner = UBlueprintNodeSpawner::Create(UK2Node_AnimGetter::StaticClass(), /*AssetNode->GetGraph()*/nullptr, UBlueprintNodeSpawner::FCustomizeNodeDelegate::CreateUObject(const_cast<UK2Node_AnimGetter*>(this), &UK2Node_AnimGetter::PostSpawnNodeSetup, Params));
-							Spawner->DynamicUiSignatureGetter = UBlueprintNodeSpawner::FUiSpecOverrideDelegate::CreateStatic(UiSpecOverride, Title);
+							Spawner->DynamicUiSignatureGetter = UBlueprintNodeSpawner::FUiSpecOverrideDelegate::CreateStatic(UiSpecOverride, Params.CachedTitle);
 							ActionRegistrar.AddBlueprintAction(AnimBlueprint, Spawner);
 						}
 					}
@@ -179,12 +181,11 @@ void UK2Node_AnimGetter::GetMenuActions(FBlueprintActionDatabaseRegistrar& Actio
 								}
 							}
 
-							FText Title = FText::Format(LOCTEXT("NodeTitle", "{0} ({1})"), Getter->GetDisplayNameText(), StateNode->GetNodeTitle(ENodeTitleType::ListView));
 							Params.SourceStateNode = StateNode;
-							Params.CachedTitle = Title;
+							UpdateCachedTitle(Params);
 
 							UBlueprintNodeSpawner* Spawner = UBlueprintNodeSpawner::Create(UK2Node_AnimGetter::StaticClass(), /*StateNode->GetGraph()*/nullptr, UBlueprintNodeSpawner::FCustomizeNodeDelegate::CreateUObject(const_cast<UK2Node_AnimGetter*>(this), &UK2Node_AnimGetter::PostSpawnNodeSetup, Params));
-							Spawner->DynamicUiSignatureGetter = UBlueprintNodeSpawner::FUiSpecOverrideDelegate::CreateStatic(UiSpecOverride, Title);
+							Spawner->DynamicUiSignatureGetter = UBlueprintNodeSpawner::FUiSpecOverrideDelegate::CreateStatic(UiSpecOverride, Params.CachedTitle);
 							ActionRegistrar.AddBlueprintAction(AnimBlueprint, Spawner);
 						}
 					}
@@ -201,12 +202,11 @@ void UK2Node_AnimGetter::GetMenuActions(FBlueprintActionDatabaseRegistrar& Actio
 								}
 							}
 
-							FText Title = FText::Format(LOCTEXT("NodeTitle", "{0} ({1})"), Getter->GetDisplayNameText(), TransitionNode->GetNodeTitle(ENodeTitleType::ListView));
 							Params.SourceStateNode = TransitionNode;
-							Params.CachedTitle = Title;
+							UpdateCachedTitle(Params);
 
 							UBlueprintNodeSpawner* Spawner = UBlueprintNodeSpawner::Create(UK2Node_AnimGetter::StaticClass(), /*TransitionNode->GetGraph()*/nullptr, UBlueprintNodeSpawner::FCustomizeNodeDelegate::CreateUObject(const_cast<UK2Node_AnimGetter*>(this), &UK2Node_AnimGetter::PostSpawnNodeSetup, Params));
-							Spawner->DynamicUiSignatureGetter = UBlueprintNodeSpawner::FUiSpecOverrideDelegate::CreateStatic(UiSpecOverride, Title);
+							Spawner->DynamicUiSignatureGetter = UBlueprintNodeSpawner::FUiSpecOverrideDelegate::CreateStatic(UiSpecOverride, Params.CachedTitle);
 							ActionRegistrar.AddBlueprintAction(AnimBlueprint, Spawner);
 						}
 					}
@@ -215,12 +215,11 @@ void UK2Node_AnimGetter::GetMenuActions(FBlueprintActionDatabaseRegistrar& Actio
 						// Only requires the state machine
 						for(UAnimGraphNode_StateMachine* MachineNode : MachineNodes)
 						{
-							FText Title = FText::Format(LOCTEXT("NodeTitle", "{0} ({1})"), Getter->GetDisplayNameText(), MachineNode->GetNodeTitle(ENodeTitleType::ListView));
 							Params.SourceNode = MachineNode;
-							Params.CachedTitle = Title;
+							UpdateCachedTitle(Params);
 
 							UBlueprintNodeSpawner* Spawner = UBlueprintNodeSpawner::Create(UK2Node_AnimGetter::StaticClass(), /*MachineNode*/nullptr, UBlueprintNodeSpawner::FCustomizeNodeDelegate::CreateUObject(const_cast<UK2Node_AnimGetter*>(this), &UK2Node_AnimGetter::PostSpawnNodeSetup, Params));
-							Spawner->DynamicUiSignatureGetter = UBlueprintNodeSpawner::FUiSpecOverrideDelegate::CreateStatic(UiSpecOverride, Title);
+							Spawner->DynamicUiSignatureGetter = UBlueprintNodeSpawner::FUiSpecOverrideDelegate::CreateStatic(UiSpecOverride, Params.CachedTitle);
 							ActionRegistrar.AddBlueprintAction(AnimBlueprint, Spawner);
 						}
 					}
@@ -228,11 +227,10 @@ void UK2Node_AnimGetter::GetMenuActions(FBlueprintActionDatabaseRegistrar& Actio
 				else
 				{
 					// Doesn't operate on a node, only need one entry
-					FText Title = FText::Format(LOCTEXT("NodeTitleNoNode", "{0}"), Getter->GetDisplayNameText());
-					Params.CachedTitle = Title;
+					UpdateCachedTitle(Params);
 
 					UBlueprintNodeSpawner* Spawner = UBlueprintNodeSpawner::Create(UK2Node_AnimGetter::StaticClass(), nullptr, UBlueprintNodeSpawner::FCustomizeNodeDelegate::CreateUObject(const_cast<UK2Node_AnimGetter*>(this), &UK2Node_AnimGetter::PostSpawnNodeSetup, Params));
-					Spawner->DynamicUiSignatureGetter = UBlueprintNodeSpawner::FUiSpecOverrideDelegate::CreateStatic(UiSpecOverride, Title);
+					Spawner->DynamicUiSignatureGetter = UBlueprintNodeSpawner::FUiSpecOverrideDelegate::CreateStatic(UiSpecOverride, Params.CachedTitle);
 					ActionRegistrar.AddBlueprintAction(AnimBlueprint, Spawner);
 				}
 			}
@@ -346,6 +344,20 @@ bool UK2Node_AnimGetter::IsActionFilteredOut(FBlueprintActionFilter const& Filte
 	return true;
 }
 
+void UK2Node_AnimGetter::RestoreStateMachineState()
+{
+	if (SourceStateNode)
+	{
+		if (UAnimationTransitionGraph* TransitionGraph = Cast<UAnimationTransitionGraph>(GetOuter()))
+		{
+			if (UAnimStateTransitionNode* StateTransitionNode = Cast<UAnimStateTransitionNode>(TransitionGraph->GetOuter()))
+			{
+				SourceStateNode = StateTransitionNode->GetPreviousState();
+			}
+		}
+	}
+}
+
 void UK2Node_AnimGetter::RestoreStateMachineNode()
 {
 	if(SourceStateNode)
@@ -361,7 +373,7 @@ void UK2Node_AnimGetter::RestoreStateMachineNode()
 	}
 }
 
-bool UK2Node_AnimGetter::GetterRequiresParameter(const UFunction* Getter, FString ParamName) const
+bool UK2Node_AnimGetter::GetterRequiresParameter(const UFunction* Getter, FString ParamName)
 {
 	bool bRequiresParameter = false;
 
@@ -427,6 +439,43 @@ bool UK2Node_AnimGetter::IsContextValidForSchema(const UEdGraphSchema* Schema) c
 	return false;
 }
 
+void UK2Node_AnimGetter::UpdateCachedTitle(FNodeSpawnData& SpawnData)
+{
+	SpawnData.CachedTitle = GenerateTitle((UFunction*)SpawnData.Getter, SpawnData.SourceStateNode, SpawnData.SourceNode);
+}
+
+FText UK2Node_AnimGetter::GenerateTitle(UFunction* Getter, UAnimStateNodeBase* SourceStateNode, UAnimGraphNode_Base* SourceNode)
+{
+	if (GetterRequiresParameter(Getter, TEXT("AssetPlayerIndex")))
+	{
+		// Should always succeed
+		if (UAnimationAsset* NodeAsset = SourceNode->GetAnimationAsset())
+		{
+			return FText::Format(LOCTEXT("NodeTitle", "{0} ({1})"), Getter->GetDisplayNameText(), FText::FromString(*NodeAsset->GetName()));
+		}
+	}
+	else if (GetterRequiresParameter(Getter, TEXT("MachineIndex")))
+	{
+		if (GetterRequiresParameter(Getter, TEXT("StateIndex")) || GetterRequiresParameter(Getter, TEXT("TransitionIndex")))
+		{
+			return  FText::Format(LOCTEXT("NodeTitle", "{0} ({1})"), Getter->GetDisplayNameText(), SourceStateNode->GetNodeTitle(ENodeTitleType::ListView));
+		}
+		else
+		{
+			// Only requires the state machine
+			return FText::Format(LOCTEXT("NodeTitle", "{0} ({1})"), Getter->GetDisplayNameText(), SourceNode->GetNodeTitle(ENodeTitleType::ListView));
+		}
+	}
+
+	// Doesn't operate on a node, only need one entry
+	return Getter->GetDisplayNameText();
+}
+
+void UK2Node_AnimGetter::UpdateCachedTitle()
+{
+	CachedTitle = GenerateTitle(GetTargetFunction(), SourceStateNode, SourceNode);
+}
+
 FNodeSpawnData::FNodeSpawnData()
 	: SourceNode(nullptr)
 	, SourceStateNode(nullptr)
diff --git a/Engine/Source/Editor/AnimGraph/Public/K2Node_AnimGetter.h b/Engine/Source/Editor/AnimGraph/Public/K2Node_AnimGetter.h
--- a/Engine/Source/Editor/AnimGraph/Public/K2Node_AnimGetter.h
+++ b/Engine/Source/Editor/AnimGraph/Public/K2Node_AnimGetter.h
@@ -113,15 +113,27 @@ protected:
 	
 	//UFunction* GetSourceBlueprintFunction() const;
 
+	// Fixes the SourceStateNode to be the state of the node's owner
+	void RestoreStateMachineState();
+
 	// Fixes the SourceNode to be the state machine owner of SourceStateNode (if it is not null)
 	void RestoreStateMachineNode();
 
 	/** Returns whether or not the provided UFunction requires the named parameter */
-	bool GetterRequiresParameter(const UFunction* Getter, FString ParamName) const;
+	static bool GetterRequiresParameter(const UFunction* Getter, FString ParamName);
 
 	/** Checks the cached context strings to make sure this getter is valid within the provided schema */
 	bool IsContextValidForSchema(const UEdGraphSchema* Schema) const;
 
 	/** Passed to blueprint spawners to configure spawned nodes */
 	void PostSpawnNodeSetup(UEdGraphNode* NewNode, bool bIsTemplateNode, FNodeSpawnData SpawnData);
+
+	/** Recache the title, used if the source node or source state changes */
+	void UpdateCachedTitle();
+
+	/** Sets CachedTitle for a new node to be created */
+	static void UpdateCachedTitle(FNodeSpawnData& SpawnData);
+
+	/** Generates a title for the node based on its function and the context it is in */
+	static FText GenerateTitle(UFunction* Getter, UAnimStateNodeBase* SourceStateNode, UAnimGraphNode_Base* SourceNode);
 };
