Adding or removing an anim layer graph's input pose will now refresh instances without requiring a manual refresh

Also supports renaming and undo/redo of input pose additions/removals

diff --git a/Engine/Source/Editor/AnimGraph/Private/AnimGraphDetails.cpp b/Engine/Source/Editor/AnimGraph/Private/AnimGraphDetails.cpp
--- a/Engine/Source/Editor/AnimGraph/Private/AnimGraphDetails.cpp
+++ b/Engine/Source/Editor/AnimGraph/Private/AnimGraphDetails.cpp
@@ -19,6 +19,7 @@
 #include "SKismetInspector.h"
 #include "AnimationGraph.h"
 #include "AnimationGraphSchema.h"
+#include "AnimGraphNode_LinkedAnimLayer.h"
 #include "Kismet2/BlueprintEditorUtils.h"
 #include "Widgets/Input/SComboButton.h"
 #include "Widgets/Input/SEditableTextBox.h"
@@ -293,6 +294,10 @@ FReply FAnimGraphDetails::OnAddNewInputPoseClicked()
 
 	FEdGraphSchemaAction_K2NewNode::SpawnNode<UAnimGraphNode_LinkedInputPose>(Graph, NewNodePosition, EK2NewNodeFlags::None);
 
+	TSharedPtr<IAnimationBlueprintEditor> AnimBlueprintEditor = AnimBlueprintEditorPtr.Pin();
+	UBlueprint* Blueprint = AnimBlueprintEditor->GetBlueprintObj();
+	UAnimGraphNode_LinkedInputPose::ReconstructLayerNodes(Blueprint);
+
 	DetailLayoutBuilder->ForceRefreshDetails();
 	
 	return FReply::Handled();
@@ -311,8 +316,11 @@ FReply FAnimGraphDetails::OnRemoveInputPoseClicked(UAnimGraphNode_LinkedInputPos
 	}
 
 	TSharedPtr<IAnimationBlueprintEditor> AnimBlueprintEditor = AnimBlueprintEditorPtr.Pin();
-	FBlueprintEditorUtils::MarkBlueprintAsStructurallyModified(AnimBlueprintEditor->GetBlueprintObj());
+	UBlueprint* Blueprint = AnimBlueprintEditor->GetBlueprintObj();
+	FBlueprintEditorUtils::MarkBlueprintAsStructurallyModified(Blueprint);
 
+	UAnimGraphNode_LinkedInputPose::ReconstructLayerNodes(Blueprint);
+	
 	DetailLayoutBuilder->ForceRefreshDetails();
 
 	return FReply::Handled();
diff --git a/Engine/Source/Editor/AnimGraph/Private/AnimGraphNode_LinkedInputPose.cpp b/Engine/Source/Editor/AnimGraph/Private/AnimGraphNode_LinkedInputPose.cpp
--- a/Engine/Source/Editor/AnimGraph/Private/AnimGraphNode_LinkedInputPose.cpp
+++ b/Engine/Source/Editor/AnimGraph/Private/AnimGraphNode_LinkedInputPose.cpp
@@ -28,6 +28,7 @@
 #include "K2Node_VariableGet.h"
 #include "AnimBlueprintCompiler.h"
 #include "AnimGraphAttributes.h"
+#include "AnimGraphNode_LinkedAnimLayer.h"
 #include "IAnimBlueprintCopyTermDefaultsContext.h"
 
 #define LOCTEXT_NAMESPACE "LinkedInputPose"
@@ -106,6 +107,24 @@ void UAnimGraphNode_LinkedInputPose::ExpandNode(class FKismetCompilerContext& In
 	});
 }
 
+void UAnimGraphNode_LinkedInputPose::ReconstructLayerNodes(UBlueprint* InBlueprint)
+{
+	if(InBlueprint)
+	{
+		TArray<UAnimGraphNode_LinkedAnimLayer*> LinkedAnimLayers;
+		FBlueprintEditorUtils::GetAllNodesOfClass<UAnimGraphNode_LinkedAnimLayer>(InBlueprint, LinkedAnimLayers);
+
+		for(UAnimGraphNode_LinkedAnimLayer* LinkedAnimLayer : LinkedAnimLayers)
+		{
+			// Only reconstruct 'self' nodes in this manner - external layers will be rebuilt via the compilation machinery
+			if(LinkedAnimLayer->Node.Interface.Get() == nullptr)
+			{
+				LinkedAnimLayer->ReconstructNode();
+			}
+		}
+	}
+}
+
 void UAnimGraphNode_LinkedInputPose::PostEditChangeProperty(FPropertyChangedEvent& PropertyChangedEvent)
 {
 	Super::PostEditChangeProperty(PropertyChangedEvent);
@@ -121,6 +140,7 @@ void UAnimGraphNode_LinkedInputPose::PostEditChangeProperty(FPropertyChangedEven
 			ReconstructNode();
 			FBlueprintEditorUtils::MarkBlueprintAsStructurallyModified(GetAnimBlueprint());
 			ReconstructNode();
+			ReconstructLayerNodes(GetAnimBlueprint());
 		}
 	}
 }
diff --git a/Engine/Source/Editor/AnimGraph/Private/AnimationGraph.cpp b/Engine/Source/Editor/AnimGraph/Private/AnimationGraph.cpp
--- a/Engine/Source/Editor/AnimGraph/Private/AnimationGraph.cpp
+++ b/Engine/Source/Editor/AnimGraph/Private/AnimationGraph.cpp
@@ -2,6 +2,8 @@
 
 #include "AnimationGraph.h"
 
+#include "AnimGraphNode_LinkedInputPose.h"
+
 #define LOCTEXT_NAMESPACE "AnimationGraph"
 
 /////////////////////////////////////////////////////
@@ -28,4 +30,18 @@ void UAnimationGraph::GetGraphNodesOfClass(TSubclassOf<UAnimGraphNode_Base> Node
 	}
 }
 
+void UAnimationGraph::PostEditUndo()
+{
+	Super::PostEditUndo();
+
+	// We may have added/removed input nodes, which requires a refresh of any node that use them
+	// We need to defer this a tick to make sure we reconstruct nodes post-compile
+	GEditor->GetTimerManager()->SetTimerForNextTick(this, &UAnimationGraph::ReconstructLayerNodes);
+}
+
+void UAnimationGraph::ReconstructLayerNodes() const
+{
+	UAnimGraphNode_LinkedInputPose::ReconstructLayerNodes(GetTypedOuter<UBlueprint>());
+}
+
 #undef LOCTEXT_NAMESPACE
diff --git a/Engine/Source/Editor/AnimGraph/Public/AnimGraphNode_LinkedInputPose.h b/Engine/Source/Editor/AnimGraph/Public/AnimGraphNode_LinkedInputPose.h
--- a/Engine/Source/Editor/AnimGraph/Public/AnimGraphNode_LinkedInputPose.h
+++ b/Engine/Source/Editor/AnimGraph/Public/AnimGraphNode_LinkedInputPose.h
@@ -112,6 +112,9 @@ public:
 	void IterateFunctionParameters(TFunctionRef<void(const FName&, const FEdGraphPinType&)> InFunc) const;
 
 private:
+	friend class FAnimGraphDetails;
+	friend class UAnimationGraph;
+	
 	// Helper function for common code in AllocateDefaultPins and ReallocatePinsDuringReconstruction
 	void AllocatePinsInternal();
 
@@ -121,6 +124,9 @@ private:
 	/** Create pins from the stub function FunctionReference */
 	void CreatePinsFromStubFunction(const UFunction* Function);
 
+	/** Reconstruct any layer nodes in this BP post-edit */
+	static void ReconstructLayerNodes(UBlueprint* InBlueprint);
+	
 private:
 	/** UI helper functions */
 	void HandleInputPoseArrayChanged();
diff --git a/Engine/Source/Editor/AnimGraph/Public/AnimationGraph.h b/Engine/Source/Editor/AnimGraph/Public/AnimationGraph.h
--- a/Engine/Source/Editor/AnimGraph/Public/AnimationGraph.h
+++ b/Engine/Source/Editor/AnimGraph/Public/AnimationGraph.h
@@ -29,5 +29,12 @@ class ANIMGRAPH_API UAnimationGraph : public UEdGraph
 	/** Returns contained graph nodes of the specified (or child) class */
 	UFUNCTION(BlueprintCallable, Category=AnimationGraph)
 	void GetGraphNodesOfClass(TSubclassOf<UAnimGraphNode_Base> NodeClass, TArray<UAnimGraphNode_Base*>& GraphNodes, bool bIncludeChildClasses = true);
+
+private:
+	// UObject interface
+	virtual void PostEditUndo() override;
+
+	// Reconstruct layer nodes post-undo
+	void ReconstructLayerNodes() const;
 };
 
