Layered Blend Per Bone: Fix for blend masks not populating correctly on load

diff --git a/Engine/Source/Editor/AnimGraph/Private/AnimGraphNode_LayeredBoneBlend.cpp b/Engine/Source/Editor/AnimGraph/Private/AnimGraphNode_LayeredBoneBlend.cpp
--- a/Engine/Source/Editor/AnimGraph/Private/AnimGraphNode_LayeredBoneBlend.cpp
+++ b/Engine/Source/Editor/AnimGraph/Private/AnimGraphNode_LayeredBoneBlend.cpp
@@ -211,4 +211,20 @@ void UAnimGraphNode_LayeredBoneBlend::Serialize(FArchive& Ar)
 	}
 }
 
+void UAnimGraphNode_LayeredBoneBlend::PostLoad()
+{
+	Super::PostLoad();
+
+	// Post-load our blend masks, in case they've been pre-loaded, but haven't had their bone references initialized yet.
+	if (Node.BlendMode == ELayeredBoneBlendMode::BlendMask)
+	{
+		int32 NumBlendMasks = Node.BlendMasks.Num();
+		for (int32 MaskIndex = 0; MaskIndex < NumBlendMasks; ++MaskIndex)
+		{
+			UBlendProfile* BlendMask = Node.BlendMasks[MaskIndex];
+			BlendMask->ConditionalPostLoad();
+		}
+	}
+}
+
 #undef LOCTEXT_NAMESPACE
diff --git a/Engine/Source/Editor/AnimGraph/Public/AnimGraphNode_LayeredBoneBlend.h b/Engine/Source/Editor/AnimGraph/Public/AnimGraphNode_LayeredBoneBlend.h
--- a/Engine/Source/Editor/AnimGraph/Public/AnimGraphNode_LayeredBoneBlend.h
+++ b/Engine/Source/Editor/AnimGraph/Public/AnimGraphNode_LayeredBoneBlend.h
@@ -23,6 +23,7 @@ class UAnimGraphNode_LayeredBoneBlend : public UAnimGraphNode_BlendListBase
 
 	// UObject interface
 	virtual void Serialize(FArchive& Ar) override;
+	virtual void PostLoad() override;
 	// End of UObject interface
 
 	// UEdGraphNode interface
