Allow self layer nodes in linked anim layer graphs.
Added compilation check for linked layers circular references and indirectly nested linked layers

#tests Compile all project animation blueprints

diff --git a/Engine/Source/Editor/AnimGraph/Private/AnimGraphNode_LinkedAnimLayer.cpp b/Engine/Source/Editor/AnimGraph/Private/AnimGraphNode_LinkedAnimLayer.cpp
--- a/Engine/Source/Editor/AnimGraph/Private/AnimGraphNode_LinkedAnimLayer.cpp
+++ b/Engine/Source/Editor/AnimGraph/Private/AnimGraphNode_LinkedAnimLayer.cpp
@@ -234,48 +234,139 @@ void UAnimGraphNode_LinkedAnimLayer::ValidateAnimNodeDuringCompilation(USkeleton
 		{
 			UAnimGraphNode_LinkedAnimLayer* OriginalThis = Cast<UAnimGraphNode_LinkedAnimLayer>(MessageLog.FindSourceObject(this));
 
-			// check layer is only used once in this blueprint
-			auto CheckGraph = [this, OriginalThis, &MessageLog](const UEdGraph* InGraph)
-			{
-				TArray<UAnimGraphNode_LinkedAnimLayer*> LayerNodes;
-				InGraph->GetNodesOfClass(LayerNodes);
-				for(const UAnimGraphNode_LinkedAnimLayer* LayerNode : LayerNodes)
-				{
-					if(LayerNode != OriginalThis)
-					{
-						if(LayerNode->Node.Layer == Node.Layer)
-						{
-							MessageLog.Error(*FText::Format(LOCTEXT("DuplicateLayerError", "Linked anim layer node @@ also uses layer '{0}', layers can be used only once in an animation blueprint."), FText::FromName(Node.Layer)).ToString(), this);
-						}
-					}
-				}
-			};
-		
 			TArray<UEdGraph*> Graphs;
 			CurrentBlueprint->GetAllGraphs(Graphs);
 
 			auto ValidateOuterGraph = [this, OriginalThis, &MessageLog](const UEdGraph* InGraph)
 			{
 				static const FName DefaultAnimGraphName("AnimGraph");
-				if (InGraph->Nodes.Contains(OriginalThis))
+				if (!InGraph->IsA<UAnimationStateGraph>() && InGraph->GetFName() != DefaultAnimGraphName && InGraph->InterfaceGuid.IsValid() && OriginalThis->InterfaceGuid.IsValid())
 				{
-					if (!InGraph->IsA<UAnimationStateGraph>() && InGraph->GetFName() != DefaultAnimGraphName && InGraph->InterfaceGuid.IsValid())
+					MessageLog.Error(*FText::Format(LOCTEXT("NestedLayer", "Linked anim layer node @@ is part of Animation Layer Graph '{0}', linked layers cannot be nested."), FText::FromName(InGraph->GetFName())).ToString(), this);
+				}
+			};
+
+			// Gather all linked anim layer node instances in this blueprint 
+			TArray<UAnimGraphNode_LinkedAnimLayer*> AllLayerNodes;
+			for (const UEdGraph* Graph : Graphs)
+			{
+				Graph->GetNodesOfClass(AllLayerNodes);
+			}
+
+			// Check for duplicates
+			for (const UAnimGraphNode_LinkedAnimLayer* LayerNode : AllLayerNodes)
+			{
+				if (LayerNode != OriginalThis)
+				{
+					if (LayerNode->Node.Layer == Node.Layer)
 					{
-						MessageLog.Error(*FText::Format(LOCTEXT("NestedLayer", "Linked anim layer node @@ is part of Animation Layer Graph '{0}', linked layers cannot be nested."), FText::FromName(InGraph->GetFName())).ToString(), this);
+						MessageLog.Error(*FText::Format(LOCTEXT("DuplicateLayerError", "Linked anim layer node @@ is also used by graph '{0}', layers can be used only once in an animation blueprint."), FText::FromName(LayerNode->GetGraph()->GetFName())).ToString(), this);
 					}
 				}
-			};
-			
-			for(const UEdGraph* Graph : Graphs)
+			}
+
+			// Check for circular references and indirect nesting
+			for (const UEdGraph* Graph : Graphs)
 			{
-				CheckGraph(Graph);
-				ValidateOuterGraph(Graph);				
+				if (Graph->GetFName() == OriginalThis->Node.Layer)
+				{
+					TArray<const UEdGraph*> GraphStack;
+					bool bWithinLinkedAnimLayer = OriginalThis->InterfaceGuid.IsValid();
+
+					UEdGraph* OuterGraph = OriginalThis->GetGraph();
+					// If our outer graph is a linked layer interface function that has no instance in the blueprint, we need to validate it for nesting as it can still be instanciated through the application of a linked layer
+					if (OuterGraph->InterfaceGuid.IsValid())
+					{
+						bool bCheckOuter = false;
+						for (const UAnimGraphNode_LinkedAnimLayer* LayerNode : AllLayerNodes)
+						{
+							if (LayerNode->Node.Layer == OuterGraph->GetFName())
+							{
+								bCheckOuter = false;
+								break;
+							}
+						}
+
+						// Check outer graph for nested linked layers
+						if (bCheckOuter)
+						{
+							static const FName DefaultAnimGraphName("AnimGraph");
+							if (OuterGraph->GetFName() != DefaultAnimGraphName)
+							{
+								ValidateOuterGraph(OuterGraph);
+
+								// Add outer graph to graph stack
+								GraphStack.Add(OuterGraph);
+								// If outer graph is a linked layer interface, account for it to properly detect nesting
+								bWithinLinkedAnimLayer = bWithinLinkedAnimLayer || OuterGraph->InterfaceGuid.IsValid();
+							}
+						}
+					}
+					
+					ValidateCircularRefAndNesting(Graph, Graphs, GraphStack, bWithinLinkedAnimLayer, MessageLog);
+				}
 			}
 		}
 	}
 }
 
+void UAnimGraphNode_LinkedAnimLayer::ValidateCircularRefAndNesting(const UEdGraph* CurrentGraph, const TArray<UEdGraph*>& AllGraphs, TArray<const UEdGraph*> GraphStack, bool bWithinLinkedLayerGraph, FCompilerResultsLog& MessageLog)
+{
+	// Build graph chain string
+	auto BuildGraphChainString = [](const TArray<const UEdGraph*>& GraphStack)
+	{
+		TStringBuilder<1024> GraphChain;
+		bool bFirst = true;
+		for (const UEdGraph* Graph : GraphStack)
+		{
+			if (!bFirst)
+			{
+				GraphChain << TEXT("->");
+			}
+			GraphChain << *Graph->GetName();
+			bFirst = false;
+		}
+		return GraphChain.ToString();
+	};
+
+	// If a graph is already in the stack we have a circular reference
+	if (GraphStack.Contains(CurrentGraph))
+	{
+		// Add the redundant graph so the error is easier to understand
+		GraphStack.Add(CurrentGraph);
+		MessageLog.Error(*FText::Format(LOCTEXT("CircularLayerReference", "Anim layer node @@ in Graph '{0}' has a circular dependency '{1}'."), FText::FromName(Node.Layer), FText::FromString(BuildGraphChainString(GraphStack))).ToString(), this);
+		return;
+	}
 
+	GraphStack.Add(CurrentGraph);
+
+	// Find layer nodes and recursively check their graphs 
+	TArray<UAnimGraphNode_LinkedAnimLayer*> LayerNodes;
+	CurrentGraph->GetNodesOfClass(LayerNodes);
+	for (const UAnimGraphNode_LinkedAnimLayer* LayerNode : LayerNodes)
+	{
+		// Check for linked anim layer nesting
+		if (bWithinLinkedLayerGraph && LayerNode->InterfaceGuid.IsValid())
+		{
+			if (GraphStack.Num() == 1)
+			{
+				MessageLog.Error(*FText::Format(LOCTEXT("NestedLayer", "Linked anim layer node @@ is part of Animation Layer Graph '{0}', linked layers cannot be nested."), FText::FromName(CurrentGraph->GetFName())).ToString(), LayerNode);
+			}
+			else
+			{
+				MessageLog.Error(*FText::Format(LOCTEXT("NestedLayer", "Linked anim layer node @@ is indirectly nested inside another Linked Anim Layer Graph '{0}', linked layers cannot be nested."), FText::FromString(BuildGraphChainString(GraphStack))).ToString(), LayerNode);
+			}
+		}
+		for (const UEdGraph* Graph : AllGraphs)
+		{
+			if (Graph->GetFName() == LayerNode->Node.Layer)
+			{
+				ValidateCircularRefAndNesting(Graph, AllGraphs, GraphStack, bWithinLinkedLayerGraph || LayerNode->InterfaceGuid.IsValid(), MessageLog);
+			}
+
+		}
+	}
+};
 void UAnimGraphNode_LinkedAnimLayer::GetLinkTarget(UObject*& OutTargetGraph, UBlueprint*& OutTargetBlueprint) const
 {
 	OutTargetGraph = nullptr;
diff --git a/Engine/Source/Editor/AnimGraph/Public/AnimGraphNode_LinkedAnimLayer.h b/Engine/Source/Editor/AnimGraph/Public/AnimGraphNode_LinkedAnimLayer.h
--- a/Engine/Source/Editor/AnimGraph/Public/AnimGraphNode_LinkedAnimLayer.h
+++ b/Engine/Source/Editor/AnimGraph/Public/AnimGraphNode_LinkedAnimLayer.h
@@ -99,6 +99,9 @@ protected:
 	// Helper function to setup a newly spawned node
 	void SetupFromLayerId(FName InLayerId);
 	
+	// Used during compilation to check if the blueprint structure causes any circular references or nested linked layer nodes
+	void ValidateCircularRefAndNesting(const UEdGraph* CurrentGraph, const TArray<UEdGraph*>& AllGraphs, TArray<const UEdGraph*> GraphStack, bool bWithinLinkedLayerGraph, FCompilerResultsLog& MessageLog);
+
 	// Handle used to hook into object being debugged changing
 	FDelegateHandle SetObjectBeingDebuggedHandle;
 };
