Fix renaming issues with linked anim layers

Move linked layers and graphs to use FMemberReference to reference the stub function used to generate node pins

diff --git a/Engine/Source/Editor/AnimGraph/Private/AnimGraphNode_LinkedAnimGraph.cpp b/Engine/Source/Editor/AnimGraph/Private/AnimGraphNode_LinkedAnimGraph.cpp
--- a/Engine/Source/Editor/AnimGraph/Private/AnimGraphNode_LinkedAnimGraph.cpp
+++ b/Engine/Source/Editor/AnimGraph/Private/AnimGraphNode_LinkedAnimGraph.cpp
@@ -74,8 +74,11 @@ void UAnimGraphNode_LinkedAnimGraph::SetupFromAsset(const FAssetData& InAssetDat
 		{
 			UAnimBlueprint* AnimBlueprint = CastChecked<UAnimBlueprint>(InAssetData.GetAsset());
 			Node.InstanceClass = AnimBlueprint->GeneratedClass.Get();
-		}	
+		}
 	}
+
+	// Set up function reference
+	FunctionReference.SetExternalMember(UEdGraphSchema_K2::GN_AnimGraph, GetTargetClass());
 }
 
 void UAnimGraphNode_LinkedAnimGraph::GetMenuActions(FBlueprintActionDatabaseRegistrar& InActionRegistrar) const
diff --git a/Engine/Source/Editor/AnimGraph/Private/AnimGraphNode_LinkedAnimGraphBase.cpp b/Engine/Source/Editor/AnimGraph/Private/AnimGraphNode_LinkedAnimGraphBase.cpp
--- a/Engine/Source/Editor/AnimGraph/Private/AnimGraphNode_LinkedAnimGraphBase.cpp
+++ b/Engine/Source/Editor/AnimGraph/Private/AnimGraphNode_LinkedAnimGraphBase.cpp
@@ -30,6 +30,35 @@ namespace LinkedAnimGraphGraphNodeConstants
 	FLinearColor TitleColor(0.2f, 0.2f, 0.8f);
 }
 
+
+void UAnimGraphNode_LinkedAnimGraphBase::Serialize(FArchive& Ar)
+{
+	Super::Serialize(Ar);
+
+	Ar.UsingCustomVersion(FUE5ReleaseStreamObjectVersion::GUID);
+
+	if(Ar.CustomVer(FUE5ReleaseStreamObjectVersion::GUID) < FUE5ReleaseStreamObjectVersion::LinkedAnimGraphMemberReference)
+	{
+		// Upgrade name to member reference. 
+		// Note that if the name has changed underneath us, we still cant recover the GUID
+		// until this asset has its node refreshed & the node is resaved
+		const FAnimNode_LinkedAnimGraph& RuntimeNode = *GetLinkedAnimGraphNode();
+		const FName FunctionName = RuntimeNode.GetDynamicLinkFunctionName();
+		UClass* TargetClass = GetTargetSkeletonClass();
+		if(FunctionName != NAME_None)
+		{
+			if(TargetClass != nullptr)
+			{
+				FunctionReference.SetExternalMember(FunctionName, TargetClass);
+			}
+			else
+			{
+				FunctionReference.SetSelfMember(FunctionName);
+			}
+		}
+	}
+}
+
 void UAnimGraphNode_LinkedAnimGraphBase::AllocatePoseLinks()
 {
 	FAnimNode_LinkedAnimGraph& RuntimeNode = *GetLinkedAnimGraphNode();
@@ -213,28 +242,28 @@ void UAnimGraphNode_LinkedAnimGraphBase::CreateOutputPins()
 	
 	// Grab the SKELETON class here as when we are reconstructed during during BP compilation
 	// the full generated class is not yet present built.
-	UClass* TargetClass = GetTargetSkeletonClass();
-
-	if(TargetClass)
+	if(UClass* TargetClass = GetTargetSkeletonClass())
 	{
-		IAnimClassInterface* AnimClassInterface = IAnimClassInterface::GetFromClass(TargetClass);
-
-		const FAnimNode_LinkedAnimGraph& Node = *GetLinkedAnimGraphNode();
-
-		// Add any pose pins
-		for(const FAnimBlueprintFunction& AnimBlueprintFunction : AnimClassInterface->GetAnimBlueprintFunctions())
+		if(UFunction* Function = FunctionReference.ResolveMember<UFunction>(TargetClass))
 		{
-			if(AnimBlueprintFunction.Name == Node.GetDynamicLinkFunctionName())
+			// Could have re-resolved to a new member name via GUID 
+			const FAnimNode_LinkedAnimGraph& Node = *GetLinkedAnimGraphNode();
+			if(Node.GetDynamicLinkFunctionName() != FunctionReference.GetMemberName())
 			{
-				for(const FName& PoseName : AnimBlueprintFunction.InputPoseNames)
+				HandleFunctionReferenceChanged(FunctionReference.GetMemberName());
+			}
+			
+			Function = FBlueprintEditorUtils::GetMostUpToDateFunction(Function);
+			
+			IterateFunctionParameters(Function, [this](FName InName, FEdGraphPinType InPinType)
+			{
+				if(UAnimationGraphSchema::IsPosePin(InPinType))
 				{
-					UEdGraphPin* NewPin = CreatePin(EEdGraphPinDirection::EGPD_Input, UAnimationGraphSchema::MakeLocalSpacePosePin(), PoseName);
-					NewPin->PinFriendlyName = FText::FromName(PoseName);
-					CustomizePinData(NewPin, PoseName, INDEX_NONE);
+					UEdGraphPin* NewPin = CreatePin(EEdGraphPinDirection::EGPD_Input, UAnimationGraphSchema::MakeLocalSpacePosePin(), InName);
+					NewPin->PinFriendlyName = FText::FromName(InName);
+					CustomizePinData(NewPin, InName, INDEX_NONE);
 				}
-
-				break;
-			}
+			});
 		}
 	}
 }
@@ -250,6 +279,26 @@ void UAnimGraphNode_LinkedAnimGraphBase::PostEditChangeProperty(FPropertyChanged
 	{
 		if (IsStructuralProperty(ChangedProperty))
 		{
+			// Set function reference if node structure changes
+			const FAnimNode_LinkedAnimGraph& Node = *GetLinkedAnimGraphNode();
+			UClass* TargetClass = GetTargetSkeletonClass();
+			FName FunctionName = Node.GetDynamicLinkFunctionName();
+			if(FunctionName != NAME_None)
+			{
+				if(TargetClass != nullptr)
+				{
+					FunctionReference.SetExternalMember(FunctionName, TargetClass);
+				}
+				else
+				{
+					FunctionReference.SetSelfMember(FunctionName);
+				}
+			}
+			else
+			{
+				FunctionReference = FMemberReference();
+			}
+
 			bRequiresNodeReconstruct = true;
 		}
 	}
@@ -529,4 +578,37 @@ TSharedPtr<FEdGraphSchemaAction> UAnimGraphNode_LinkedAnimGraphBase::GetEventNod
 	return NodeAction;
 }
 
+
+void UAnimGraphNode_LinkedAnimGraphBase::IterateFunctionParameters(UFunction* InFunction, TFunctionRef<void(const FName&, const FEdGraphPinType&)> InFunc) const
+{
+	const UEdGraphSchema_K2* K2Schema = GetDefault<UEdGraphSchema_K2>();
+
+	auto IterateParam = [K2Schema, &InFunc](FProperty* InParam, bool bInPoses)
+	{
+		const bool bIsFunctionInput = !InParam->HasAnyPropertyFlags(CPF_OutParm) || InParam->HasAnyPropertyFlags(CPF_ReferenceParm);
+		if (bIsFunctionInput)
+		{
+			FEdGraphPinType PinType;
+			if(K2Schema->ConvertPropertyToPinType(InParam, PinType))
+			{
+				if(UAnimationGraphSchema::IsPosePin(PinType) == bInPoses)
+				{
+					InFunc(InParam->GetFName(), PinType);
+				}
+			}
+		}
+	};
+
+	// Iterate poses first, then params
+	for (TFieldIterator<FProperty> PropIt(InFunction); PropIt && (PropIt->PropertyFlags & CPF_Parm); ++PropIt)
+	{
+		IterateParam(*PropIt, true);
+	}
+
+	for (TFieldIterator<FProperty> PropIt(InFunction); PropIt && (PropIt->PropertyFlags & CPF_Parm); ++PropIt)
+	{
+		IterateParam(*PropIt, false);
+	}
+}
+
 #undef LOCTEXT_NAMESPACE
diff --git a/Engine/Source/Editor/AnimGraph/Private/AnimGraphNode_LinkedAnimLayer.cpp b/Engine/Source/Editor/AnimGraph/Private/AnimGraphNode_LinkedAnimLayer.cpp
--- a/Engine/Source/Editor/AnimGraph/Private/AnimGraphNode_LinkedAnimLayer.cpp
+++ b/Engine/Source/Editor/AnimGraph/Private/AnimGraphNode_LinkedAnimLayer.cpp
@@ -114,7 +114,7 @@ FText UAnimGraphNode_LinkedAnimLayer::GetNodeTitle(ENodeTitleType::Type TitleTyp
 		FFormatNamedArguments Args;
 		Args.Add(TEXT("NodeType"), bIsSelf ? LOCTEXT("SelfLayerNodeTitle", "Anim Layer (self)") : DefaultNodeTitle);
 		Args.Add(TEXT("TargetClass"), bIsSelf ? LOCTEXT("ClassSelf", "Self") : FText::FromString(TargetAnimBlueprintInterface->GetName()));
-		Args.Add(TEXT("Layer"), (Node.Layer == NAME_None) ? LOCTEXT("LayerNone", "None") : FText::FromName(Node.Layer));
+		Args.Add(TEXT("Layer"), (GetLayerName() == NAME_None) ? LOCTEXT("LayerNone", "None") : FText::FromName(GetLayerName()));
 
 		if (TitleType == ENodeTitleType::ListView)
 		{
@@ -157,7 +157,7 @@ void UAnimGraphNode_LinkedAnimLayer::ValidateAnimNodeDuringCompilation(USkeleton
 {
 	Super::ValidateAnimNodeDuringCompilation(ForSkeleton, MessageLog);
 
-	if(Node.Layer == NAME_None)
+	if(GetLayerName() == NAME_None)
 	{
 		MessageLog.Error(*LOCTEXT("NoLayerError", "Linked anim layer node @@ does not specify a layer.").ToString(), this);
 	}
@@ -218,7 +218,7 @@ void UAnimGraphNode_LinkedAnimLayer::ValidateAnimNodeDuringCompilation(USkeleton
 			IAnimClassInterface* AnimClassInterface = IAnimClassInterface::GetFromClass(TargetClass);
 			for(const FAnimBlueprintFunction& AnimBlueprintFunction : AnimClassInterface->GetAnimBlueprintFunctions())
 			{
-				if(AnimBlueprintFunction.Name == Node.Layer)
+				if(AnimBlueprintFunction.Name == GetLayerName())
 				{
 					bFoundFunction = true;
 				}
@@ -226,7 +226,7 @@ void UAnimGraphNode_LinkedAnimLayer::ValidateAnimNodeDuringCompilation(USkeleton
 
 			if(!bFoundFunction)
 			{
-				MessageLog.Error(*FText::Format(LOCTEXT("MissingLayerError", "Linked anim layer node @@ uses invalid layer '{0}'."), FText::FromName(Node.Layer)).ToString(), this);
+				MessageLog.Error(*FText::Format(LOCTEXT("MissingLayerError", "Linked anim layer node @@ uses invalid layer '{0}'."), FText::FromName(GetLayerName())).ToString(), this);
 			}
 		}
 
@@ -258,7 +258,7 @@ void UAnimGraphNode_LinkedAnimLayer::ValidateAnimNodeDuringCompilation(USkeleton
 			{
 				if (LayerNode != OriginalThis)
 				{
-					if (LayerNode->Node.Layer == Node.Layer)
+					if (LayerNode->GetLayerName() == GetLayerName())
 					{
 						MessageLog.Error(*FText::Format(LOCTEXT("DuplicateLayerError", "Linked anim layer node @@ is also used by graph '{0}', layers can be used only once in an animation blueprint."), FText::FromName(LayerNode->GetGraph()->GetFName())).ToString(), this);
 					}
@@ -382,7 +382,7 @@ void UAnimGraphNode_LinkedAnimLayer::GetLinkTarget(UObject*& OutTargetGraph, UBl
 			TArray<UEdGraph*> Graphs;
 			TargetAnimBlueprint->GetAllGraphs(Graphs);
 
-			UEdGraph** FoundGraph = Graphs.FindByPredicate([this](UEdGraph* InGraph) { return InGraph->GetFName() == Node.Layer; });
+			UEdGraph** FoundGraph = Graphs.FindByPredicate([this](UEdGraph* InGraph) { return InGraph->GetFName() == GetLayerName(); });
 			if (FoundGraph)
 			{
 				OutTargetBlueprint = TargetAnimBlueprint;
@@ -399,7 +399,7 @@ void UAnimGraphNode_LinkedAnimLayer::GetLinkTarget(UObject*& OutTargetGraph, UBl
 		TArray<UEdGraph*> Graphs;
 		GetBlueprint()->GetAllGraphs(Graphs);
 
-		UEdGraph** FoundGraph = Graphs.FindByPredicate([this](UEdGraph* InGraph) { return InGraph->GetFName() == Node.Layer; });
+		UEdGraph** FoundGraph = Graphs.FindByPredicate([this](UEdGraph* InGraph) { return InGraph->GetFName() == GetLayerName(); });
 		if (FoundGraph)
 		{
 			OutTargetBlueprint = nullptr;
@@ -524,7 +524,7 @@ void UAnimGraphNode_LinkedAnimLayer::CustomizeDetails(IDetailLayoutBuilder& Deta
 					PropertyCustomizationHelpers::MakePropertyComboBox(
 						LayerHandle, 
 						FOnGetPropertyComboBoxStrings::CreateUObject(this,  &UAnimGraphNode_LinkedAnimLayer::GetLayerNames),
-						FOnGetPropertyComboBoxValue::CreateUObject(this,  &UAnimGraphNode_LinkedAnimLayer::GetLayerName)
+						FOnGetPropertyComboBoxValue::CreateUObject(this,  &UAnimGraphNode_LinkedAnimLayer::GetLayerNameString)
 					)
 				]
 			]
@@ -595,7 +595,7 @@ bool UAnimGraphNode_LinkedAnimLayer::OnShouldFilterInstanceBlueprint(const FAsse
 		{
 			if(InterfaceDesc.Interface && InterfaceDesc.Interface->IsChildOf<UAnimLayerInterface>())
 			{
-				if(Node.Layer == NAME_None || InterfaceDesc.Interface->FindFunctionByName(Node.Layer))
+				if(GetLayerName() == NAME_None || InterfaceDesc.Interface->FindFunctionByName(GetLayerName()))
 				{
 					AnimInterfaces.Add(InterfaceDesc.Interface);
 				}
@@ -785,9 +785,9 @@ void UAnimGraphNode_LinkedAnimLayer::GetLayerNames(TArray<TSharedPtr<FString>>&
 	}
 }
 
-FString UAnimGraphNode_LinkedAnimLayer::GetLayerName() const
+FString UAnimGraphNode_LinkedAnimLayer::GetLayerNameString() const
 {
-	return Node.Layer.ToString();
+	return FunctionReference.GetMemberName().ToString();
 }
 
 bool UAnimGraphNode_LinkedAnimLayer::IsStructuralProperty(FProperty* InProperty) const
@@ -831,7 +831,7 @@ TSubclassOf<UInterface> UAnimGraphNode_LinkedAnimLayer::GetInterfaceForLayer() c
 		{
 			for(UEdGraph* InterfaceGraph : InterfaceDesc.Graphs)
 			{
-				if(InterfaceGraph->GetFName() == Node.Layer)
+				if(InterfaceGraph->GetFName() == GetLayerName())
 				{
 					return InterfaceDesc.Interface;
 				}
@@ -850,6 +850,26 @@ void UAnimGraphNode_LinkedAnimLayer::UpdateGuidForLayer()
 	}
 }
 
+void UAnimGraphNode_LinkedAnimLayer::SetLayerName(FName InName)
+{
+	Node.Layer = InName;
+
+	if(UClass* TargetClass = GetTargetClass())
+	{
+		FunctionReference.SetExternalMember(InName, TargetClass);
+	}
+	else
+	{
+		FunctionReference.SetSelfMember(InName);
+	}
+}
+
+FName UAnimGraphNode_LinkedAnimLayer::GetLayerName() const
+{
+	ensure(FunctionReference.GetMemberName() == Node.Layer);
+	return FunctionReference.GetMemberName();
+}
+
 FGuid UAnimGraphNode_LinkedAnimLayer::GetGuidForLayer() const
 {
 	if (UAnimBlueprint* CurrentBlueprint = Cast<UAnimBlueprint>(GetBlueprint()))
@@ -859,7 +879,7 @@ FGuid UAnimGraphNode_LinkedAnimLayer::GetGuidForLayer() const
 		{
 			for (UEdGraph* InterfaceGraph : InterfaceDesc.Graphs)
 			{
-				if (InterfaceGraph->GetFName() == Node.Layer)
+				if (InterfaceGraph->GetFName() == GetLayerName())
 				{
 					return InterfaceGraph->InterfaceGuid;
 				}
@@ -918,7 +938,7 @@ bool UAnimGraphNode_LinkedAnimLayer::HasValidNonSelfLayer() const
 			{
 				if(InterfaceDesc.Interface && InterfaceDesc.Interface->IsChildOf<UAnimLayerInterface>())
 				{
-					if(InterfaceDesc.Interface->FindFunctionByName(Node.Layer))
+					if(InterfaceDesc.Interface->FindFunctionByName(GetLayerName()))
 					{
 						return true;
 					}
@@ -964,6 +984,17 @@ void UAnimGraphNode_LinkedAnimLayer::SetupFromLayerId(FName InLayer)
 		// Self layers cannot have override implementations
 		Node.InstanceClass = nullptr;
 	}
+
+	// Set up function reference
+	UClass* TargetClass = GetTargetClass();
+	if(TargetClass)
+	{
+		FunctionReference.SetExternalMember(InLayer, TargetClass);
+	}
+	else
+	{
+		FunctionReference.SetSelfMember(InLayer);
+	}
 }
 
 void UAnimGraphNode_LinkedAnimLayer::GetMenuActions(FBlueprintActionDatabaseRegistrar& ActionRegistrar) const
@@ -1071,7 +1102,7 @@ bool UAnimGraphNode_LinkedAnimLayer::IsActionFilteredOut(class FBlueprintActionF
 				IAnimClassInterface* AnimClassInterface = IAnimClassInterface::GetFromClass(TargetClass);
 				for(const FAnimBlueprintFunction& AnimBlueprintFunction : AnimClassInterface->GetAnimBlueprintFunctions())
 				{
-					if(Node.Layer == AnimBlueprintFunction.Name)
+					if(GetLayerName() == AnimBlueprintFunction.Name)
 					{
 						bImplemented = true;
 						break;
@@ -1101,5 +1132,10 @@ const FAnimNode_LinkedAnimGraph* UAnimGraphNode_LinkedAnimLayer::GetLinkedAnimGr
 	return RuntimeLinkedAnimGraphNode ? RuntimeLinkedAnimGraphNode : &Node;
 }
 
+void UAnimGraphNode_LinkedAnimLayer::HandleFunctionReferenceChanged(FName InNewName)
+{
+	Node.Layer = InNewName;
+	ensure(FunctionReference.GetMemberName() == Node.Layer);
+}
 
 #undef LOCTEXT_NAMESPACE
diff --git a/Engine/Source/Editor/AnimGraph/Private/AnimationGraphSchema.cpp b/Engine/Source/Editor/AnimGraph/Private/AnimationGraphSchema.cpp
--- a/Engine/Source/Editor/AnimGraph/Private/AnimationGraphSchema.cpp
+++ b/Engine/Source/Editor/AnimGraph/Private/AnimationGraphSchema.cpp
@@ -1104,9 +1104,9 @@ void UAnimationGraphSchema::ConformAnimLayersByGuid(const UAnimBlueprint* InAnim
 				for (UEdGraph* InterfaceGraph : InterfaceGraphs)
 				{
 					// Check to see if GUID matches but name does not and update if so
-					if (InterfaceGraph->GraphGuid == LayerNode->InterfaceGuid && InterfaceGraph->GetFName() != LayerNode->Node.Layer)
+					if (InterfaceGraph->GraphGuid == LayerNode->InterfaceGuid && InterfaceGraph->GetFName() != LayerNode->GetLayerName())
 					{
-						LayerNode->Node.Layer = InterfaceGraph->GetFName();
+						LayerNode->SetLayerName(InterfaceGraph->GetFName());
 					}
 				}
 			}
diff --git a/Engine/Source/Editor/AnimGraph/Public/AnimGraphNode_LinkedAnimGraphBase.h b/Engine/Source/Editor/AnimGraph/Public/AnimGraphNode_LinkedAnimGraphBase.h
--- a/Engine/Source/Editor/AnimGraph/Public/AnimGraphNode_LinkedAnimGraphBase.h
+++ b/Engine/Source/Editor/AnimGraph/Public/AnimGraphNode_LinkedAnimGraphBase.h
@@ -24,6 +24,9 @@ class UAnimGraphNode_LinkedAnimGraphBase : public UAnimGraphNode_CustomProperty,
 	GENERATED_BODY()
 
 public:
+	// UObject interface
+	virtual void Serialize(FArchive& Ar) override;
+	
 	//~ Begin UEdGraphNode Interface.
 	virtual FLinearColor GetNodeTitleColor() const override;
 	virtual FSlateIcon GetIconAndTint(FLinearColor& OutColor) const override;
@@ -84,6 +87,16 @@ protected:
 
 	virtual FLinearColor GetDefaultNodeTitleColor() const;
 
+	// Handler for when the function reference gets re-resolved on node reconstruction
+	virtual void HandleFunctionReferenceChanged(FName InNewName) {}
+	
+	// Helper func for generating pins
+	void IterateFunctionParameters(UFunction* InFunction, TFunctionRef<void(const FName&, const FEdGraphPinType&)> InFunc) const;
+	
+	// Reference to the stub function that this node uses
+	UPROPERTY()
+	FMemberReference FunctionReference;
+	
 	// Skeleton name used for filtering unloaded classes 
 	FString SkeletonName;
 
diff --git a/Engine/Source/Editor/AnimGraph/Public/AnimGraphNode_LinkedAnimLayer.h b/Engine/Source/Editor/AnimGraph/Public/AnimGraphNode_LinkedAnimLayer.h
--- a/Engine/Source/Editor/AnimGraph/Public/AnimGraphNode_LinkedAnimLayer.h
+++ b/Engine/Source/Editor/AnimGraph/Public/AnimGraphNode_LinkedAnimLayer.h
@@ -53,6 +53,11 @@ public:
 	// Optionally updates layer GUID if it is invalid
 	void UpdateGuidForLayer();
 
+	// Sets the name of the layer we refer to
+	void SetLayerName(FName InName);
+
+	// Gets the name of the layer we refer to
+	FName GetLayerName() const;
 protected:
 
 	void GetLinkTarget(UObject* &OutTargetGraph, UBlueprint* &OutTargetBlueprint) const;
@@ -60,7 +65,7 @@ protected:
 	// ----- UI CALLBACKS ----- //
 	// Handlers for layer combo
 	void GetLayerNames(TArray<TSharedPtr<FString>>& OutStrings, TArray<TSharedPtr<SToolTip>>& OutToolTips, TArray<bool>& OutRestrictedItems);
-	FString GetLayerName() const;
+	FString GetLayerNameString() const;
 	void OnLayerChanged(IDetailLayoutBuilder* DetailBuilder);
 	bool HasAvailableLayers() const;
 	bool HasValidNonSelfLayer() const;
@@ -84,7 +89,8 @@ protected:
 	virtual FString GetCurrentInstanceBlueprintPath() const override;
 	virtual bool IsStructuralProperty(FProperty* InProperty) const override;
 	virtual FLinearColor GetDefaultNodeTitleColor() const override;
-
+	virtual void HandleFunctionReferenceChanged(FName InNewName) override;
+	
 	friend class FAnimationLayerDragDropAction;
 	
 	// Helper function to get the interface currently in use by the selected layer
diff --git a/Engine/Source/Runtime/Core/Public/UObject/UE5ReleaseStreamObjectVersion.h b/Engine/Source/Runtime/Core/Public/UObject/UE5ReleaseStreamObjectVersion.h
--- a/Engine/Source/Runtime/Core/Public/UObject/UE5ReleaseStreamObjectVersion.h
+++ b/Engine/Source/Runtime/Core/Public/UObject/UE5ReleaseStreamObjectVersion.h
@@ -133,6 +133,9 @@ struct CORE_API FUE5ReleaseStreamObjectVersion
 		// Removed font from FEditableTextBoxStyle, and added FTextBlockStyle instead.
 		RemoveDuplicatedStyleInfo,
 
+		// Added member reference to linked anim graphs
+		LinkedAnimGraphMemberReference,
+		
 		// -----<new versions can be added above this line>-------------------------------------------------
 		VersionPlusOne,
 		LatestVersion = VersionPlusOne - 1
