Fixed LOCTEXT conflict in LinkedAnimLayerNode

diff --git a/Engine/Source/Editor/AnimGraph/Private/AnimGraphNode_LinkedAnimLayer.cpp b/Engine/Source/Editor/AnimGraph/Private/AnimGraphNode_LinkedAnimLayer.cpp
--- a/Engine/Source/Editor/AnimGraph/Private/AnimGraphNode_LinkedAnimLayer.cpp
+++ b/Engine/Source/Editor/AnimGraph/Private/AnimGraphNode_LinkedAnimLayer.cpp
@@ -354,7 +354,7 @@ void UAnimGraphNode_LinkedAnimLayer::ValidateCircularRefAndNesting(const UEdGrap
 			}
 			else
 			{
-				MessageLog.Error(*FText::Format(LOCTEXT("NestedLayer", "Linked anim layer node @@ is indirectly nested inside another Linked Anim Layer Graph '{0}', linked layers cannot be nested."), FText::FromString(BuildGraphChainString(GraphStack))).ToString(), LayerNode);
+				MessageLog.Error(*FText::Format(LOCTEXT("IndirectlyNestedLayer", "Linked anim layer node @@ is indirectly nested inside another Linked Anim Layer Graph '{0}', linked layers cannot be nested."), FText::FromString(BuildGraphChainString(GraphStack))).ToString(), LayerNode);
 			}
 		}
 		for (const UEdGraph* Graph : AllGraphs)
