Fixed crash when loading a layered bone per blend node that has empty blend mask entries

diff --git a/Engine/Source/Editor/AnimGraph/Private/AnimGraphNode_LayeredBoneBlend.cpp b/Engine/Source/Editor/AnimGraph/Private/AnimGraphNode_LayeredBoneBlend.cpp
--- a/Engine/Source/Editor/AnimGraph/Private/AnimGraphNode_LayeredBoneBlend.cpp
+++ b/Engine/Source/Editor/AnimGraph/Private/AnimGraphNode_LayeredBoneBlend.cpp
@@ -221,8 +221,10 @@ void UAnimGraphNode_LayeredBoneBlend::PostLoad()
 		int32 NumBlendMasks = Node.BlendMasks.Num();
 		for (int32 MaskIndex = 0; MaskIndex < NumBlendMasks; ++MaskIndex)
 		{
-			UBlendProfile* BlendMask = Node.BlendMasks[MaskIndex];
-			BlendMask->ConditionalPostLoad();
+			if(UBlendProfile* BlendMask = Node.BlendMasks[MaskIndex])
+			{
+				BlendMask->ConditionalPostLoad();
+			}
 		}
 	}
 }
