Fixed crash when attempting to draw pose-watched anim nodes which use custom edit modes such as 'Two Bone IK' and 'Transform (Modify) Bones'.
* Only nodes which have been setup to have drawable elements should be drawn when pose watched, such as 'RigidBody'.
* Anim nodes with custom edit modes should only draw their widgets when explicitly selected or override IAnimNodeEditMode::SupportsPoseWatch

diff --git a/Engine/Source/Editor/AnimGraph/Private/AnimGraphNode_Base.cpp b/Engine/Source/Editor/AnimGraph/Private/AnimGraphNode_Base.cpp
--- a/Engine/Source/Editor/AnimGraph/Private/AnimGraphNode_Base.cpp
+++ b/Engine/Source/Editor/AnimGraph/Private/AnimGraphNode_Base.cpp
@@ -666,7 +666,18 @@ void UAnimGraphNode_Base::OnPoseWatchChanged(const bool IsPoseWatchEnabled, TObj
 			InModeTools.ActivateMode(ModeID);
 			if (FEdMode* EdMode = InModeTools.GetActiveMode(ModeID))
 			{
-				static_cast<IAnimNodeEditMode*>(EdMode)->RegisterPoseWatchedNode(this, InRuntimeNode);
+				const bool bSupportsPoseWatch =
+					   static_cast<IAnimNodeEditMode*>(EdMode)->SupportsPoseWatch()
+					|| EdMode->GetID() == AnimNodeEditModes::AnimNode;
+
+				if (bSupportsPoseWatch)
+				{
+					static_cast<IAnimNodeEditMode*>(EdMode)->RegisterPoseWatchedNode(this, InRuntimeNode);
+				}
+				else
+				{
+					InModeTools.DeactivateMode(ModeID);
+				}
 			}
 		}
 		else
diff --git a/Engine/Source/Editor/AnimGraph/Public/AnimNodeEditMode.h b/Engine/Source/Editor/AnimGraph/Public/AnimNodeEditMode.h
--- a/Engine/Source/Editor/AnimGraph/Public/AnimNodeEditMode.h
+++ b/Engine/Source/Editor/AnimGraph/Public/AnimNodeEditMode.h
@@ -51,6 +51,7 @@ public:
 	virtual void DoScale(FVector& InScale) override;
 	virtual void EnterMode(class UAnimGraphNode_Base* InEditorNode, struct FAnimNode_Base* InRuntimeNode) override;
 	virtual void ExitMode() override;
+	virtual bool SupportsPoseWatch() override { return false; };
 
 	/** IPersonaEditMode interface */
 	virtual bool GetCameraTarget(FSphere& OutTarget) const override;
diff --git a/Engine/Source/Editor/AnimGraph/Public/IAnimNodeEditMode.h b/Engine/Source/Editor/AnimGraph/Public/IAnimNodeEditMode.h
--- a/Engine/Source/Editor/AnimGraph/Public/IAnimNodeEditMode.h
+++ b/Engine/Source/Editor/AnimGraph/Public/IAnimNodeEditMode.h
@@ -43,6 +43,9 @@ public:
 	/** Called when exiting this edit mode */
 	virtual void ExitMode() = 0;
 
+	/** Called to determine whether this edit mode should be drawn when nodes edited by this edit mode are pose watched */
+	virtual bool SupportsPoseWatch() = 0;
+
 	/** Called when a Pose Watch is created on a node edited by this edit mode */
 	virtual void RegisterPoseWatchedNode(UAnimGraphNode_Base* InEditorNode, FAnimNode_Base* InRuntimeNode) = 0;
 };
