Truncation fixes for AnimationGraph and AnimationGraphRuntime module. Part of UE-166634 and UE-166633

diff --git a/Engine/Source/Editor/AnimGraph/Private/AnimGraphNode_ApplyLimits.cpp b/Engine/Source/Editor/AnimGraph/Private/AnimGraphNode_ApplyLimits.cpp
--- a/Engine/Source/Editor/AnimGraph/Private/AnimGraphNode_ApplyLimits.cpp
+++ b/Engine/Source/Editor/AnimGraph/Private/AnimGraphNode_ApplyLimits.cpp
@@ -54,15 +54,15 @@ void UAnimGraphNode_ApplyLimits::DrawAngularLimits(FPrimitiveDrawInterface* PDI,
 
 			FTransform XAxisConeTM(YAxis, XAxis ^ YAxis, XAxis, JointTransform.GetTranslation());
 			XAxisConeTM.SetRotation(FQuat(XAxis, FMath::DegreesToRadians(-Middle.X)) * XAxisConeTM.GetRotation());
-			DrawCone(PDI, FScaleMatrix(30.0f) * XAxisConeTM.ToMatrixWithScale(), FMath::DegreesToRadians(AngleRange.X / 2.0f), 0.0f, 24, false, FLinearColor::Red, GEngine->ConstraintLimitMaterialX->GetRenderProxy(), SDPG_World);
+			DrawCone(PDI, FScaleMatrix(30.0f) * XAxisConeTM.ToMatrixWithScale(), static_cast<float>(FMath::DegreesToRadians(AngleRange.X / 2.0)), 0.0f, 24, false, FLinearColor::Red, GEngine->ConstraintLimitMaterialX->GetRenderProxy(), SDPG_World);
 
 			FTransform YAxisConeTM(ZAxis, YAxis ^ ZAxis, YAxis, JointTransform.GetTranslation());
 			YAxisConeTM.SetRotation(FQuat(YAxis, FMath::DegreesToRadians(Middle.Y)) * YAxisConeTM.GetRotation());
-			DrawCone(PDI, FScaleMatrix(30.0f) * YAxisConeTM.ToMatrixWithScale(), FMath::DegreesToRadians(AngleRange.Y / 2.0f), 0.0f, 24, false, FLinearColor::Green, GEngine->ConstraintLimitMaterialY->GetRenderProxy(), SDPG_World);
+			DrawCone(PDI, FScaleMatrix(30.0f) * YAxisConeTM.ToMatrixWithScale(), static_cast<float>(FMath::DegreesToRadians(AngleRange.Y / 2.0)), 0.0f, 24, false, FLinearColor::Green, GEngine->ConstraintLimitMaterialY->GetRenderProxy(), SDPG_World);
 	
 			FTransform ZAxisConeTM(XAxis, ZAxis ^ XAxis, ZAxis, JointTransform.GetTranslation());
 			ZAxisConeTM.SetRotation(FQuat(ZAxis, FMath::DegreesToRadians(Middle.Z)) * ZAxisConeTM.GetRotation());
-			DrawCone(PDI, FScaleMatrix(30.0f) * ZAxisConeTM.ToMatrixWithScale(), FMath::DegreesToRadians(AngleRange.Z / 2.0f), 0.0f, 24, false, FLinearColor::Blue, GEngine->ConstraintLimitMaterialZ->GetRenderProxy(), SDPG_World);
+			DrawCone(PDI, FScaleMatrix(30.0f) * ZAxisConeTM.ToMatrixWithScale(), static_cast<float>(FMath::DegreesToRadians(AngleRange.Z / 2.0)), 0.0f, 24, false, FLinearColor::Blue, GEngine->ConstraintLimitMaterialZ->GetRenderProxy(), SDPG_World);
 		}
 	}
 }
diff --git a/Engine/Source/Editor/AnimGraph/Private/AnimGraphNode_Trail.cpp b/Engine/Source/Editor/AnimGraph/Private/AnimGraphNode_Trail.cpp
--- a/Engine/Source/Editor/AnimGraph/Private/AnimGraphNode_Trail.cpp
+++ b/Engine/Source/Editor/AnimGraph/Private/AnimGraphNode_Trail.cpp
@@ -170,15 +170,15 @@ void UAnimGraphNode_Trail::DrawAngularLimits(FPrimitiveDrawInterface* PDI, USkel
 
 			FTransform XAxisConeTM(YAxis, XAxis ^ YAxis, XAxis, JointTransform.GetTranslation());
 			XAxisConeTM.SetRotation(FQuat(XAxis, FMath::DegreesToRadians(-Middle.X)) * XAxisConeTM.GetRotation());
-			DrawCone(PDI, FScaleMatrix(30.0f) * XAxisConeTM.ToMatrixWithScale(), FMath::DegreesToRadians(AngleRange.X / 2.0f), 0.0f, 24, false, FLinearColor::Red, GEngine->ConstraintLimitMaterialX->GetRenderProxy(), SDPG_World);
+			DrawCone(PDI, FScaleMatrix(30.0f) * XAxisConeTM.ToMatrixWithScale(), static_cast<float>(FMath::DegreesToRadians(AngleRange.X / 2.0)), 0.0f, 24, false, FLinearColor::Red, GEngine->ConstraintLimitMaterialX->GetRenderProxy(), SDPG_World);
 
 			FTransform YAxisConeTM(ZAxis, YAxis ^ ZAxis, YAxis, JointTransform.GetTranslation());
 			YAxisConeTM.SetRotation(FQuat(YAxis, FMath::DegreesToRadians(Middle.Y)) * YAxisConeTM.GetRotation());
-			DrawCone(PDI, FScaleMatrix(30.0f) * YAxisConeTM.ToMatrixWithScale(), FMath::DegreesToRadians(AngleRange.Y / 2.0f), 0.0f, 24, false, FLinearColor::Green, GEngine->ConstraintLimitMaterialY->GetRenderProxy(), SDPG_World);
+			DrawCone(PDI, FScaleMatrix(30.0f) * YAxisConeTM.ToMatrixWithScale(), static_cast<float>(FMath::DegreesToRadians(AngleRange.Y / 2.0)), 0.0f, 24, false, FLinearColor::Green, GEngine->ConstraintLimitMaterialY->GetRenderProxy(), SDPG_World);
 
 			FTransform ZAxisConeTM(XAxis, ZAxis ^ XAxis, ZAxis, JointTransform.GetTranslation());
 			ZAxisConeTM.SetRotation(FQuat(ZAxis, FMath::DegreesToRadians(Middle.Z)) * ZAxisConeTM.GetRotation());
-			DrawCone(PDI, FScaleMatrix(30.0f) * ZAxisConeTM.ToMatrixWithScale(), FMath::DegreesToRadians(AngleRange.Z / 2.0f), 0.0f, 24, false, FLinearColor::Blue, GEngine->ConstraintLimitMaterialZ->GetRenderProxy(), SDPG_World);
+			DrawCone(PDI, FScaleMatrix(30.0f) * ZAxisConeTM.ToMatrixWithScale(), static_cast<float>(FMath::DegreesToRadians(AngleRange.Z / 2.0)), 0.0f, 24, false, FLinearColor::Blue, GEngine->ConstraintLimitMaterialZ->GetRenderProxy(), SDPG_World);
 		}
 	}
 }
diff --git a/Engine/Source/Editor/AnimGraph/Private/AnimPreviewInstance.cpp b/Engine/Source/Editor/AnimGraph/Private/AnimPreviewInstance.cpp
--- a/Engine/Source/Editor/AnimGraph/Private/AnimPreviewInstance.cpp
+++ b/Engine/Source/Editor/AnimGraph/Private/AnimPreviewInstance.cpp
@@ -796,10 +796,10 @@ float UAnimPreviewInstance::MontagePreview_CalculateStepLength()
 {
 	if (UAnimMontage* Montage = Cast<UAnimMontage>(CurrentAsset))
 	{
-		return Montage->GetSamplingFrameRate().AsInterval();
+		return static_cast<float>(Montage->GetSamplingFrameRate().AsInterval());
 	}
 
-	return UAnimationSettings::Get()->GetDefaultFrameRate().AsInterval();
+	return static_cast<float>(UAnimationSettings::Get()->GetDefaultFrameRate().AsInterval());
 }
 
 void UAnimPreviewInstance::MontagePreview_JumpToStart()
diff --git a/Engine/Source/Editor/AnimGraph/Private/AnimationGraphSchema.cpp b/Engine/Source/Editor/AnimGraph/Private/AnimationGraphSchema.cpp
--- a/Engine/Source/Editor/AnimGraph/Private/AnimationGraphSchema.cpp
+++ b/Engine/Source/Editor/AnimGraph/Private/AnimationGraphSchema.cpp
@@ -93,8 +93,8 @@ FReply FAnimationLayerDragDropAction::DroppedOnPanel(const TSharedRef< class SWi
 			const FName GraphName = TargetRigGraph->GetFName();
 			LinkedAnimLayerNode->SetupFromLayerId(SourceFuncName);
 			LinkedInputLayerNodeCreator.Finalize();
-			LinkedAnimLayerNode->NodePosX = GraphPosition.X;
-			LinkedAnimLayerNode->NodePosY = GraphPosition.Y;
+			LinkedAnimLayerNode->NodePosX = static_cast<int32>(GraphPosition.X);
+			LinkedAnimLayerNode->NodePosY = static_cast<int32>(GraphPosition.Y);
 		}
 	}
 	return FReply::Unhandled();
@@ -966,35 +966,35 @@ void UAnimationGraphSchema::AutoArrangeInterfaceGraph(UEdGraph& Graph)
 
 	FBox2D RootBounds(FVector2D(Root->NodePosX, Root->NodePosY), FVector2D(Root->NodePosX + 130, Root->NodePosY + 200));
 
-	float TotalHeight = 0.0f;
-	float MaxWidth = 0.0f;
+	double TotalHeight = 0.0;
+	double MaxWidth = 0.0;
 	const int32 HeightPerProperty = 30;
 
 	for(UAnimGraphNode_LinkedInputPose* Node : LinkedInputPoseNodes)
 	{
 		FBox2D LinkedInputPoseBounds(
 			FVector2D(Node->NodePosX, Node->NodePosY),
-			FVector2D(Node->NodePosX + 400, Node->NodePosY + 100 + (Node->GetNumInputs() * HeightPerProperty))
+			FVector2D(Node->NodePosX + 400.0, Node->NodePosY + 100.0 + (Node->GetNumInputs() * HeightPerProperty))
 		);
 
 		FVector2D BoundsSize = LinkedInputPoseBounds.GetSize();
-		TotalHeight += BoundsSize.Y + 10.0f;
+		TotalHeight += BoundsSize.Y + 10.0;
 		MaxWidth = FMath::Max(BoundsSize.X, MaxWidth);
 	}
 
-	float NodeOffset = RootBounds.GetCenter().Y - (TotalHeight * 0.5f);
-	float NodePosX = RootBounds.Min.X - (MaxWidth + 100.0f);
+	double NodeOffset = RootBounds.GetCenter().Y - (TotalHeight * 0.5);
+	double NodePosX = RootBounds.Min.X - (MaxWidth + 100.0);
 	for(UAnimGraphNode_LinkedInputPose* Node : LinkedInputPoseNodes)
 	{
-		Node->NodePosX = NodePosX;
-		Node->NodePosY = NodeOffset;
+		Node->NodePosX = static_cast<int32>(NodePosX);
+		Node->NodePosY = static_cast<int32>(NodeOffset);
 
 		FBox2D LinkedInputPoseBounds(
 			FVector2D(Node->NodePosX, Node->NodePosY),
-			FVector2D(Node->NodePosX + 400, Node->NodePosY + 100 + (Node->GetNumInputs() * HeightPerProperty))
+			FVector2D(Node->NodePosX + 400.0, Node->NodePosY + 100.0 + (Node->GetNumInputs() * HeightPerProperty))
 		);
 
-		NodeOffset += LinkedInputPoseBounds.GetSize().Y + 10.0f;
+		NodeOffset += LinkedInputPoseBounds.GetSize().Y + 10.0;
 	}
 }
 
@@ -1068,8 +1068,8 @@ void UAnimationGraphSchema::ConformAnimGraphToInterface(UBlueprint* InBlueprint,
 							LinkedInputPoseNodeCreator.Finalize();
 
 							FVector2D NewPosition = GetPositionForNewLinkedInputPoseNode(InGraph);
-							LinkedInputPoseNode->NodePosX = NewPosition.X;
-							LinkedInputPoseNode->NodePosY = NewPosition.Y;
+							LinkedInputPoseNode->NodePosX = static_cast<int32>(NewPosition.X);
+							LinkedInputPoseNode->NodePosY = static_cast<int32>(NewPosition.Y);
 						}
 
 						CurrentPoseIndex++;
diff --git a/Engine/Source/Editor/AnimGraph/Private/AnimationStateMachineSchema.cpp b/Engine/Source/Editor/AnimGraph/Private/AnimationStateMachineSchema.cpp
--- a/Engine/Source/Editor/AnimGraph/Private/AnimationStateMachineSchema.cpp
+++ b/Engine/Source/Editor/AnimGraph/Private/AnimationStateMachineSchema.cpp
@@ -67,8 +67,8 @@ UEdGraphNode* FEdGraphSchemaAction_NewStateNode::PerformAction(class UEdGraph* P
 		NodeTemplate->AllocateDefaultPins();
 		NodeTemplate->AutowireNewNode(FromPin);
 
-		NodeTemplate->NodePosX = Location.X;
-		NodeTemplate->NodePosY = Location.Y;
+		NodeTemplate->NodePosX = static_cast<int32>(Location.X);
+		NodeTemplate->NodePosY = static_cast<int32>(Location.Y);
 		NodeTemplate->SnapToGrid(GetDefault<UEditorStyleSettings>()->GridSnapSize);
 
 		ResultNode = NodeTemplate;
diff --git a/Engine/Source/Editor/AnimGraph/Private/EditModes/AnimDynamicsEditMode.cpp b/Engine/Source/Editor/AnimGraph/Private/EditModes/AnimDynamicsEditMode.cpp
--- a/Engine/Source/Editor/AnimGraph/Private/EditModes/AnimDynamicsEditMode.cpp
+++ b/Engine/Source/Editor/AnimGraph/Private/EditModes/AnimDynamicsEditMode.cpp
@@ -83,21 +83,21 @@ void DrawAngularLimits(FPrimitiveDrawInterface* PDI, FTransform JointTransform,
 	{
 		FTransform XAxisConeTM(YAxis, XAxis ^ YAxis, XAxis, JointTransform.GetTranslation());
 		XAxisConeTM.SetRotation(FQuat(XAxis, FMath::DegreesToRadians(-Middle.X)) * XAxisConeTM.GetRotation());
-		DrawCone(PDI, FScaleMatrix(30.0f) * XAxisConeTM.ToMatrixWithScale(), FMath::DegreesToRadians(AngleRange.X / 2.0f), 0.0f, 24, false, FLinearColor::White, GEngine->ConstraintLimitMaterialX->GetRenderProxy(), SDPG_World);
+		DrawCone(PDI, FScaleMatrix(30.0f) * XAxisConeTM.ToMatrixWithScale(), FMath::DegreesToRadians(static_cast<float>(AngleRange.X) / 2.0f), 0.0f, 24, false, FLinearColor::White, GEngine->ConstraintLimitMaterialX->GetRenderProxy(), SDPG_World);
 	}
 
 	if (AngleRange.Y > 0.0f && AngleRange.Y < 180.0f)
 	{
 		FTransform YAxisConeTM(ZAxis, YAxis ^ ZAxis, YAxis, JointTransform.GetTranslation());
 		YAxisConeTM.SetRotation(FQuat(YAxis, FMath::DegreesToRadians(Middle.Y)) * YAxisConeTM.GetRotation());
-		DrawCone(PDI, FScaleMatrix(30.0f) * YAxisConeTM.ToMatrixWithScale(), FMath::DegreesToRadians(AngleRange.Y / 2.0f), 0.0f, 24, false, FLinearColor::White, GEngine->ConstraintLimitMaterialY->GetRenderProxy(), SDPG_World);
+		DrawCone(PDI, FScaleMatrix(30.0f) * YAxisConeTM.ToMatrixWithScale(), FMath::DegreesToRadians(static_cast<float>(AngleRange.Y) / 2.0f), 0.0f, 24, false, FLinearColor::White, GEngine->ConstraintLimitMaterialY->GetRenderProxy(), SDPG_World);
 	}
 
 	if (AngleRange.Z > 0.0f && AngleRange.Z < 180.0f)
 	{
 		FTransform ZAxisConeTM(XAxis, ZAxis ^ XAxis, ZAxis, JointTransform.GetTranslation());
 		ZAxisConeTM.SetRotation(FQuat(ZAxis, FMath::DegreesToRadians(Middle.Z)) * ZAxisConeTM.GetRotation());
-		DrawCone(PDI, FScaleMatrix(30.0f) * ZAxisConeTM.ToMatrixWithScale(), FMath::DegreesToRadians(AngleRange.Z / 2.0f), 0.0f, 24, false, FLinearColor::White, GEngine->ConstraintLimitMaterialZ->GetRenderProxy(), SDPG_World);
+		DrawCone(PDI, FScaleMatrix(30.0f) * ZAxisConeTM.ToMatrixWithScale(), FMath::DegreesToRadians(static_cast<float>(AngleRange.Z) / 2.0f), 0.0f, 24, false, FLinearColor::White, GEngine->ConstraintLimitMaterialZ->GetRenderProxy(), SDPG_World);
 	}
 }
 
@@ -392,7 +392,7 @@ bool FAnimDynamicsEditMode::HandleClick(FEditorViewportClient* InViewportClient,
 	const bool bIsAltKeyDown = Viewport->KeyState(EKeys::LeftAlt) || Viewport->KeyState(EKeys::RightAlt);
 	bool bHandled = false;
 
-	const bool bModifySelection = bIsCtrlKeyDown | bIsShiftKeyDown;
+	const bool bModifySelection = bIsCtrlKeyDown || bIsShiftKeyDown;
 
 	const HAnimDynamicsEditModeHitProxy* const AnimDynamicsProxy = HitProxyCast<HAnimDynamicsEditModeHitProxy>(HitProxy);
 
@@ -759,7 +759,7 @@ void FAnimDynamicsEditMode::DoRotation(FRotator& InRotation)
 
 void FAnimDynamicsEditMode::DoScale(FVector& InScale)
 {
-	const float Scale = InScale.X + InScale.Y + InScale.Z;
+	const float Scale = static_cast<float>(InScale.X + InScale.Y + InScale.Z);
 
 	for (const FAnimDynamicsViewportObjectReference& SelectedObjectRef : SelectedViewportObjects)
 	{
diff --git a/Engine/Source/Editor/AnimGraph/Private/EditModes/SplineIKEditMode.cpp b/Engine/Source/Editor/AnimGraph/Private/EditModes/SplineIKEditMode.cpp
--- a/Engine/Source/Editor/AnimGraph/Private/EditModes/SplineIKEditMode.cpp
+++ b/Engine/Source/Editor/AnimGraph/Private/EditModes/SplineIKEditMode.cpp
@@ -95,9 +95,9 @@ void FSplineIKEditMode::Render(const FSceneView* View, FViewport* Viewport, FPri
 
 		PDI->SetHitProxy(new HSplineHandleHitProxy(SplineHandleIndex));
 		FTransform StartTransform = SplineIKRuntimeNode->GetTransformedSplinePoint(SplineHandleIndex);
-		const float Scale = View->WorldToScreen(StartTransform.GetLocation()).W * (4.0f / View->UnscaledViewRect.Width() / View->ViewMatrices.GetProjectionMatrix().M[0][0]);
-		DrawSphere(PDI, StartTransform.GetLocation(), FRotator::ZeroRotator, FVector(4.0f) * Scale, 64, 64, SphereMaterialProxy, SDPG_Foreground);
-		DrawCoordinateSystem(PDI, StartTransform.GetLocation(), StartTransform.GetRotation().Rotator(), 30.0f * Scale, SDPG_Foreground);
+		const double Scale = View->WorldToScreen(StartTransform.GetLocation()).W * (4.0 / View->UnscaledViewRect.Width() / View->ViewMatrices.GetProjectionMatrix().M[0][0]);
+		DrawSphere(PDI, StartTransform.GetLocation(), FRotator::ZeroRotator, FVector(4.0) * Scale, 64, 64, SphereMaterialProxy, SDPG_Foreground);
+		DrawCoordinateSystem(PDI, StartTransform.GetLocation(), StartTransform.GetRotation().Rotator(), static_cast<float>(30.0 * Scale), SDPG_Foreground);
 	}
 
 	PDI->SetHitProxy(nullptr);
diff --git a/Engine/Source/Editor/AnimGraph/Private/PoseDriverDetails.cpp b/Engine/Source/Editor/AnimGraph/Private/PoseDriverDetails.cpp
--- a/Engine/Source/Editor/AnimGraph/Private/PoseDriverDetails.cpp
+++ b/Engine/Source/Editor/AnimGraph/Private/PoseDriverDetails.cpp
@@ -163,15 +163,15 @@ TSharedRef< SWidget > SPDD_TargetRow::GenerateWidgetForColumn(const FName& Colum
 	TSharedPtr<SVerticalBox> TargetEntryVertBox;
 
 	TSharedRef<SWidget> RowWidget = SNew(SBox)
-	.Padding(2)
+	.Padding(2.f)
 	[
 		SNew(SBorder)
-		.Padding(0)
+		.Padding(0.f)
 		.ForegroundColor(FLinearColor::White)
 		.BorderImage(FAppStyle::GetBrush("NoBorder"))
 		[
 			SAssignNew(ExpandArea, SExpandableArea)
-			.Padding(0)
+			.Padding(0.f)
 			.InitiallyCollapsed(true)
 			.BorderBackgroundColor(FLinearColor(.6f, .6f, .6f))
 			.OnAreaExpansionChanged(this, &SPDD_TargetRow::OnTargetExpansionChanged)
@@ -198,7 +198,7 @@ TSharedRef< SWidget > SPDD_TargetRow::GenerateWidgetForColumn(const FName& Colum
 				.AutoWidth()
 				[
 					SNew(SBox)
-					.MinDesiredWidth(150)
+					.MinDesiredWidth(150.f)
 					.Content()
 					[
 						SNew(SHorizontalBox)
@@ -225,8 +225,8 @@ TSharedRef< SWidget > SPDD_TargetRow::GenerateWidgetForColumn(const FName& Colum
 				.Padding(FMargin(3, 0))
 				[
 					SNew(SBox)
-					.MinDesiredWidth(40)
-					.MaxDesiredWidth(40)
+					.MinDesiredWidth(40.f)
+					.MaxDesiredWidth(40.f)
 					.Content()
 					[
 						SNew(STextBlock)
@@ -600,7 +600,7 @@ TOptional<float> SPDD_TargetRow::GetTranslation(int32 BoneIndex, EAxis::Type Axi
 	const FPoseDriverTarget* Target = GetTarget();
 	if (Target && Target->BoneTransforms.IsValidIndex(BoneIndex))
 	{
-		Translation = Target->BoneTransforms[BoneIndex].TargetTranslation.GetComponentForAxis(Axis);
+		Translation = static_cast<float>(Target->BoneTransforms[BoneIndex].TargetTranslation.GetComponentForAxis(Axis));
 	}
 
 	return Translation;
@@ -612,7 +612,7 @@ TOptional<float> SPDD_TargetRow::GetRotation(int32 BoneIndex, EAxis::Type Axis)
 	const FPoseDriverTarget* Target = GetTarget();
 	if (Target && Target->BoneTransforms.IsValidIndex(BoneIndex))
 	{
-		Rotation = Target->BoneTransforms[BoneIndex].TargetRotation.GetComponentForAxis(Axis);
+		Rotation = static_cast<float>(Target->BoneTransforms[BoneIndex].TargetRotation.GetComponentForAxis(Axis));
 	}
 
 	return Rotation;
@@ -1110,7 +1110,7 @@ void FPoseDriverDetails::CustomizeDetails(IDetailLayoutBuilder& DetailBuilder)
 			.AutoWidth()
 			[
 				SAssignNew(ToolsButton, SComboButton)
-				.ContentPadding(3)
+				.ContentPadding(3.f)
 				.ForegroundColor(this, &FPoseDriverDetails::GetToolsForegroundColor)
 				.ButtonStyle(FAppStyle::Get(), "ToggleButton") // Use the tool bar item style for this button
 				.OnGetMenuContent(this, &FPoseDriverDetails::GetToolsMenuContent)
diff --git a/Engine/Source/Runtime/AnimGraphRuntime/Private/AnimNodes/AnimNode_BlendSpaceGraphBase.cpp b/Engine/Source/Runtime/AnimGraphRuntime/Private/AnimNodes/AnimNode_BlendSpaceGraphBase.cpp
--- a/Engine/Source/Runtime/AnimGraphRuntime/Private/AnimNodes/AnimNode_BlendSpaceGraphBase.cpp
+++ b/Engine/Source/Runtime/AnimGraphRuntime/Private/AnimNodes/AnimNode_BlendSpaceGraphBase.cpp
@@ -121,6 +121,6 @@ void FAnimNode_BlendSpaceGraphBase::SnapToPosition(const FVector& NewPosition)
 	const int32 NumAxis = FMath::Min(BlendFilter.FilterPerAxis.Num(), 3);
 	for (int32 Idx = 0; Idx < NumAxis; Idx++)
 	{
-		BlendFilter.FilterPerAxis[Idx].SetToValue(NewPosition[Idx]);
+		BlendFilter.FilterPerAxis[Idx].SetToValue(static_cast<float>(NewPosition[Idx]));
 	}
 }
diff --git a/Engine/Source/Runtime/AnimGraphRuntime/Private/AnimNodes/AnimNode_RandomPlayer.cpp b/Engine/Source/Runtime/AnimGraphRuntime/Private/AnimNodes/AnimNode_RandomPlayer.cpp
--- a/Engine/Source/Runtime/AnimGraphRuntime/Private/AnimNodes/AnimNode_RandomPlayer.cpp
+++ b/Engine/Source/Runtime/AnimGraphRuntime/Private/AnimNodes/AnimNode_RandomPlayer.cpp
@@ -433,7 +433,7 @@ void FAnimNode_RandomPlayer::InitPlayData(FRandomAnimPlayData& Data, int32 InVal
 
 	Data.Entry = Entry;
 	Data.BlendWeight = InBlendWeight;
-	Data.PlayRate = RandomStream.FRandRange(Entry->MinPlayRate, Entry->MaxPlayRate);
+	Data.PlayRate = static_cast<float>(RandomStream.FRandRange(Entry->MinPlayRate, Entry->MaxPlayRate));
 	Data.RemainingLoops = FMath::Clamp(RandomStream.RandRange(Entry->MinLoopCount, Entry->MaxLoopCount), 0, MAX_int32);
 
 	Data.PlayStartTime = 0.0f;
diff --git a/Engine/Source/Runtime/AnimGraphRuntime/Private/BoneControllers/AnimNode_CCDIK.cpp b/Engine/Source/Runtime/AnimGraphRuntime/Private/BoneControllers/AnimNode_CCDIK.cpp
--- a/Engine/Source/Runtime/AnimGraphRuntime/Private/BoneControllers/AnimNode_CCDIK.cpp
+++ b/Engine/Source/Runtime/AnimGraphRuntime/Private/BoneControllers/AnimNode_CCDIK.cpp
@@ -97,7 +97,7 @@ void FAnimNode_CCDIK::EvaluateSkeletalControl_AnyThread(FComponentSpacePoseConte
 		OutBoneTransforms[TransformIndex] = FBoneTransform(BoneIndex, BoneCSTransform);
 
 		// Calculate the combined length of this segment of skeleton
-		float const BoneLength = FVector::Dist(BoneCSPosition, OutBoneTransforms[TransformIndex - 1].Transform.GetLocation());
+		double const BoneLength = FVector::Dist(BoneCSPosition, OutBoneTransforms[TransformIndex - 1].Transform.GetLocation());
 
 		if (!FMath::IsNearlyZero(BoneLength))
 		{
diff --git a/Engine/Source/Runtime/AnimGraphRuntime/Private/BoneControllers/AnimNode_ScaleChainLength.cpp b/Engine/Source/Runtime/AnimGraphRuntime/Private/BoneControllers/AnimNode_ScaleChainLength.cpp
--- a/Engine/Source/Runtime/AnimGraphRuntime/Private/BoneControllers/AnimNode_ScaleChainLength.cpp
+++ b/Engine/Source/Runtime/AnimGraphRuntime/Private/BoneControllers/AnimNode_ScaleChainLength.cpp
@@ -104,13 +104,13 @@ void FAnimNode_ScaleChainLength::Evaluate_AnyThread(FPoseContext& Output)
 
 	const FTransform StartTransformCompSpace = CSPose.GetComponentSpaceTransform(ChainBoneIndices[0]);
 
-	const float DesiredChainLength = (TargetLocationCompSpace - StartTransformCompSpace.GetLocation()).Size();
-	const float InitialChainLength = GetInitialChainLength(Output.Pose, CSPose);
-	const float ChainLengthScale = !FMath::IsNearlyZero(InitialChainLength) ? (DesiredChainLength / InitialChainLength) : 1.f;
-	const float ChainLengthScaleWithAlpha = FMath::LerpStable(1.f, ChainLengthScale, ActualAlpha);
+	const double DesiredChainLength = (TargetLocationCompSpace - StartTransformCompSpace.GetLocation()).Size();
+	const double InitialChainLength = GetInitialChainLength(Output.Pose, CSPose);
+	const double ChainLengthScale = !FMath::IsNearlyZero(InitialChainLength) ? (DesiredChainLength / InitialChainLength) : 1.0;
+	const double ChainLengthScaleWithAlpha = FMath::LerpStable(1.0, ChainLengthScale, ActualAlpha);
 
 	// If we're not going to scale anything, early out.
-	if (FMath::IsNearlyEqual(ChainLengthScaleWithAlpha, 1.f))
+	if (FMath::IsNearlyEqual(ChainLengthScaleWithAlpha, 1.0))
 	{
 		return;
 	}
@@ -124,7 +124,7 @@ void FAnimNode_ScaleChainLength::Evaluate_AnyThread(FPoseContext& Output)
 	}
 }
 
-float FAnimNode_ScaleChainLength::GetInitialChainLength(FCompactPose& InLSPose, FCSPose<FCompactPose>& InCSPose) const
+double FAnimNode_ScaleChainLength::GetInitialChainLength(FCompactPose& InLSPose, FCSPose<FCompactPose>& InCSPose) const
 {
 	switch (ChainInitialLength)
 	{
@@ -137,7 +137,7 @@ float FAnimNode_ScaleChainLength::GetInitialChainLength(FCompactPose& InLSPose,
 
 	case EScaleChainInitialLength::ChainLength :
 	{
-		float ChainLength = 0.f;
+		double ChainLength = 0.0;
 		for (const FCompactPoseBoneIndex& BoneIndex : ChainBoneIndices)
 		{
 			ChainLength += InLSPose[BoneIndex].GetTranslation().Size();
diff --git a/Engine/Source/Runtime/AnimGraphRuntime/Private/BoneControllers/BoneControllerSolvers.cpp b/Engine/Source/Runtime/AnimGraphRuntime/Private/BoneControllers/BoneControllerSolvers.cpp
--- a/Engine/Source/Runtime/AnimGraphRuntime/Private/BoneControllers/BoneControllerSolvers.cpp
+++ b/Engine/Source/Runtime/AnimGraphRuntime/Private/BoneControllers/BoneControllerSolvers.cpp
@@ -17,8 +17,8 @@ FTransform FIKFootPelvisPullDownSolver::Solve(FTransform PelvisTransform, TArray
 	FVector AdjustedPelvisLocation = InitialPelvisLocation;
 	FVector DeltaAdjustment = FVector::ZeroVector;
 
-	const float PerFootWeight = 1.f / static_cast<float>(IKFootLocationsCount);
-	const float AdjustmentDistMaxSquared = FMath::Pow(PelvisAdjustmentMaxDistance, 2.f);
+	const double PerFootWeight = 1.0 / static_cast<double>(IKFootLocationsCount);
+	const double AdjustmentDistMaxSquared = FMath::Pow(PelvisAdjustmentMaxDistance, 2.0);
 
 	// Pull pelvis closer to feet iteratively
 	for (int32 Iter = 0; Iter < PelvisAdjustmentMaxIter; ++Iter)
@@ -35,7 +35,7 @@ FTransform FIKFootPelvisPullDownSolver::Solve(FTransform PelvisTransform, TArray
 
 		const FVector PrevDeltaAdjustment = DeltaAdjustment;
 		DeltaAdjustment = AdjustedPelvisLocation - InitialPelvisLocation;
-		const float DeltaAdjustmentDist = FVector::Dist(PrevDeltaAdjustment, DeltaAdjustment);
+		const double DeltaAdjustmentDist = FVector::Dist(PrevDeltaAdjustment, DeltaAdjustment);
 
 		// Keep track of how much delta adjustment is being applied per iteration
 		if (DeltaAdjustmentDist <= PelvisAdjustmentErrorTolerance)
diff --git a/Engine/Source/Runtime/AnimGraphRuntime/Public/BoneControllers/AnimNode_ScaleChainLength.h b/Engine/Source/Runtime/AnimGraphRuntime/Public/BoneControllers/AnimNode_ScaleChainLength.h
--- a/Engine/Source/Runtime/AnimGraphRuntime/Public/BoneControllers/AnimNode_ScaleChainLength.h
+++ b/Engine/Source/Runtime/AnimGraphRuntime/Public/BoneControllers/AnimNode_ScaleChainLength.h
@@ -74,5 +74,5 @@ public:
 	// End of FAnimNode_Base interface
 
 protected:
-	float GetInitialChainLength(FCompactPose& InLSPose, FCSPose<FCompactPose>& InCSPose) const;
+	double GetInitialChainLength(FCompactPose& InLSPose, FCSPose<FCompactPose>& InCSPose) const;
 };
diff --git a/Engine/Source/Runtime/AnimGraphRuntime/Public/BoneControllers/BoneControllerSolvers.h b/Engine/Source/Runtime/AnimGraphRuntime/Public/BoneControllers/BoneControllerSolvers.h
--- a/Engine/Source/Runtime/AnimGraphRuntime/Public/BoneControllers/BoneControllerSolvers.h
+++ b/Engine/Source/Runtime/AnimGraphRuntime/Public/BoneControllers/BoneControllerSolvers.h
@@ -22,11 +22,11 @@ struct ANIMGRAPHRUNTIME_API FIKFootPelvisPullDownSolver
 	// Specifies an alpha between the original and final adjusted pelvis locations
 	// This is used to retain some degree of the original pelvis motion
 	UPROPERTY(EditAnywhere, Category=Settings, meta=(ClampMin="0.0", ClampMax="1.0"))
-	float PelvisAdjustmentInterpAlpha = 0.5f;
+	double PelvisAdjustmentInterpAlpha = 0.5;
 
 	// Specifies the maximum displacement the pelvis can be adjusted relative to its original location
 	UPROPERTY(EditAnywhere, Category=Settings, meta=(ClampMin="0.0"))
-	float PelvisAdjustmentMaxDistance = 10.f;
+	double PelvisAdjustmentMaxDistance = 10.0;
 
 	// Specifies the pelvis adjustment distance error that is tolerated for each iteration of the solver
 	// 
@@ -36,7 +36,7 @@ struct ANIMGRAPHRUNTIME_API FIKFootPelvisPullDownSolver
 	//
 	// The default value of 0.01 specifies 1 centimeter of error
 	UPROPERTY(EditAnywhere, Category=Advanced, meta=(ClampMin="0.001"))
-	float PelvisAdjustmentErrorTolerance = 0.01f;
+	double PelvisAdjustmentErrorTolerance = 0.01;
 
 	// Specifies the maximum number of iterations to run for the pelvis adjustment solver
 	// Higher iterations will guarantee closer PelvisAdjustmentErrorTolerance convergence at the cost of performance
