Fix super not being called in UAnimGraphNode_LinkedAnimGraph::PostPasteNode & UAnimStateAliasNode::PostPasteNode()

#tests PIE
#rnx

diff --git a/Engine/Source/Editor/AnimGraph/Private/AnimGraphNode_LinkedAnimGraph.cpp b/Engine/Source/Editor/AnimGraph/Private/AnimGraphNode_LinkedAnimGraph.cpp
--- a/Engine/Source/Editor/AnimGraph/Private/AnimGraphNode_LinkedAnimGraph.cpp
+++ b/Engine/Source/Editor/AnimGraph/Private/AnimGraphNode_LinkedAnimGraph.cpp
@@ -17,6 +17,8 @@
 
 void UAnimGraphNode_LinkedAnimGraph::PostPasteNode()
 {
+	Super::PostPasteNode();
+
 	// Clear incompatible target class
 	if(UClass* InstanceClass = GetTargetClass())
 	{
diff --git a/Engine/Source/Editor/AnimGraph/Private/AnimStateAliasNode.cpp b/Engine/Source/Editor/AnimGraph/Private/AnimStateAliasNode.cpp
--- a/Engine/Source/Editor/AnimGraph/Private/AnimStateAliasNode.cpp
+++ b/Engine/Source/Editor/AnimGraph/Private/AnimStateAliasNode.cpp
@@ -69,6 +69,8 @@ void UAnimStateAliasNode::AutowireNewNode(UEdGraphPin* FromPin)
 
 void UAnimStateAliasNode::PostPasteNode()
 {
+	Super::PostPasteNode();
+
 	// Find an interesting name, but try to keep the same if possible
 	TSharedPtr<INameValidatorInterface> NameValidator = FNameValidatorFactory::MakeValidator(this);
 	NameValidator->FindValidString(StateAliasName);
