Fixes for LookAt AimOffset:
* Lets" ShowDebug Animation" work
* Lets input filtering work
* Allows the source/pivot to be a bone or socket

diff --git a/Engine/Source/Editor/AnimGraph/Private/AnimGraphNode_AimOffsetLookAt.cpp b/Engine/Source/Editor/AnimGraph/Private/AnimGraphNode_AimOffsetLookAt.cpp
--- a/Engine/Source/Editor/AnimGraph/Private/AnimGraphNode_AimOffsetLookAt.cpp
+++ b/Engine/Source/Editor/AnimGraph/Private/AnimGraphNode_AimOffsetLookAt.cpp
@@ -180,24 +180,59 @@ void UAnimGraphNode_AimOffsetLookAt::ValidateAnimNodeDuringCompilation(class USk
 				MessageLog.Error(TEXT("@@ references blendspace that uses an incompatible skeleton @@"), this, BlendSpaceSkeleton);
 			}
 
-			// Make sure that the source socket name is a valid one for the skeleton
-			UEdGraphPin* SocketNamePin = FindPin(GET_MEMBER_NAME_STRING_CHECKED(FAnimNode_AimOffsetLookAt, SourceSocketName));
-			FName SocketNameToCheck = (SocketNamePin != nullptr) ? FName(*SocketNamePin->DefaultValue) : Node.SourceSocketName;
-
 			// Temporary fix where skeleton is not fully loaded during AnimBP compilation and thus the socket name check is invalid UE-39499 (NEED FIX) 
 			if (!BlendSpaceSkeleton->HasAnyFlags(RF_NeedPostLoad))
 			{
-				const bool bValidValue = SocketNamePin == nullptr && BlendSpaceSkeleton->FindSocket(Node.SourceSocketName);
-				const bool bValidPinValue = SocketNamePin != nullptr && BlendSpaceSkeleton->FindSocket(FName(*SocketNamePin->DefaultValue));
-				const bool bValidConnectedPin = SocketNamePin != nullptr && SocketNamePin->LinkedTo.Num();
-
-				if (!bValidValue && !bValidPinValue && !bValidConnectedPin)
-				{ 
-					FFormatNamedArguments Args;
-					Args.Add(TEXT("SocketName"), FText::FromName(SocketNameToCheck));
+				{
+					// Make sure that the source socket name is a valid one for the skeleton
+					UEdGraphPin* SocketNamePin = FindPin(GET_MEMBER_NAME_STRING_CHECKED(FAnimNode_AimOffsetLookAt, SourceSocketName));
+					FName SocketNameToCheck = (SocketNamePin != nullptr) ? FName(*SocketNamePin->DefaultValue) : Node.SourceSocketName;
+
+					const FReferenceSkeleton& RefSkel = BlendSpaceSkeleton->GetReferenceSkeleton();
+				
+					const bool bValidValue = SocketNamePin == nullptr && (
+						BlendSpaceSkeleton->FindSocket(Node.SourceSocketName) || 
+						RefSkel.FindBoneIndex(Node.SourceSocketName) != INDEX_NONE);
+					const bool bValidPinValue = SocketNamePin != nullptr && (
+						BlendSpaceSkeleton->FindSocket(FName(*SocketNamePin->DefaultValue)) || 
+						RefSkel.FindBoneIndex(FName(*SocketNamePin->DefaultValue)) != INDEX_NONE);
+					const bool bValidConnectedPin = SocketNamePin != nullptr && SocketNamePin->LinkedTo.Num();
+
+					if (!bValidValue && !bValidPinValue && !bValidConnectedPin)
+					{ 
+						FFormatNamedArguments Args;
+						Args.Add(TEXT("SocketName"), FText::FromName(SocketNameToCheck));
+
+						const FText Msg = FText::Format(LOCTEXT("SocketNameNotFound", "@@ - Socket {SocketName} not found in Skeleton"), Args);
+						MessageLog.Error(*Msg.ToString(), this);
+					}
+				}
 
-					const FText Msg = FText::Format(LOCTEXT("SocketNameNotFound", "@@ - Socket {SocketName} not found in Skeleton"), Args);
-					MessageLog.Error(*Msg.ToString(), this);
+				// And similarly the pivot socket, though here we allow it to be blank (but not invalid)
+				{
+					// Make sure that the source socket name is a valid one for the skeleton
+					UEdGraphPin* SocketNamePin = FindPin(GET_MEMBER_NAME_STRING_CHECKED(FAnimNode_AimOffsetLookAt, PivotSocketName));
+					FName SocketNameToCheck = (SocketNamePin != nullptr) ? FName(*SocketNamePin->DefaultValue) : Node.PivotSocketName;
+
+					const FReferenceSkeleton& RefSkel = BlendSpaceSkeleton->GetReferenceSkeleton();
+
+					const bool bValidValue = SocketNamePin == nullptr && (
+						Node.PivotSocketName.IsNone() || 
+						BlendSpaceSkeleton->FindSocket(Node.PivotSocketName) ||
+						RefSkel.FindBoneIndex(Node.PivotSocketName) != INDEX_NONE);
+					const bool bValidPinValue = SocketNamePin != nullptr && (
+						BlendSpaceSkeleton->FindSocket(FName(*SocketNamePin->DefaultValue)) ||
+						RefSkel.FindBoneIndex(FName(*SocketNamePin->DefaultValue)) != INDEX_NONE);
+					const bool bValidConnectedPin = SocketNamePin != nullptr && SocketNamePin->LinkedTo.Num();
+
+					if (!bValidValue && !bValidPinValue && !bValidConnectedPin)
+					{
+						FFormatNamedArguments Args;
+						Args.Add(TEXT("SocketName"), FText::FromName(SocketNameToCheck));
+
+						const FText Msg = FText::Format(LOCTEXT("PivotSocketNameNotFound", "@@ - PivotSocket {SocketName} not found in Skeleton"), Args);
+						MessageLog.Error(*Msg.ToString(), this);
+					}
 				}
 			}
 		}
diff --git a/Engine/Source/Runtime/AnimGraphRuntime/Private/AnimNodes/AnimNode_AimOffsetLookAt.cpp b/Engine/Source/Runtime/AnimGraphRuntime/Private/AnimNodes/AnimNode_AimOffsetLookAt.cpp
--- a/Engine/Source/Runtime/AnimGraphRuntime/Private/AnimNodes/AnimNode_AimOffsetLookAt.cpp
+++ b/Engine/Source/Runtime/AnimGraphRuntime/Private/AnimNodes/AnimNode_AimOffsetLookAt.cpp
@@ -42,30 +42,33 @@ void FAnimNode_AimOffsetLookAt::OnInitializeAnimInstance(const FAnimInstanceProx
 				SocketLocalTransform = Socket->GetSocketLocalTransform();
 				SocketBoneReference.BoneName = Socket->BoneName;
 			}
+			else if (SkelMeshComp->GetBoneIndex(SourceSocketName) != INDEX_NONE)
+			{
+				SocketLocalTransform.SetIdentity();
+				SocketBoneReference.BoneName = SourceSocketName;
+			}
 
 			if (const USkeletalMeshSocket* Socket = SkelMesh->FindSocket(PivotSocketName))
 			{
 				PivotSocketLocalTransform = Socket->GetSocketLocalTransform();
 				PivotSocketBoneReference.BoneName = Socket->BoneName;
 			}
+			else if (SkelMeshComp->GetBoneIndex(PivotSocketName) != INDEX_NONE)
+			{
+				PivotSocketLocalTransform.SetIdentity();
+				PivotSocketBoneReference.BoneName = PivotSocketName;
+			}
 		}
 	}
 }
 
 void FAnimNode_AimOffsetLookAt::UpdateAssetPlayer(const FAnimationUpdateContext& Context)
 {
-	GetEvaluateGraphExposedInputs().Execute(Context);
-
 	bIsLODEnabled = IsLODEnabled(Context.AnimInstanceProxy);
-
-	// We don't support ticking and advancing time, because Inputs are determined during Evaluate.
-	// it may be possible to advance time there (is it a problem with notifies?)
-	// But typically AimOffsets contain single frame poses, so time doesn't matter.
-
-// 	if (bIsLODEnabled)
-// 	{
-// 		FAnimNode_BlendSpacePlayer::UpdateAssetPlayer(Context);
-// 	}
+ 	if (bIsLODEnabled)
+ 	{
+ 		FAnimNode_BlendSpacePlayer::UpdateAssetPlayer(Context);
+ 	}
 
 	BasePose.Update(Context);
 
@@ -166,10 +169,11 @@ void FAnimNode_AimOffsetLookAt::UpdateFromLookAtTarget(FPoseContext& LocalPoseCo
 #endif // ENABLE_DRAW_DEBUG
 	}
 
-	// Generate BlendSampleDataCache from inputs.
+	// Update Blend Space and put the result into BlendSampleDataCache.
 	if (CurrentBlendSpace)
 	{
-		CurrentBlendSpace->GetSamplesFromBlendInput(CurrentBlendInput, BlendSampleDataCache, CachedTriangulationIndex, true);
+		CurrentBlendSpace->UpdateBlendSamples(
+			CurrentBlendInput, DeltaTimeRecord.Delta, BlendSampleDataCache, CachedTriangulationIndex);
 	}
 }
 
diff --git a/Engine/Source/Runtime/AnimGraphRuntime/Public/AnimNodes/AnimNode_AimOffsetLookAt.h b/Engine/Source/Runtime/AnimGraphRuntime/Public/AnimNodes/AnimNode_AimOffsetLookAt.h
--- a/Engine/Source/Runtime/AnimGraphRuntime/Public/AnimNodes/AnimNode_AimOffsetLookAt.h
+++ b/Engine/Source/Runtime/AnimGraphRuntime/Public/AnimNodes/AnimNode_AimOffsetLookAt.h
@@ -35,11 +35,15 @@ struct ANIMGRAPHRUNTIME_API FAnimNode_AimOffsetLookAt : public FAnimNode_BlendSp
 	UPROPERTY(EditAnywhere, BlueprintReadWrite, Category = Performance, meta = (DisplayName = "LOD Threshold"))
 	int32 LODThreshold;
 
-	/** Socket to treat as the look at source */
+	/** Socket or bone to treat as the look at source. This will then be pointed at LookAtLocation */
 	UPROPERTY(EditAnywhere, BlueprintReadWrite, Category = LookAt, meta = (PinHiddenByDefault))
 	FName SourceSocketName;
 
-	/** Socket to treat as the look at pivot (optional). This will overwrite the translation of the source socket transform to better match the lookat direction */
+	/** 
+	 * Socket or bone to treat as the look at pivot (optional). This will overwrite the translation of the 
+	 * source socket transform improve the lookat direction, especially when the target is close 
+	 * to the character 
+	 */
 	UPROPERTY(EditAnywhere, BlueprintReadWrite, Category = LookAt, meta = (PinHiddenByDefault))
 	FName PivotSocketName;
 
@@ -47,7 +51,7 @@ struct ANIMGRAPHRUNTIME_API FAnimNode_AimOffsetLookAt : public FAnimNode_BlendSp
 	UPROPERTY(EditAnywhere, BlueprintReadWrite, Category = LookAt, meta = (PinShownByDefault))
 	FVector LookAtLocation;
 
-	/** Axis in the socket transform to consider the 'forward' or look at axis */
+	/** Direction in the socket transform to consider the 'forward' or look at axis */
 	UPROPERTY(EditAnywhere, BlueprintReadWrite, Category = LookAt, meta = (PinHiddenByDefault))
 	FVector SocketAxis;
 
