Fix redundant node spawner actions being registered for every loaded Animation Blueprint asset context in the editor.

diff --git a/Engine/Source/Editor/AnimGraph/Private/AnimGraphNode_CallFunction.cpp b/Engine/Source/Editor/AnimGraph/Private/AnimGraphNode_CallFunction.cpp
--- a/Engine/Source/Editor/AnimGraph/Private/AnimGraphNode_CallFunction.cpp
+++ b/Engine/Source/Editor/AnimGraph/Private/AnimGraphNode_CallFunction.cpp
@@ -167,9 +167,15 @@ void UAnimGraphNode_CallFunction::SetupFromFunction(UFunction* InFunction)
 	// Create graph and inner node
 	InnerGraph = FBlueprintEditorUtils::CreateNewGraph(this, NAME_None, UEdGraph::StaticClass(), UEdGraphSchema_K2::StaticClass());
 
+	bool bIsConsideredSelfContext = false;
+	if (const UAnimBlueprint* AnimBlueprint = Cast<UAnimBlueprint>(FBlueprintEditorUtils::FindBlueprintForNode(this)))
+	{
+		bIsConsideredSelfContext = InFunction->GetOwnerClass() == AnimBlueprint->GetAnimBlueprintGeneratedClass();
+	}
+
 	FGraphNodeCreator<UK2Node_CallFunction> NodeCreator(*InnerGraph);
 	CallFunctionPrototype = NodeCreator.CreateNode();
-	CallFunctionPrototype->FunctionReference.SetFromField<UFunction>(InFunction, true);
+	CallFunctionPrototype->FunctionReference.SetFromField<UFunction>(InFunction, bIsConsideredSelfContext);
 	NodeCreator.Finalize();
 
 	BindDelegates();
@@ -259,50 +265,51 @@ bool UAnimGraphNode_CallFunction::ValidateFunction(const UFunction* InFunction,
 
 void UAnimGraphNode_CallFunction::GetMenuActions(FBlueprintActionDatabaseRegistrar& ActionRegistrar) const
 {
-	const UAnimBlueprint* AnimBlueprint = Cast<UAnimBlueprint>(ActionRegistrar.GetActionKeyFilter());
-	if(AnimBlueprint && ActionRegistrar.IsOpenForRegistration(AnimBlueprint))
+	auto MakeFunctionActionsForClass = [this, &ActionRegistrar](const UClass* InClass)
 	{
-		auto MakeFunctionActionsForClass = [this, &ActionRegistrar, AnimBlueprint](UClass* InClass)
+		auto MakeFunctionAction = [this, &ActionRegistrar, InClass](UFunction* InFunction)
 		{
-			auto MakeFunctionAction = [this, &ActionRegistrar, AnimBlueprint](UFunction* InFunction)
+			if (UEdGraphSchema_K2::CanUserKismetCallFunction(InFunction) && ValidateFunction(InFunction))
 			{
-				if(UEdGraphSchema_K2::CanUserKismetCallFunction(InFunction) && ValidateFunction(InFunction))
+				auto CustomizeNode = [InFunction](UEdGraphNode* InNode, bool bIsTemplate)
+				{
+					UAnimGraphNode_CallFunction* CallFunctionNode = CastChecked<UAnimGraphNode_CallFunction>(InNode);
+					CallFunctionNode->SetupFromFunction(InFunction);
+				};
+
+				UBlueprintNodeSpawner* Spawner = UBlueprintNodeSpawner::Create(UAnimGraphNode_CallFunction::StaticClass(), nullptr, UBlueprintNodeSpawner::FCustomizeNodeDelegate::CreateLambda(CustomizeNode));
+				FBlueprintActionUiSpec& MenuSignature = Spawner->DefaultMenuSignature;
+
+				MenuSignature.MenuName = FText::Format(LOCTEXT("MenuNameFormat", "{0} (From Anim Graph)"), UK2Node_CallFunction::GetUserFacingFunctionName(InFunction));
+				MenuSignature.Category = UK2Node_CallFunction::GetDefaultCategoryForFunction(InFunction, LOCTEXT("BaseCategory", "Call Function From Anim Graph"));
+				MenuSignature.Tooltip = FText::FromString(UK2Node_CallFunction::GetDefaultTooltipForFunction(InFunction));
+				MenuSignature.Keywords = UK2Node_CallFunction::GetKeywordsForFunction(InFunction);
+
+				// add at least one character, so that PrimeDefaultUiSpec() doesn't attempt to query the template node
+				if (MenuSignature.Keywords.IsEmpty())
 				{
-					auto CustomizeNode = [InFunction](UEdGraphNode* InNode, bool bIsTemplate)
-					{
-						UAnimGraphNode_CallFunction* CallFunctionNode = CastChecked<UAnimGraphNode_CallFunction>(InNode);
-						CallFunctionNode->SetupFromFunction(InFunction);
-					};
-
-					UBlueprintNodeSpawner* Spawner = UBlueprintNodeSpawner::Create(UAnimGraphNode_CallFunction::StaticClass(), nullptr, UBlueprintNodeSpawner::FCustomizeNodeDelegate::CreateLambda(CustomizeNode));
-					FBlueprintActionUiSpec& MenuSignature = Spawner->DefaultMenuSignature;
-					
-					MenuSignature.MenuName = FText::Format(LOCTEXT("MenuNameFormat", "{0} (From Anim Graph)"), UK2Node_CallFunction::GetUserFacingFunctionName(InFunction));
-					MenuSignature.Category = UK2Node_CallFunction::GetDefaultCategoryForFunction(InFunction, LOCTEXT("BaseCategory", "Call Function From Anim Graph"));
-					MenuSignature.Tooltip = FText::FromString(UK2Node_CallFunction::GetDefaultTooltipForFunction(InFunction));
-					MenuSignature.Keywords = UK2Node_CallFunction::GetKeywordsForFunction(InFunction);
-
-					// add at least one character, so that PrimeDefaultUiSpec() doesn't attempt to query the template node
-					if (MenuSignature.Keywords.IsEmpty())
-                	{
-                		MenuSignature.Keywords = FText::FromString(TEXT(" "));
-                	}
-					
-					ActionRegistrar.AddBlueprintAction(AnimBlueprint, Spawner);
+					MenuSignature.Keywords = FText::FromString(TEXT(" "));
 				}
-			};
 
-			for(TFieldIterator<UFunction> It(InClass); It; ++It)
-			{
-				MakeFunctionAction(*It);
+				ActionRegistrar.AddBlueprintAction(InClass, Spawner);
 			}
 		};
 
+		for (TFieldIterator<UFunction> It(InClass); It; ++It)
+		{
+			MakeFunctionAction(*It);
+		}
+	};
+
+	if (const UAnimBlueprint* AnimBlueprint = Cast<UAnimBlueprint>(ActionRegistrar.GetActionKeyFilter()))
+	{
 		// Add this class
 		MakeFunctionActionsForClass(AnimBlueprint->GetAnimBlueprintGeneratedClass());
-
+	}
+	else if (ActionRegistrar.IsOpenForRegistration(GetClass()))
+	{
 		// Add function libraries too
-		for(TObjectIterator<UBlueprintFunctionLibrary> It(RF_NoFlags); It; ++It)
+		for (TObjectIterator<UBlueprintFunctionLibrary> It(RF_NoFlags); It; ++It)
 		{
 			MakeFunctionActionsForClass(It->GetClass());
 		}
