Fixed inability to spawn external linked anim layer nodes

Also fixes ensure on linked anim layer node spawn

diff --git a/Engine/Source/Editor/AnimGraph/Private/AnimGraphNode_LinkedAnimLayer.cpp b/Engine/Source/Editor/AnimGraph/Private/AnimGraphNode_LinkedAnimLayer.cpp
--- a/Engine/Source/Editor/AnimGraph/Private/AnimGraphNode_LinkedAnimLayer.cpp
+++ b/Engine/Source/Editor/AnimGraph/Private/AnimGraphNode_LinkedAnimLayer.cpp
@@ -973,7 +973,11 @@ void UAnimGraphNode_LinkedAnimLayer::HandleInstanceChanged()
 void UAnimGraphNode_LinkedAnimLayer::SetupFromLayerId(FName InLayer)
 {
 	Node.Layer = InLayer;
-	
+
+	// Set to self member first, so we have a valid name in the member reference (otherwise GetInterfaceForLayer will fail)
+	// We will override this below
+	FunctionReference.SetSelfMember(InLayer);
+
 	// Get the interface for this layer. If null, then we are using a 'self' layer.
 	Node.Interface = GetInterfaceForLayer();
 
