Fix being unable to choose another self anim layer after choosing a linked anim layer

Due to reconstruction on edit, the member reference GUID was always conforming back to the original layer node once it was set. This fix ensures that the GUID is appropriately set each time the member reference is updated.

diff --git a/Engine/Source/Editor/AnimGraph/Private/AnimGraphNode_LinkedAnimGraphBase.cpp b/Engine/Source/Editor/AnimGraph/Private/AnimGraphNode_LinkedAnimGraphBase.cpp
--- a/Engine/Source/Editor/AnimGraph/Private/AnimGraphNode_LinkedAnimGraphBase.cpp
+++ b/Engine/Source/Editor/AnimGraph/Private/AnimGraphNode_LinkedAnimGraphBase.cpp
@@ -287,7 +287,9 @@ void UAnimGraphNode_LinkedAnimGraphBase::PostEditChangeProperty(FPropertyChanged
 			{
 				if(TargetClass != nullptr)
 				{
-					FunctionReference.SetExternalMember(FunctionName, TargetClass);
+					FGuid FunctionGuid;
+					FBlueprintEditorUtils::GetFunctionGuidFromClassByFieldName(FBlueprintEditorUtils::GetMostUpToDateClass(TargetClass), FunctionName, FunctionGuid);
+					FunctionReference.SetExternalMember(FunctionName, TargetClass, FunctionGuid);
 				}
 				else
 				{
diff --git a/Engine/Source/Editor/AnimGraph/Private/AnimGraphNode_LinkedAnimLayer.cpp b/Engine/Source/Editor/AnimGraph/Private/AnimGraphNode_LinkedAnimLayer.cpp
--- a/Engine/Source/Editor/AnimGraph/Private/AnimGraphNode_LinkedAnimLayer.cpp
+++ b/Engine/Source/Editor/AnimGraph/Private/AnimGraphNode_LinkedAnimLayer.cpp
@@ -857,6 +857,8 @@ void UAnimGraphNode_LinkedAnimLayer::SetLayerName(FName InName)
 
 	if(UClass* TargetClass = GetTargetClass())
 	{
+		FGuid FunctionGuid;
+		FBlueprintEditorUtils::GetFunctionGuidFromClassByFieldName(FBlueprintEditorUtils::GetMostUpToDateClass(TargetClass), InName, FunctionGuid);
 		FunctionReference.SetExternalMember(InName, TargetClass);
 	}
 	else
@@ -895,6 +897,18 @@ void UAnimGraphNode_LinkedAnimLayer::OnLayerChanged(IDetailLayoutBuilder* Detail
 {
 	OnStructuralPropertyChanged(DetailBuilder);
 
+	UClass* TargetClass = GetTargetClass();
+	if(TargetClass)
+	{
+		FGuid FunctionGuid;
+		FBlueprintEditorUtils::GetFunctionGuidFromClassByFieldName(FBlueprintEditorUtils::GetMostUpToDateClass(TargetClass), Node.Layer, FunctionGuid);	
+		FunctionReference.SetExternalMember(Node.Layer, TargetClass);
+	}
+	else
+	{
+		FunctionReference.SetSelfMember(Node.Layer);
+	}
+	
 	// Get the interface for this layer. If null, then we are using a 'self' layer.
 	Node.Interface = GetInterfaceForLayer();
 
@@ -994,7 +1008,9 @@ void UAnimGraphNode_LinkedAnimLayer::SetupFromLayerId(FName InLayer)
 	UClass* TargetClass = GetTargetClass();
 	if(TargetClass)
 	{
-		FunctionReference.SetExternalMember(InLayer, TargetClass);
+		FGuid FunctionGuid;
+		FBlueprintEditorUtils::GetFunctionGuidFromClassByFieldName(FBlueprintEditorUtils::GetMostUpToDateClass(TargetClass), InLayer, FunctionGuid);
+		FunctionReference.SetExternalMember(InLayer, TargetClass, FunctionGuid);
 	}
 	else
 	{
