Fix crash using property access node with PinHiddenByDefault properties

diff --git a/Engine/Source/Editor/AnimGraph/Private/AnimBlueprintExtension_PropertyAccess.cpp b/Engine/Source/Editor/AnimGraph/Private/AnimBlueprintExtension_PropertyAccess.cpp
--- a/Engine/Source/Editor/AnimGraph/Private/AnimBlueprintExtension_PropertyAccess.cpp
+++ b/Engine/Source/Editor/AnimGraph/Private/AnimBlueprintExtension_PropertyAccess.cpp
@@ -19,6 +19,7 @@
 #include "IPropertyAccessEditor.h"
 #include "Internationalization/Internationalization.h"
 #include "K2Node.h"
+#include "K2Node_BreakStruct.h"
 #include "K2Node_CallFunction.h"
 #include "K2Node_GetArrayItem.h"
 #include "K2Node_VariableGet.h"
@@ -282,6 +283,16 @@ void UAnimBlueprintExtension_PropertyAccess::ExpandPropertyAccess(FKismetCompile
 				// Create a break struct/split pin node
 				const UEdGraphSchema_K2* K2Schema = CastChecked<UEdGraphSchema_K2>(InParentGraph->GetSchema());
 				UK2Node* SplitPinNode = K2Schema->CreateSplitPinNode(CurrentPin, UEdGraphSchema_K2::FCreateSplitPinNodeParams(&InCompilerContext, InParentGraph));
+				if(UK2Node_BreakStruct* BreakStructNode = Cast<UK2Node_BreakStruct>(SplitPinNode))
+				{
+					// If we made a 'break struct' node, then we need to take into account that all the pins may not be visible
+					for(FOptionalPinFromProperty& OptionalPin : BreakStructNode->ShowPinForProperties)
+					{
+						OptionalPin.bShowPin = true;
+					}
+
+					BreakStructNode->ReconstructNode();
+				}
 				UEdGraphPin* InputPin = SplitPinNode->FindPinByPredicate([ScriptStruct](UEdGraphPin* InPin)
 				{
 					if(InPin && InPin->Direction == EGPD_Input && InPin->PinType.PinCategory == UEdGraphSchema_K2::PC_Struct)
@@ -298,7 +309,12 @@ void UAnimBlueprintExtension_PropertyAccess::ExpandPropertyAccess(FKismetCompile
 				CurrentPin->MakeLinkTo(InputPin);
 
 				// Current pin is the property of the struct
-				CurrentPin = SplitPinNode->FindPinChecked(InPropertyName);
+				// Note that this can fail if the struct failed to spawn a valid pin
+				// Some custom struct pin logic (e.g. bOverride_* pins) can cause this to occur
+				if(UEdGraphPin* FoundPin = SplitPinNode->FindPin(InPropertyName))
+				{
+					CurrentPin = FoundPin;
+				}
 			}
 		}
 	};
