Fix sparse class data crash in the presence of circular dependencies

This ensures that 'old' sparse class data (that may be regenerated by a compilation) is always stashed when beginning a compilation pass. This is needed because when sparse class data is cleared, all (loaded) child classes are also cleared, so when they also come to be compiled, their sparse class data will already be null and there will be a dangling reference to it in the package's export table (despite it being moved into the transient package). The dangling reference will then cause any subsequent loads (e.g. further children) to reference an invalid structure.

diff --git a/Engine/Source/Editor/AnimGraph/Private/AnimBlueprintCompiler.cpp b/Engine/Source/Editor/AnimGraph/Private/AnimBlueprintCompiler.cpp
--- a/Engine/Source/Editor/AnimGraph/Private/AnimBlueprintCompiler.cpp
+++ b/Engine/Source/Editor/AnimGraph/Private/AnimBlueprintCompiler.cpp
@@ -157,6 +157,12 @@ FAnimBlueprintCompilerContext::FAnimBlueprintCompilerContext(UAnimBlueprint* Sou
 	// Regenerate temporary stub functions
 	// We do this here to catch the standard and 'fast' (compilation manager) compilation paths
 	CreateAnimGraphStubFunctions();
+
+	// Stash any existing sparse class data struct here, as we may regenerate it
+	if (AnimBlueprint->GeneratedClass && AnimBlueprint->GeneratedClass->GetSparseClassDataStruct())
+	{
+		OldSparseClassDataStruct = AnimBlueprint->GeneratedClass->GetSparseClassDataStruct();
+	}
 }
 
 FAnimBlueprintCompilerContext::~FAnimBlueprintCompilerContext()
@@ -1358,9 +1364,6 @@ void FAnimBlueprintCompilerContext::CleanAndSanitizeClass(UBlueprintGeneratedCla
 	UScriptStruct* CurrentSparseClassDataStruct = AnimBlueprintClassToClean->GetSparseClassDataStruct();
 	if(CurrentSparseClassDataStruct && CurrentSparseClassDataStruct->GetOuter() == AnimBlueprintClassToClean)
 	{
-		// Stash a reference to patch the linker later in RecreateSparseClassData
-		OldSparseClassDataStruct = CurrentSparseClassDataStruct;
-
 		// Only 'clear' (which renames the struct aside) if we own the sparse class data
 		// We do this because this could be a parent classes struct and we dont want to be 
 		// altering other classes during compilation
