trigger property reset to default to invalidate the target list infos

diff --git a/Engine/Source/Editor/AnimGraph/Private/PoseDriverDetails.cpp b/Engine/Source/Editor/AnimGraph/Private/PoseDriverDetails.cpp
--- a/Engine/Source/Editor/AnimGraph/Private/PoseDriverDetails.cpp
+++ b/Engine/Source/Editor/AnimGraph/Private/PoseDriverDetails.cpp
@@ -1010,6 +1010,8 @@ void FPoseDriverDetails::CustomizeDetails(IDetailLayoutBuilder& DetailBuilder)
 	UpdateDrivenNameOptions();
 
 	TSharedPtr<IPropertyHandle> PoseTargetsProperty = DetailBuilder.GetProperty("Node.PoseTargets", UAnimGraphNode_PoseDriver::StaticClass());
+	// Update target infos when resetting the property 
+	PoseTargetsProperty->SetOnPropertyResetToDefault(FSimpleDelegate::CreateSP(this, &FPoseDriverDetails::UpdateTargetInfosList));
 
 	IDetailPropertyRow& PoseTargetsRow = PoseTargetsCategory.AddProperty(PoseTargetsProperty);
 	PoseTargetsRow.ShowPropertyButtons(false);
