Additional fixup for UE-168246 to ensure that function GUID gets set up correctly in all circumstances

diff --git a/Engine/Source/Editor/AnimGraph/Private/AnimGraphNode_LinkedAnimLayer.cpp b/Engine/Source/Editor/AnimGraph/Private/AnimGraphNode_LinkedAnimLayer.cpp
--- a/Engine/Source/Editor/AnimGraph/Private/AnimGraphNode_LinkedAnimLayer.cpp
+++ b/Engine/Source/Editor/AnimGraph/Private/AnimGraphNode_LinkedAnimLayer.cpp
@@ -859,7 +859,7 @@ void UAnimGraphNode_LinkedAnimLayer::SetLayerName(FName InName)
 	{
 		FGuid FunctionGuid;
 		FBlueprintEditorUtils::GetFunctionGuidFromClassByFieldName(FBlueprintEditorUtils::GetMostUpToDateClass(TargetClass), InName, FunctionGuid);
-		FunctionReference.SetExternalMember(InName, TargetClass);
+		FunctionReference.SetExternalMember(InName, TargetClass, FunctionGuid);
 	}
 	else
 	{
@@ -902,7 +902,7 @@ void UAnimGraphNode_LinkedAnimLayer::OnLayerChanged(IDetailLayoutBuilder* Detail
 	{
 		FGuid FunctionGuid;
 		FBlueprintEditorUtils::GetFunctionGuidFromClassByFieldName(FBlueprintEditorUtils::GetMostUpToDateClass(TargetClass), Node.Layer, FunctionGuid);	
-		FunctionReference.SetExternalMember(Node.Layer, TargetClass);
+		FunctionReference.SetExternalMember(Node.Layer, TargetClass, FunctionGuid);
 	}
 	else
 	{
