AimOffsetLookAt and RotationOffset blendspaces now correctly validate for dynamic properties
#UE-174019

diff --git a/Engine/Source/Editor/AnimGraph/Private/AnimGraphNode_AimOffsetLookAt.cpp b/Engine/Source/Editor/AnimGraph/Private/AnimGraphNode_AimOffsetLookAt.cpp
--- a/Engine/Source/Editor/AnimGraph/Private/AnimGraphNode_AimOffsetLookAt.cpp
+++ b/Engine/Source/Editor/AnimGraph/Private/AnimGraphNode_AimOffsetLookAt.cpp
@@ -150,6 +150,8 @@ void UAnimGraphNode_AimOffsetLookAt::ValidateAnimNodeDuringCompilation(class USk
 {
 	Super::ValidateAnimNodeDuringCompilation(ForSkeleton, MessageLog);
 
+	ValidateAnimNodeDuringCompilationHelper(ForSkeleton, MessageLog, Node.GetBlendSpace(), UBlendSpace::StaticClass(), FindPin(GET_MEMBER_NAME_STRING_CHECKED(FAnimNode_AimOffsetLookAt, BlendSpace)), GET_MEMBER_NAME_CHECKED(FAnimNode_AimOffsetLookAt, BlendSpace));
+
 	UBlendSpace* BlendSpaceToCheck = Node.GetBlendSpace();
 	UEdGraphPin* BlendSpacePin = FindPin(GET_MEMBER_NAME_STRING_CHECKED(FAnimNode_AimOffsetLookAt, BlendSpace));
 	if (BlendSpacePin != nullptr && BlendSpaceToCheck == nullptr)
@@ -157,21 +159,14 @@ void UAnimGraphNode_AimOffsetLookAt::ValidateAnimNodeDuringCompilation(class USk
 		BlendSpaceToCheck = Cast<UBlendSpace>(BlendSpacePin->DefaultObject);
 	}
 
-	if (!BlendSpaceToCheck)
+	if (BlendSpaceToCheck)
 	{
-		// we may have a connected node
-		if (BlendSpacePin == nullptr || BlendSpacePin->LinkedTo.Num() == 0)
+		if (Cast<UAimOffsetBlendSpace>(BlendSpaceToCheck) == NULL &&
+			Cast<UAimOffsetBlendSpace1D>(BlendSpaceToCheck) == NULL)
 		{
-			MessageLog.Error(TEXT("@@ references an unknown blend space"), this);
+			MessageLog.Error(TEXT("@@ references an invalid blend space (one that is not an aim offset)"), this);
 		}
-	}
-	else if (Cast<UAimOffsetBlendSpace>(BlendSpaceToCheck) == nullptr &&
-		Cast<UAimOffsetBlendSpace1D>(BlendSpaceToCheck) == nullptr)
-	{
-		MessageLog.Error(TEXT("@@ references an invalid blend space (one that is not an aim offset)"), this);
-	}
-	else
-	{
+
 		const USkeleton* BlendSpaceSkeleton = BlendSpaceToCheck->GetSkeleton();
 		if (BlendSpaceSkeleton) // if blend space doesn't have skeleton, it might be due to blend space not loaded yet, @todo: wait with anim blueprint compilation until all assets are loaded?
 		{
@@ -184,17 +179,17 @@ void UAnimGraphNode_AimOffsetLookAt::ValidateAnimNodeDuringCompilation(class USk
 					FName SocketNameToCheck = (SocketNamePin != nullptr) ? FName(*SocketNamePin->DefaultValue) : Node.SourceSocketName;
 
 					const FReferenceSkeleton& RefSkel = BlendSpaceSkeleton->GetReferenceSkeleton();
-				
+
 					const bool bValidValue = SocketNamePin == nullptr && (
-						BlendSpaceSkeleton->FindSocket(Node.SourceSocketName) || 
+						BlendSpaceSkeleton->FindSocket(Node.SourceSocketName) ||
 						RefSkel.FindBoneIndex(Node.SourceSocketName) != INDEX_NONE);
 					const bool bValidPinValue = SocketNamePin != nullptr && (
-						BlendSpaceSkeleton->FindSocket(FName(*SocketNamePin->DefaultValue)) || 
+						BlendSpaceSkeleton->FindSocket(FName(*SocketNamePin->DefaultValue)) ||
 						RefSkel.FindBoneIndex(FName(*SocketNamePin->DefaultValue)) != INDEX_NONE);
 					const bool bValidConnectedPin = SocketNamePin != nullptr && SocketNamePin->LinkedTo.Num();
 
 					if (!bValidValue && !bValidPinValue && !bValidConnectedPin)
-					{ 
+					{
 						FFormatNamedArguments Args;
 						Args.Add(TEXT("SocketName"), FText::FromName(SocketNameToCheck));
 
@@ -211,7 +206,7 @@ void UAnimGraphNode_AimOffsetLookAt::ValidateAnimNodeDuringCompilation(class USk
 					const FReferenceSkeleton& RefSkel = BlendSpaceSkeleton->GetReferenceSkeleton();
 
 					const bool bValidValue = SocketNamePin == nullptr && (
-						Node.PivotSocketName.IsNone() || 
+						Node.PivotSocketName.IsNone() ||
 						BlendSpaceSkeleton->FindSocket(Node.PivotSocketName) ||
 						RefSkel.FindBoneIndex(Node.PivotSocketName) != INDEX_NONE);
 					const bool bValidPinValue = SocketNamePin != nullptr && (
diff --git a/Engine/Source/Editor/AnimGraph/Private/AnimGraphNode_RotationOffsetBlendSpace.cpp b/Engine/Source/Editor/AnimGraph/Private/AnimGraphNode_RotationOffsetBlendSpace.cpp
--- a/Engine/Source/Editor/AnimGraph/Private/AnimGraphNode_RotationOffsetBlendSpace.cpp
+++ b/Engine/Source/Editor/AnimGraph/Private/AnimGraphNode_RotationOffsetBlendSpace.cpp
@@ -138,6 +138,8 @@ void UAnimGraphNode_RotationOffsetBlendSpace::ValidateAnimNodeDuringCompilation(
 {
 	Super::ValidateAnimNodeDuringCompilation(ForSkeleton, MessageLog);
 
+	ValidateAnimNodeDuringCompilationHelper(ForSkeleton, MessageLog, Node.GetBlendSpace(), UBlendSpace::StaticClass(), FindPin(GET_MEMBER_NAME_STRING_CHECKED(FAnimNode_RotationOffsetBlendSpace, BlendSpace)), GET_MEMBER_NAME_CHECKED(FAnimNode_RotationOffsetBlendSpace, BlendSpace));
+
 	UBlendSpace* BlendSpaceToCheck = Node.GetBlendSpace();
 	UEdGraphPin* BlendSpacePin = FindPin(GET_MEMBER_NAME_STRING_CHECKED(FAnimNode_RotationOffsetBlendSpace, BlendSpace));
 	if (BlendSpacePin != nullptr && BlendSpaceToCheck == nullptr)
@@ -145,25 +147,12 @@ void UAnimGraphNode_RotationOffsetBlendSpace::ValidateAnimNodeDuringCompilation(
 		BlendSpaceToCheck = Cast<UBlendSpace>(BlendSpacePin->DefaultObject);
 	}
 
-	if (BlendSpaceToCheck == NULL)
-	{
-		// we may have a connected node
-		if (BlendSpacePin == nullptr || BlendSpacePin->LinkedTo.Num() == 0)
-		{
-			MessageLog.Error(TEXT("@@ references an unknown blend space"), this);
-		}
-	}
-	else if (Cast<UAimOffsetBlendSpace>(BlendSpaceToCheck) == NULL &&
-			 Cast<UAimOffsetBlendSpace1D>(BlendSpaceToCheck) == NULL)
-	{
-		MessageLog.Error(TEXT("@@ references an invalid blend space (one that is not an aim offset)"), this);
-	}
-	else if (ForSkeleton)
+	if (BlendSpaceToCheck)
 	{
-		USkeleton* BlendSpaceSkeleton = BlendSpaceToCheck->GetSkeleton();
-		if (BlendSpaceSkeleton == nullptr)
+		if (Cast<UAimOffsetBlendSpace>(BlendSpaceToCheck) == NULL &&
+			Cast<UAimOffsetBlendSpace1D>(BlendSpaceToCheck) == NULL)
 		{
-			MessageLog.Error(TEXT("@@ references blendspace that uses a missing skeleton @@"), this, BlendSpaceSkeleton);
+			MessageLog.Error(TEXT("@@ references an invalid blend space (one that is not an aim offset)"), this);
 		}
 	}
 
