Fix crash right-clicking a modify curve node in a template animation blueprint

Lack of modify curve node's utility in template anim BPs will only be fixed with the 5.3 curve refactor

diff --git a/Engine/Source/Editor/AnimGraph/Private/AnimGraphNode_ModifyCurve.cpp b/Engine/Source/Editor/AnimGraph/Private/AnimGraphNode_ModifyCurve.cpp
--- a/Engine/Source/Editor/AnimGraph/Private/AnimGraphNode_ModifyCurve.cpp
+++ b/Engine/Source/Editor/AnimGraph/Private/AnimGraphNode_ModifyCurve.cpp
@@ -38,17 +38,20 @@ TArray<FName> UAnimGraphNode_ModifyCurve::GetCurvesToAdd() const
 {
 	TArray<FName> CurvesToAdd;
 
-	const FSmartNameMapping* Mapping = GetAnimBlueprint()->TargetSkeleton->GetSmartNameContainer(USkeleton::AnimCurveMappingName);
-	if (Mapping)
+	if(HasValidBlueprint() && GetAnimBlueprint()->TargetSkeleton)
 	{
-		Mapping->FillNameArray(CurvesToAdd);
-
-		for (FName ExistingCurveName : Node.CurveNames)
+		const FSmartNameMapping* Mapping = GetAnimBlueprint()->TargetSkeleton->GetSmartNameContainer(USkeleton::AnimCurveMappingName);
+		if (Mapping)
 		{
-			CurvesToAdd.RemoveSingleSwap(ExistingCurveName, false);
-		}
+			Mapping->FillNameArray(CurvesToAdd);
+
+			for (FName ExistingCurveName : Node.CurveNames)
+			{
+				CurvesToAdd.RemoveSingleSwap(ExistingCurveName, false);
+			}
 
-		CurvesToAdd.Sort(FNameLexicalLess());
+			CurvesToAdd.Sort(FNameLexicalLess());
+		}
 	}
 
 	return CurvesToAdd;
