After five years, it's probably time to remove the 'experimental' flag from the Constraint anim node.

#rnx

diff --git a/Engine/Source/Editor/AnimGraph/Public/AnimGraphNode_Constraint.h b/Engine/Source/Editor/AnimGraph/Public/AnimGraphNode_Constraint.h
--- a/Engine/Source/Editor/AnimGraph/Public/AnimGraphNode_Constraint.h
+++ b/Engine/Source/Editor/AnimGraph/Public/AnimGraphNode_Constraint.h
@@ -7,7 +7,7 @@
 #include "EdGraph/EdGraphNodeUtils.h" // for FNodeTitleTextTable
 #include "AnimGraphNode_Constraint.generated.h"
 
-UCLASS(Experimental)
+UCLASS()
 class ANIMGRAPH_API UAnimGraphNode_Constraint : public UAnimGraphNode_SkeletalControlBase
 {
 	GENERATED_UCLASS_BODY()
