Fix stack corruption caused by calling an anim node function with an altered signature

Adds compiler checks to alert content creators to fix their content
Adds runtime parameter size checks to ensure we dont call functions with an incorrect sized set of parameters

diff --git a/Engine/Source/Editor/AnimGraph/Private/AnimGraphNode_Base.cpp b/Engine/Source/Editor/AnimGraph/Private/AnimGraphNode_Base.cpp
--- a/Engine/Source/Editor/AnimGraph/Private/AnimGraphNode_Base.cpp
+++ b/Engine/Source/Editor/AnimGraph/Private/AnimGraphNode_Base.cpp
@@ -227,7 +227,7 @@ void UAnimGraphNode_Base::ValidateAnimNodeDuringCompilation(USkeleton* ForSkelet
 		}
 	}
 
-	auto ValidateFunctionRef = [this, &MessageLog](const FMemberReference& InRef, const FText& InFunctionName)
+	auto ValidateFunctionRef = [this, &MessageLog](FName InPropertyName, const FMemberReference& InRef, const FText& InFunctionName)
 	{
 		if(InRef.GetMemberName() != NAME_None)
 		{
@@ -238,6 +238,17 @@ void UAnimGraphNode_Base::ValidateAnimNodeDuringCompilation(USkeleton* ForSkelet
 			}
 			else
 			{
+				// Check signatures match
+				const FProperty* Property = UAnimGraphNode_Base::StaticClass()->FindPropertyByName(InPropertyName);
+				check(Property);
+				const FString& PrototypeFunctionName = Property->GetMetaData("PrototypeFunction");
+				const UFunction* PrototypeFunction = PrototypeFunctionName.IsEmpty() ? nullptr : FindObject<UFunction>(nullptr, *PrototypeFunctionName);
+				if(PrototypeFunction != nullptr && !PrototypeFunction->IsSignatureCompatibleWith(Function))
+				{
+					MessageLog.Error(*FText::Format(LOCTEXT("FunctionSignatureErrorFormat", "{0} function's signature is not compatible @@"), InFunctionName).ToString(), this);
+				}
+
+				// Check thread safety
 				if(!FBlueprintEditorUtils::HasFunctionBlueprintThreadSafeMetaData(Function))
 				{
 					MessageLog.Error(*FText::Format(LOCTEXT("FunctionThreadSafetyErrorFormat", "{0} function is not thread safe @@"), InFunctionName).ToString(), this);
@@ -259,9 +270,9 @@ void UAnimGraphNode_Base::ValidateAnimNodeDuringCompilation(USkeleton* ForSkelet
 		MessageLog.Note(*(LOCTEXT("EarlyAccessNode", "@@ - Node is in early access")).ToString(), this);
 	}
 
-	ValidateFunctionRef(InitialUpdateFunction, LOCTEXT("InitialUpdateFunctionName", "Initial Update"));
-	ValidateFunctionRef(BecomeRelevantFunction, LOCTEXT("BecomeRelevantFunctionName", "Become Relevant"));
-	ValidateFunctionRef(UpdateFunction, LOCTEXT("UpdateFunctionName", "Update"));
+	ValidateFunctionRef(GET_MEMBER_NAME_CHECKED(UAnimGraphNode_Base, InitialUpdateFunction), InitialUpdateFunction, LOCTEXT("InitialUpdateFunctionName", "Initial Update"));
+	ValidateFunctionRef(GET_MEMBER_NAME_CHECKED(UAnimGraphNode_Base, BecomeRelevantFunction), BecomeRelevantFunction, LOCTEXT("BecomeRelevantFunctionName", "Become Relevant"));
+	ValidateFunctionRef(GET_MEMBER_NAME_CHECKED(UAnimGraphNode_Base, UpdateFunction), UpdateFunction, LOCTEXT("UpdateFunctionName", "Update"));
 }
 
 void UAnimGraphNode_Base::CopyTermDefaultsToDefaultObject(IAnimBlueprintCopyTermDefaultsContext& InCompilationContext, IAnimBlueprintNodeCopyTermDefaultsContext& InPerNodeContext, IAnimBlueprintGeneratedClassCompiledData& OutCompiledData)
diff --git a/Engine/Source/Runtime/Engine/Private/Animation/AnimNodeFunctionRef.cpp b/Engine/Source/Runtime/Engine/Private/Animation/AnimNodeFunctionRef.cpp
--- a/Engine/Source/Runtime/Engine/Private/Animation/AnimNodeFunctionRef.cpp
+++ b/Engine/Source/Runtime/Engine/Private/Animation/AnimNodeFunctionRef.cpp
@@ -76,23 +76,24 @@ namespace Anim
 {
 
 template<typename WrapperType, typename ContextType>
-static void CallFunctionHelper(const FAnimNodeFunctionRef& InFunction, ContextType InContext, FAnimNode_Base& InNode)
+static void CallFunctionHelper(const FAnimNodeFunctionRef& InFunction, const ContextType& InContext, FAnimNode_Base& InNode)
 {
+	struct FAnimNodeFunctionParams
+	{
+		WrapperType ExecutionContext;
+		FAnimNodeReference NodeReference;
+	};
+
 	if(InFunction.IsValid())
 	{
-		UAnimInstance* AnimInstance = CastChecked<UAnimInstance>(InContext.GetAnimInstanceObject());
-		
-		TSharedRef<FAnimExecutionContext::FData> ContextData = MakeShared<FAnimExecutionContext::FData>(InContext);
-			
-		struct FAnimNodeFunctionParams
+		if(ensureMsgf(InFunction.GetFunction()->ParmsSize == sizeof(FAnimNodeFunctionParams), TEXT("Function parameters size of %s does not match expected."), *InFunction.GetFunction()->GetName()))
 		{
-			WrapperType ExecutionContext;
-			FAnimNodeReference NodeReference;
-		};
-			
-		FAnimNodeFunctionParams Params = { WrapperType(ContextData), FAnimNodeReference(AnimInstance, InNode) };
-			
-		InFunction.Call(AnimInstance, &Params);
+			UAnimInstance* AnimInstance = CastChecked<UAnimInstance>(InContext.GetAnimInstanceObject());
+			TSharedRef<FAnimExecutionContext::FData> ContextData = MakeShared<FAnimExecutionContext::FData>(InContext);
+			FAnimNodeFunctionParams Params = { WrapperType(ContextData), FAnimNodeReference(AnimInstance, InNode) };
+
+			InFunction.Call(AnimInstance, &Params);
+		}
 	}
 }
 
