[PoseDriver] PoseAsset is not required anymore when using DriveCurves

diff --git a/Engine/Source/Editor/AnimGraph/Private/AnimGraphNode_PoseDriver.cpp b/Engine/Source/Editor/AnimGraph/Private/AnimGraphNode_PoseDriver.cpp
--- a/Engine/Source/Editor/AnimGraph/Private/AnimGraphNode_PoseDriver.cpp
+++ b/Engine/Source/Editor/AnimGraph/Private/AnimGraphNode_PoseDriver.cpp
@@ -149,7 +149,16 @@ void UAnimGraphNode_PoseDriver::ValidateAnimNodeDuringCompilation(USkeleton* For
 			this, GetData(Node.PoseTargets[Node.SoloTargetIndex].DrivenName.ToString()));
 	}
 
-	Super::ValidateAnimNodeDuringCompilation(ForSkeleton, MessageLog);
+	// Note: UAnimGraphNode_PoseHandler::ValidateAnimNodeDuringCompilation checks if PoseAsset is valid, 
+	// This check is only necessary when using DrivePoses
+	if (Node.DriveOutput == EPoseDriverOutput::DrivePoses)
+	{
+		Super::ValidateAnimNodeDuringCompilation(ForSkeleton, MessageLog);
+	}
+	else
+	{
+		UAnimGraphNode_AssetPlayerBase::ValidateAnimNodeDuringCompilation(ForSkeleton, MessageLog);
+	}
 }
 
 FEditorModeID UAnimGraphNode_PoseDriver::GetEditorMode() const
