Output warnings when the LegIK node is missing definitions or contains references to invalid bones.

diff --git a/Engine/Source/Editor/AnimGraph/Private/AnimGraphNode_LegIK.cpp b/Engine/Source/Editor/AnimGraph/Private/AnimGraphNode_LegIK.cpp
--- a/Engine/Source/Editor/AnimGraph/Private/AnimGraphNode_LegIK.cpp
+++ b/Engine/Source/Editor/AnimGraph/Private/AnimGraphNode_LegIK.cpp
@@ -27,5 +27,27 @@ FText UAnimGraphNode_LegIK::GetNodeTitle(ENodeTitleType::Type TitleType) const
 	return GetControllerDescription();
 }
 
+void UAnimGraphNode_LegIK::ValidateAnimNodeDuringCompilation(USkeleton* ForSkeleton, FCompilerResultsLog& MessageLog)
+{
+	if (Node.LegsDefinition.IsEmpty())
+	{
+		MessageLog.Warning(*LOCTEXT("MissingLegsDefinition", "@@ - Legs Definition is empty").ToString(), this);
+	}
+	else if (ForSkeleton)
+	{
+		for (const FAnimLegIKDefinition& LegDefinition : Node.LegsDefinition)
+		{
+			if (ForSkeleton->GetReferenceSkeleton().FindBoneIndex(LegDefinition.FKFootBone.BoneName) == INDEX_NONE)
+			{
+				MessageLog.Warning(*LOCTEXT("InvalidFKFootBone", "@@ - Invalid FKFoot Bone in Legs Definition").ToString(), this);
+			}
+
+			if (ForSkeleton->GetReferenceSkeleton().FindBoneIndex(LegDefinition.IKFootBone.BoneName) == INDEX_NONE)
+			{
+				MessageLog.Warning(*LOCTEXT("InvalidFKFootBone", "@@ - Invalid IKFoot Bone in Legs Definition").ToString(), this);
+			}
+		}
+	}
+}
 
 #undef LOCTEXT_NAMESPACE
diff --git a/Engine/Source/Editor/AnimGraph/Public/AnimGraphNode_LegIK.h b/Engine/Source/Editor/AnimGraph/Public/AnimGraphNode_LegIK.h
--- a/Engine/Source/Editor/AnimGraph/Public/AnimGraphNode_LegIK.h
+++ b/Engine/Source/Editor/AnimGraph/Public/AnimGraphNode_LegIK.h
@@ -21,6 +21,7 @@ public:
 	// UEdGraphNode interface
 	virtual FText GetNodeTitle(ENodeTitleType::Type TitleType) const override;
 	virtual FText GetTooltipText() const override;
+	virtual void ValidateAnimNodeDuringCompilation(USkeleton* ForSkeleton, FCompilerResultsLog& MessageLog) override;
 	// End of UEdGraphNode interface
 
 protected:
