Fix loctext in legik node

diff --git a/Engine/Source/Editor/AnimGraph/Private/AnimGraphNode_LegIK.cpp b/Engine/Source/Editor/AnimGraph/Private/AnimGraphNode_LegIK.cpp
--- a/Engine/Source/Editor/AnimGraph/Private/AnimGraphNode_LegIK.cpp
+++ b/Engine/Source/Editor/AnimGraph/Private/AnimGraphNode_LegIK.cpp
@@ -44,7 +44,7 @@ void UAnimGraphNode_LegIK::ValidateAnimNodeDuringCompilation(USkeleton* ForSkele
 
 			if (ForSkeleton->GetReferenceSkeleton().FindBoneIndex(LegDefinition.IKFootBone.BoneName) == INDEX_NONE)
 			{
-				MessageLog.Warning(*LOCTEXT("InvalidFKFootBone", "@@ - Invalid IKFoot Bone in Legs Definition").ToString(), this);
+				MessageLog.Warning(*LOCTEXT("InvalidIKFootBone", "@@ - Invalid IKFoot Bone in Legs Definition").ToString(), this);
 			}
 		}
 	}
