Disabled option to duplicate a state from the tree view.

This option is being disabled for now since it has never worked properly and it causes a engine crash.

diff --git a/Engine/Source/Editor/AnimGraph/Private/AnimationStateGraphSchema.cpp b/Engine/Source/Editor/AnimGraph/Private/AnimationStateGraphSchema.cpp
--- a/Engine/Source/Editor/AnimGraph/Private/AnimationStateGraphSchema.cpp
+++ b/Engine/Source/Editor/AnimGraph/Private/AnimationStateGraphSchema.cpp
@@ -39,4 +39,10 @@ void UAnimationStateGraphSchema::GetGraphDisplayInformation(const UEdGraph& Grap
 	DisplayInfo.DisplayName = DisplayInfo.PlainName;
 }
 
+bool UAnimationStateGraphSchema::CanDuplicateGraph(UEdGraph* InSourceGraph) const
+{
+	return false;
+}
+
+
 #undef LOCTEXT_NAMESPACE
diff --git a/Engine/Source/Editor/AnimGraph/Public/AnimationStateGraphSchema.h b/Engine/Source/Editor/AnimGraph/Public/AnimationStateGraphSchema.h
--- a/Engine/Source/Editor/AnimGraph/Public/AnimationStateGraphSchema.h
+++ b/Engine/Source/Editor/AnimGraph/Public/AnimationStateGraphSchema.h
@@ -15,5 +15,6 @@ class UAnimationStateGraphSchema : public UAnimationGraphSchema
 	//~ Begin UEdGraphSchema Interface.
 	virtual void CreateDefaultNodesForGraph(UEdGraph& Graph) const override;
 	virtual void GetGraphDisplayInformation(const UEdGraph& Graph, /*out*/ FGraphDisplayInfo& DisplayInfo) const override;
+	virtual bool CanDuplicateGraph(UEdGraph* InSourceGraph) const override;;
 	//~ End UEdGraphSchema Interface.
 };
