Added missing tooltips for animation state machine schemas

diff --git a/Engine/Source/Editor/AnimGraph/Private/AnimationConduitGraphSchema.cpp b/Engine/Source/Editor/AnimGraph/Private/AnimationConduitGraphSchema.cpp
--- a/Engine/Source/Editor/AnimGraph/Private/AnimationConduitGraphSchema.cpp
+++ b/Engine/Source/Editor/AnimGraph/Private/AnimationConduitGraphSchema.cpp
@@ -6,6 +6,8 @@
 #include "AnimationTransitionGraph.h"
 #include "AnimStateConduitNode.h"
 
+#define LOCTEXT_NAMESPACE "AnimationConduitSchema"
+
 /////////////////////////////////////////////////////
 // UAnimationConduitGraphSchema
 
@@ -35,10 +37,11 @@ void UAnimationConduitGraphSchema::GetGraphDisplayInformation(const UEdGraph& Gr
 		FFormatNamedArguments Args;
 		Args.Add(TEXT("NodeTitle"), ConduitNode->GetNodeTitle(ENodeTitleType::FullTitle) );
 
-		DisplayInfo.PlainName = FText::Format( NSLOCTEXT("Animation", "ConduitRuleGraphTitle", "{NodeTitle} (conduit rule)"), Args);
+		DisplayInfo.PlainName = FText::Format( LOCTEXT("ConduitRuleGraphTitle", "{NodeTitle} (conduit rule)"), Args);
 	}
 
 	DisplayInfo.DisplayName = DisplayInfo.PlainName;
+	DisplayInfo.Tooltip = LOCTEXT("GraphTooltip_ConduitSchema", "Conduits can be used to create 1-to-many, many-to-1, or many-to-many transitions");
 }
 
 void UAnimationConduitGraphSchema::HandleGraphBeingDeleted(UEdGraph& GraphBeingRemoved) const
@@ -74,3 +77,5 @@ void UAnimationConduitGraphSchema::HandleGraphBeingDeleted(UEdGraph& GraphBeingR
 		}
 	}
 }
+
+#undef LOCTEXT_NAMESPACE
\ No newline at end of file
diff --git a/Engine/Source/Editor/AnimGraph/Private/AnimationCustomTransitionSchema.cpp b/Engine/Source/Editor/AnimGraph/Private/AnimationCustomTransitionSchema.cpp
--- a/Engine/Source/Editor/AnimGraph/Private/AnimationCustomTransitionSchema.cpp
+++ b/Engine/Source/Editor/AnimGraph/Private/AnimationCustomTransitionSchema.cpp
@@ -7,6 +7,8 @@
 #include "AnimStateTransitionNode.h"
 #include "AnimationCustomTransitionGraph.h"
 
+#define LOCTEXT_NAMESPACE "AnimationCustomTransitionSchema"
+
 /////////////////////////////////////////////////////
 // UAnimationCustomTransitionSchema
 
@@ -56,10 +58,11 @@ void UAnimationCustomTransitionSchema::GetGraphDisplayInformation(const UEdGraph
 
 	if (const UAnimStateTransitionNode* TransNode = Cast<const UAnimStateTransitionNode>(Graph.GetOuter()))
 	{
-		DisplayInfo.PlainName = FText::Format( NSLOCTEXT("Animation", "CustomBlendGraphTitle", "{0} (custom blend)"), TransNode->GetNodeTitle(ENodeTitleType::FullTitle) );
+		DisplayInfo.PlainName = FText::Format( LOCTEXT("CustomBlendGraphTitle", "{0} (custom blend)"), TransNode->GetNodeTitle(ENodeTitleType::FullTitle) );
 	}
 
 	DisplayInfo.DisplayName = DisplayInfo.PlainName;
+	DisplayInfo.Tooltip = LOCTEXT("GraphTooltip_CustomTransitionSchema", "Custom blend graphs offer manual blending from source to destination states");
 }
 
 void UAnimationCustomTransitionSchema::HandleGraphBeingDeleted(UEdGraph& GraphBeingRemoved) const
@@ -82,3 +85,5 @@ void UAnimationCustomTransitionSchema::HandleGraphBeingDeleted(UEdGraph& GraphBe
 		}
 	}
 }
+
+#undef LOCTEXT_NAMESPACE
\ No newline at end of file
diff --git a/Engine/Source/Editor/AnimGraph/Private/AnimationStateGraphSchema.cpp b/Engine/Source/Editor/AnimGraph/Private/AnimationStateGraphSchema.cpp
--- a/Engine/Source/Editor/AnimGraph/Private/AnimationStateGraphSchema.cpp
+++ b/Engine/Source/Editor/AnimGraph/Private/AnimationStateGraphSchema.cpp
@@ -37,6 +37,7 @@ void UAnimationStateGraphSchema::GetGraphDisplayInformation(const UEdGraph& Grap
 	}
 
 	DisplayInfo.DisplayName = DisplayInfo.PlainName;
+	DisplayInfo.Tooltip = LOCTEXT("GraphTooltip_StateSchema", "States contain animation graphs that contribute to the state machine's output pose when the state is active");
 }
 
 bool UAnimationStateGraphSchema::CanDuplicateGraph(UEdGraph* InSourceGraph) const
diff --git a/Engine/Source/Editor/AnimGraph/Private/AnimationStateMachineSchema.cpp b/Engine/Source/Editor/AnimGraph/Private/AnimationStateMachineSchema.cpp
--- a/Engine/Source/Editor/AnimGraph/Private/AnimationStateMachineSchema.cpp
+++ b/Engine/Source/Editor/AnimGraph/Private/AnimationStateMachineSchema.cpp
@@ -453,6 +453,7 @@ void UAnimationStateMachineSchema::GetGraphDisplayInformation(const UEdGraph& Gr
 {
 	DisplayInfo.PlainName = FText::FromString( Graph.GetName() );
 	DisplayInfo.DisplayName = DisplayInfo.PlainName;
+	DisplayInfo.Tooltip = LOCTEXT("GraphTooltip_StateMachineSchema", "Graph used to transition between different states each with separate animation graphs");
 }
 
 void UAnimationStateMachineSchema::DroppedAssetsOnGraph(const TArray<FAssetData>& Assets, const FVector2D& GraphPosition, UEdGraph* Graph) const
diff --git a/Engine/Source/Editor/AnimGraph/Private/AnimationTransitionSchema.cpp b/Engine/Source/Editor/AnimGraph/Private/AnimationTransitionSchema.cpp
--- a/Engine/Source/Editor/AnimGraph/Private/AnimationTransitionSchema.cpp
+++ b/Engine/Source/Editor/AnimGraph/Private/AnimationTransitionSchema.cpp
@@ -59,7 +59,8 @@ void UAnimationTransitionSchema::GetGraphDisplayInformation(const UEdGraph& Grap
 	{
 		FFormatNamedArguments Args;
 		Args.Add(TEXT("NodeTitle"), TransNode->GetNodeTitle(ENodeTitleType::FullTitle));
-		DisplayInfo.PlainName = FText::Format( NSLOCTEXT("Animation", "TransitionRuleGraphTitle", "{NodeTitle} (rule)"), Args );
+		DisplayInfo.PlainName = FText::Format( LOCTEXT( "TransitionRuleGraphTitle", "{NodeTitle} (rule)"), Args );
+		DisplayInfo.Tooltip = LOCTEXT("GraphTooltip_TransitionSchema", "Transitions contain rules that define when to move between states");
 	}
 
 	DisplayInfo.DisplayName = DisplayInfo.PlainName;
diff --git a/Engine/Source/Editor/Kismet/Private/SBlueprintPalette.cpp b/Engine/Source/Editor/Kismet/Private/SBlueprintPalette.cpp
--- a/Engine/Source/Editor/Kismet/Private/SBlueprintPalette.cpp
+++ b/Engine/Source/Editor/Kismet/Private/SBlueprintPalette.cpp
@@ -7,6 +7,7 @@
 #include "AnimationGraph.h"
 #include "AnimationStateGraph.h"
 #include "AnimationStateMachineGraph.h"
+#include "AnimationCustomTransitionGraph.h"
 #include "AnimationStateMachineSchema.h"
 #include "AnimationTransitionGraph.h"
 #include "AssetRegistry/AssetData.h"
@@ -274,6 +275,11 @@ static void GetSubGraphIcon(FEdGraphSchemaAction_K2Graph const* const ActionIn,
 				IconOut = FAppStyle::GetBrush(TEXT("BlendSpace.SampleGraph") );
 				ToolTipOut = LOCTEXT("BlendSpaceSample_ToolTip", "BlendSpace Sample");
 			}
+			else if (Cast<UAnimationCustomTransitionGraph>(ActionIn->EdGraph))
+			{
+				IconOut = FAppStyle::GetBrush(TEXT("GraphEditor.SubGraph_16x"));
+				ToolTipOut = LOCTEXT("CustomBlendGraph_ToolTip", "Custom Blend");
+			}
 			else
 			{
 				IconOut = FAppStyle::GetBrush(TEXT("GraphEditor.SubGraph_16x") );
