[UE] Prevent the creation of multiple copies of pose watch elements by the RBAN AnimGraphNode
- Find and reuse existing Pose Watch Elements in the RBAN node, rather than creating new ones every time the editor is re-opened.

#rnx

diff --git a/Engine/Source/Editor/AnimGraph/Private/AnimGraphNode_RigidBody.cpp b/Engine/Source/Editor/AnimGraph/Private/AnimGraphNode_RigidBody.cpp
--- a/Engine/Source/Editor/AnimGraph/Private/AnimGraphNode_RigidBody.cpp
+++ b/Engine/Source/Editor/AnimGraph/Private/AnimGraphNode_RigidBody.cpp
@@ -90,11 +90,11 @@ void UAnimGraphNode_RigidBody::OnPoseWatchChanged(const bool IsPoseWatchEnabled,
 
 	UPoseWatch* const PoseWatch = InPoseWatch.Get();
 
-	if (PoseWatch && !(PoseWatch->Contains(PoseWatchElementBodies.Get()) && PoseWatch->Contains(PoseWatchElementConstraints.Get())))
+	if (PoseWatch)
 	{
 		// A new pose watch has been created for this node - add node specific pose watch components.
-		PoseWatchElementBodies = InPoseWatch.Get()->AddElement(FText(LOCTEXT("PoseWatchElementLabel_RigidBody_PhysicsBodies", "Physics Bodies")), TEXT("PhysicsAssetEditor.Tree.Body"));
-		PoseWatchElementConstraints = InPoseWatch.Get()->AddElement(FText(LOCTEXT("PoseWatchElementLabel__RigidBody_PhysicsConstraints", "Physics Constraints")), TEXT("PhysicsAssetEditor.Tree.Constraint"));
+		PoseWatchElementBodies = InPoseWatch.Get()->FindOrAddElement(FText(LOCTEXT("PoseWatchElementLabel_RigidBody_PhysicsBodies", "Physics Bodies")), TEXT("PhysicsAssetEditor.Tree.Body"));
+		PoseWatchElementConstraints = InPoseWatch.Get()->FindOrAddElement(FText(LOCTEXT("PoseWatchElementLabel__RigidBody_PhysicsConstraints", "Physics Constraints")), TEXT("PhysicsAssetEditor.Tree.Constraint"));
 
 		check(PoseWatchElementConstraints.IsValid()); // Expect to find a valid component;
 		if (PoseWatchElementConstraints.IsValid())
diff --git a/Engine/Source/Runtime/Engine/Classes/Engine/PoseWatch.h b/Engine/Source/Runtime/Engine/Classes/Engine/PoseWatch.h
--- a/Engine/Source/Runtime/Engine/Classes/Engine/PoseWatch.h
+++ b/Engine/Source/Runtime/Engine/Classes/Engine/PoseWatch.h
@@ -363,6 +363,12 @@ public:
 	/** Returns a reference to the array of elements stored in this */
 	TArray<TObjectPtr<UPoseWatchElement>>& GetElements() { return Elements; }
 
+	/** Find an element with the supplied label. */
+	TObjectPtr<UPoseWatchElement> FindElement(const FText InLabel);
+
+	/** Find and return and existing element with the supplied label or, if none exists, create a new element and add it to this pose watch. */
+	TObjectPtr<UPoseWatchElement> FindOrAddElement(const FText InLabel, const FName InIconName);
+
 	/** Returns true if this pose watch contains the specified element */
 	bool Contains(const TObjectPtr<UPoseWatchElement> InElement) { return Elements.Contains(InElement); }
 	bool Contains(const UPoseWatchElement* const InElement) { return Elements.ContainsByPredicate([InElement](const TObjectPtr<UPoseWatchElement>& ContainedElement) { return InElement == ContainedElement.Get(); }); }
diff --git a/Engine/Source/Runtime/Engine/Private/Engine/PoseWatch.cpp b/Engine/Source/Runtime/Engine/Private/Engine/PoseWatch.cpp
--- a/Engine/Source/Runtime/Engine/Private/Engine/PoseWatch.cpp
+++ b/Engine/Source/Runtime/Engine/Private/Engine/PoseWatch.cpp
@@ -759,6 +759,35 @@ TObjectPtr<UPoseWatchElement> UPoseWatch::AddElement(const FText InLabel, const
 	return AddElement<UPoseWatchElement>(InLabel, InIconName);
 }
 
+TObjectPtr<UPoseWatchElement> UPoseWatch::FindElement(const FText InLabel)
+{
+	TObjectPtr<UPoseWatchElement> Element(nullptr);
+
+	if (TObjectPtr<UPoseWatchElement>* FoundElement = GetElements().FindByPredicate([InLabel](const TObjectPtr<UPoseWatchElement>& Element) { return Element->GetLabel().EqualTo(InLabel); }))
+	{
+		Element = *FoundElement;
+	}
+
+	return Element;
+}
+
+TObjectPtr<UPoseWatchElement> UPoseWatch::FindOrAddElement(const FText InLabel, const FName InIconName)
+{
+	TObjectPtr<UPoseWatchElement> Element = FindElement(InLabel);
+
+	if (!Element)
+	{
+		Element = AddElement(InLabel, InIconName);
+	}
+	else
+	{
+		Element->SetIconName(InIconName);
+	}
+
+	return Element;
+}
+
+
 #endif // WITH_EDITOR
 
 #if WITH_EDITORONLY_DATA
