State machine entry, exit, fully blended in, and fully blended out anim node functions.

Additional changes:
- Added warnings for users using anim state events since we want to push towards using the state node functions.
- Added support for anim nodes to show any bound functions and not have them hardcoded.
- Added base context type for FAnimExecutionContext.
- Now users can specify the string for the default binding name of a function member reference in a BP by using the "DefaultBindingName" metadata tag. This will give the functionality to avoid having the default binding name be "NewFunction" and instead use the input string from the metadata tag.

diff --git a/Engine/Config/BaseEngine.ini b/Engine/Config/BaseEngine.ini
--- a/Engine/Config/BaseEngine.ini
+++ b/Engine/Config/BaseEngine.ini
@@ -1146,6 +1146,7 @@ DefaultWorldPartitionHLODLayer=/Engine/Maps/Templates/HLODs/HLODLayer_Instanced.
 +ClassRedirects=(OldName="/Script/UMGEditor.DetailsView", NewName="/Script/ScriptableEditorWidgets.DetailsView")
 +ClassRedirects=(OldName="/Script/UMGEditor.SinglePropertyView", NewName="/Script/ScriptableEditorWidgets.SinglePropertyView")
 +ClassRedirects=(OldName="/Script/UMGEditor.PropertyViewBase", NewName="/Script/ScriptableEditorWidgets.PropertyViewBase")
++StructRedirects=(OldName="/Script/AnimGraphRuntime.AnimNode_StateResult", NewName="/Script/Engine.AnimNode_StateResult")
 
 ; Refactoring most of HeadMountedDisplay engine module into XRBase plugin.
 +ClassRedirects=(OldName="HeadMountedDisplayFunctionLibrary", NewName="/Script/XRBase.HeadMountedDisplayFunctionLibrary")
diff --git a/Engine/Source/Editor/AnimGraph/Private/AnimGraphNode_Base.cpp b/Engine/Source/Editor/AnimGraph/Private/AnimGraphNode_Base.cpp
--- a/Engine/Source/Editor/AnimGraph/Private/AnimGraphNode_Base.cpp
+++ b/Engine/Source/Editor/AnimGraph/Private/AnimGraphNode_Base.cpp
@@ -201,6 +201,56 @@ void UAnimGraphNode_Base::CreateOutputPins()
 	}
 }
 
+void UAnimGraphNode_Base::ValidateFunctionRef(FName InPropertyName, const FMemberReference& InRef, const FText& InFunctionName, FCompilerResultsLog& MessageLog)
+{
+	if (InRef.GetMemberName() != NAME_None)
+	{
+		const UFunction* Function = InRef.ResolveMember<UFunction>(GetAnimBlueprint()->SkeletonGeneratedClass);
+		if(Function == nullptr)
+		{
+			MessageLog.Error(*FText::Format(LOCTEXT("CouldNotResolveFunctionErrorFormat", "Could not resolve {0} function @@"), InFunctionName).ToString(), this);
+		}
+		else
+		{
+			// Check signatures match
+			const FProperty* Property = GetClass()->FindPropertyByName(InPropertyName);
+			check(Property);
+			const FString& PrototypeFunctionName = Property->GetMetaData("PrototypeFunction");
+			const UFunction* PrototypeFunction = PrototypeFunctionName.IsEmpty() ? nullptr : FindObject<UFunction>(nullptr, *PrototypeFunctionName);
+			if (PrototypeFunction != nullptr && !PrototypeFunction->IsSignatureCompatibleWith(Function))
+			{
+				MessageLog.Error(*FText::Format(LOCTEXT("FunctionSignatureErrorFormat", "{0} function's signature is not compatible @@"), InFunctionName).ToString(), this);
+			}
+
+			// Check thread safety
+			if (!FBlueprintEditorUtils::HasFunctionBlueprintThreadSafeMetaData(Function))
+			{
+				MessageLog.Error(*FText::Format(LOCTEXT("FunctionThreadSafetyErrorFormat", "{0} function is not thread safe @@"), InFunctionName).ToString(), this);
+			}
+		}
+	}
+}
+
+void UAnimGraphNode_Base::GetBoundFunctionsInfo(TArray<TPair<FName, FName>>& InOutBindingsInfo)
+{
+	FName CategoryName = TEXT("Functions");
+
+	if (InitialUpdateFunction.ResolveMember<UFunction>(GetBlueprintClassFromNode()) != nullptr)
+	{
+		InOutBindingsInfo.Emplace(CategoryName, GET_MEMBER_NAME_CHECKED(UAnimGraphNode_Base, InitialUpdateFunction));
+	}
+	
+	if (BecomeRelevantFunction.ResolveMember<UFunction>(GetBlueprintClassFromNode()) != nullptr)
+	{
+		InOutBindingsInfo.Emplace(CategoryName, GET_MEMBER_NAME_CHECKED(UAnimGraphNode_Base, BecomeRelevantFunction));
+	}
+
+	if (UpdateFunction.ResolveMember<UFunction>(GetBlueprintClassFromNode()) != nullptr)
+	{
+		InOutBindingsInfo.Emplace(CategoryName, GET_MEMBER_NAME_CHECKED(UAnimGraphNode_Base, UpdateFunction));
+	}
+};
+
 void UAnimGraphNode_Base::ValidateAnimNodeDuringCompilation(USkeleton* ForSkeleton, FCompilerResultsLog& MessageLog)
 {
 	// Validate any bone references we have
@@ -227,36 +277,6 @@ void UAnimGraphNode_Base::ValidateAnimNodeDuringCompilation(USkeleton* ForSkelet
 		}
 	}
 
-	auto ValidateFunctionRef = [this, &MessageLog](FName InPropertyName, const FMemberReference& InRef, const FText& InFunctionName)
-	{
-		if(InRef.GetMemberName() != NAME_None)
-		{
-			UFunction* Function = InRef.ResolveMember<UFunction>(GetAnimBlueprint()->SkeletonGeneratedClass);
-			if(Function == nullptr)
-			{
-				MessageLog.Error(*FText::Format(LOCTEXT("CouldNotResolveFunctionErrorFormat", "Could not resolve {0} function @@"), InFunctionName).ToString(), this);
-			}
-			else
-			{
-				// Check signatures match
-				const FProperty* Property = UAnimGraphNode_Base::StaticClass()->FindPropertyByName(InPropertyName);
-				check(Property);
-				const FString& PrototypeFunctionName = Property->GetMetaData("PrototypeFunction");
-				const UFunction* PrototypeFunction = PrototypeFunctionName.IsEmpty() ? nullptr : FindObject<UFunction>(nullptr, *PrototypeFunctionName);
-				if(PrototypeFunction != nullptr && !PrototypeFunction->IsSignatureCompatibleWith(Function))
-				{
-					MessageLog.Error(*FText::Format(LOCTEXT("FunctionSignatureErrorFormat", "{0} function's signature is not compatible @@"), InFunctionName).ToString(), this);
-				}
-
-				// Check thread safety
-				if(!FBlueprintEditorUtils::HasFunctionBlueprintThreadSafeMetaData(Function))
-				{
-					MessageLog.Error(*FText::Format(LOCTEXT("FunctionThreadSafetyErrorFormat", "{0} function is not thread safe @@"), InFunctionName).ToString(), this);
-				}
-			}
-		}
-	};
-	
 	bool bBaseClassIsExperimental = false;
 	bool bBaseClassIsEarlyAccess = false;
 	FString MostDerivedDevelopmentClassName;
@@ -270,9 +290,9 @@ void UAnimGraphNode_Base::ValidateAnimNodeDuringCompilation(USkeleton* ForSkelet
 		MessageLog.Note(*(LOCTEXT("EarlyAccessNode", "@@ - Node is in early access")).ToString(), this);
 	}
 
-	ValidateFunctionRef(GET_MEMBER_NAME_CHECKED(UAnimGraphNode_Base, InitialUpdateFunction), InitialUpdateFunction, LOCTEXT("InitialUpdateFunctionName", "Initial Update"));
-	ValidateFunctionRef(GET_MEMBER_NAME_CHECKED(UAnimGraphNode_Base, BecomeRelevantFunction), BecomeRelevantFunction, LOCTEXT("BecomeRelevantFunctionName", "Become Relevant"));
-	ValidateFunctionRef(GET_MEMBER_NAME_CHECKED(UAnimGraphNode_Base, UpdateFunction), UpdateFunction, LOCTEXT("UpdateFunctionName", "Update"));
+	ValidateFunctionRef(GET_MEMBER_NAME_CHECKED(UAnimGraphNode_Base, InitialUpdateFunction), InitialUpdateFunction, LOCTEXT("InitialUpdateFunctionName", "Initial Update"), MessageLog);
+	ValidateFunctionRef(GET_MEMBER_NAME_CHECKED(UAnimGraphNode_Base, BecomeRelevantFunction), BecomeRelevantFunction, LOCTEXT("BecomeRelevantFunctionName", "Become Relevant"), MessageLog);
+	ValidateFunctionRef(GET_MEMBER_NAME_CHECKED(UAnimGraphNode_Base, UpdateFunction), UpdateFunction, LOCTEXT("UpdateFunctionName", "Update"), MessageLog);
 }
 
 void UAnimGraphNode_Base::CopyTermDefaultsToDefaultObject(IAnimBlueprintCopyTermDefaultsContext& InCompilationContext, IAnimBlueprintNodeCopyTermDefaultsContext& InPerNodeContext, IAnimBlueprintGeneratedClassCompiledData& OutCompiledData)
diff --git a/Engine/Source/Editor/AnimGraph/Private/AnimGraphNode_StateResult.cpp b/Engine/Source/Editor/AnimGraph/Private/AnimGraphNode_StateResult.cpp
--- a/Engine/Source/Editor/AnimGraph/Private/AnimGraphNode_StateResult.cpp
+++ b/Engine/Source/Editor/AnimGraph/Private/AnimGraphNode_StateResult.cpp
@@ -1,6 +1,7 @@
 // Copyright Epic Games, Inc. All Rights Reserved.
 
 #include "AnimGraphNode_StateResult.h"
+#include "FindInBlueprintManager.h"
 #include "GraphEditorSettings.h"
 #include "IAnimBlueprintCompilationContext.h"
 
@@ -49,6 +50,114 @@ void UAnimGraphNode_StateResult::OnProcessDuringCompilation(IAnimBlueprintCompil
 	UAnimGraphNode_StateResult* TrueNode = InCompilationContext.GetMessageLog().FindSourceObjectTypeChecked<UAnimGraphNode_StateResult>(this);
 
 	Node.SetName(TrueNode->GetGraph()->GetFName());
+
+	// Resolve functions
+	Node.StateEntryFunction.SetFromFunction(StateEntryFunction.ResolveMember<UFunction>(GetBlueprintClassFromNode()));
+	Node.StateFullyBlendedInFunction.SetFromFunction(StateFullyBlendedInFunction.ResolveMember<UFunction>(GetBlueprintClassFromNode()));
+	Node.StateExitFunction.SetFromFunction(StateExitFunction.ResolveMember<UFunction>(GetBlueprintClassFromNode()));
+	Node.StateFullyBlendedOutFunction.SetFromFunction(StateFullyBlendedOutFunction.ResolveMember<UFunction>(GetBlueprintClassFromNode()));
+}
+
+void UAnimGraphNode_StateResult::ValidateAnimNodeDuringCompilation(USkeleton* ForSkeleton, FCompilerResultsLog& MessageLog)
+{
+	Super::ValidateAnimNodeDuringCompilation(ForSkeleton, MessageLog);
+	
+	ValidateFunctionRef(GET_MEMBER_NAME_CHECKED(UAnimGraphNode_StateResult, StateEntryFunction), StateEntryFunction, LOCTEXT("StateEntryFunctionName", "State Entry"), MessageLog);
+	ValidateFunctionRef(GET_MEMBER_NAME_CHECKED(UAnimGraphNode_StateResult, StateFullyBlendedInFunction), StateFullyBlendedInFunction, LOCTEXT("StateFullyBlendedInFunctionName", "State Fully Blended In"), MessageLog);
+	ValidateFunctionRef(GET_MEMBER_NAME_CHECKED(UAnimGraphNode_StateResult, StateExitFunction), StateExitFunction, LOCTEXT("StateExitFunctionName", "State Exit"), MessageLog);
+	ValidateFunctionRef(GET_MEMBER_NAME_CHECKED(UAnimGraphNode_StateResult, StateFullyBlendedOutFunction), StateFullyBlendedOutFunction, LOCTEXT("StateFullyBlendedOutFunctionName", "State Fully Blended Out"), MessageLog);
+}
+
+void UAnimGraphNode_StateResult::PostEditChangeProperty(FPropertyChangedEvent& PropertyChangedEvent)
+{
+	const FName PropertyName = (PropertyChangedEvent.Property != nullptr) ? PropertyChangedEvent.Property->GetFName() : NAME_None;
+
+	bool bAnyAnimNodeFunctionChanged = false;
+	
+	if (PropertyName == GET_MEMBER_NAME_CHECKED(UAnimGraphNode_StateResult, StateEntryFunction))
+	{
+		Node.StateEntryFunction.SetFromFunction(StateEntryFunction.ResolveMember<UFunction>(GetBlueprintClassFromNode()));
+		bAnyAnimNodeFunctionChanged = true;
+	}
+	else if (PropertyName == GET_MEMBER_NAME_CHECKED(UAnimGraphNode_StateResult, StateFullyBlendedInFunction))
+	{
+		Node.StateFullyBlendedInFunction.SetFromFunction(StateFullyBlendedInFunction.ResolveMember<UFunction>(GetBlueprintClassFromNode()));
+		bAnyAnimNodeFunctionChanged = true;
+	}
+	else if (PropertyName == GET_MEMBER_NAME_CHECKED(UAnimGraphNode_StateResult, StateExitFunction))
+	{
+		Node.StateExitFunction.SetFromFunction(StateExitFunction.ResolveMember<UFunction>(GetBlueprintClassFromNode()));
+		bAnyAnimNodeFunctionChanged = true;
+	}
+	else if (PropertyName == GET_MEMBER_NAME_CHECKED(UAnimGraphNode_StateResult, StateFullyBlendedOutFunction))
+	{
+		Node.StateFullyBlendedOutFunction.SetFromFunction(StateFullyBlendedOutFunction.ResolveMember<UFunction>(GetBlueprintClassFromNode()));
+		bAnyAnimNodeFunctionChanged = true;
+	}
+
+	if (bAnyAnimNodeFunctionChanged)
+	{
+		GetAnimBlueprint()->RequestRefreshExtensions();
+		GetSchema()->ReconstructNode(*this);
+	}
+	
+	Super::PostEditChangeProperty(PropertyChangedEvent);
+}
+
+void UAnimGraphNode_StateResult::AddSearchMetaDataInfo(TArray<struct FSearchTagDataPair>& OutTaggedMetaData) const
+{
+	Super::AddSearchMetaDataInfo(OutTaggedMetaData);
+
+	const auto ConditionallyTagNodeFuncRef = [&OutTaggedMetaData](const FMemberReference& FuncMember, const FText& LocText)
+	{
+		if (IsPotentiallyBoundFunction(FuncMember))
+		{
+			const FText FunctionName = FText::FromName(FuncMember.GetMemberName());
+			OutTaggedMetaData.Add(FSearchTagDataPair(LocText, FunctionName));
+		}
+	};
+
+	// Conditionally include anim node function references as part of the node's search metadata
+	ConditionallyTagNodeFuncRef(StateEntryFunction, LOCTEXT("StateEntryFunctionName", "State Entry"));
+	ConditionallyTagNodeFuncRef(StateFullyBlendedInFunction, LOCTEXT("StateFullyBlendedInFunctionName", "State Fully Blended In"));
+	ConditionallyTagNodeFuncRef(StateExitFunction, LOCTEXT("StateExitFunctionName", "State Exit"));
+	ConditionallyTagNodeFuncRef(StateFullyBlendedOutFunction, LOCTEXT("StateFullyBlendedOutFunctionName", "State Fully Blended Out"));
+}
+
+void UAnimGraphNode_StateResult::GetBoundFunctionsInfo(TArray<TPair<FName, FName>>& InOutBindingsInfo)
+{
+	Super::GetBoundFunctionsInfo(InOutBindingsInfo);
+
+	const FName CategoryName = TEXT("Functions|State");
+	
+	if (StateEntryFunction.ResolveMember<UFunction>(GetBlueprintClassFromNode()) != nullptr)
+	{
+		InOutBindingsInfo.Emplace(CategoryName, GET_MEMBER_NAME_CHECKED(UAnimGraphNode_StateResult, StateEntryFunction));
+	}
+
+	if (StateFullyBlendedInFunction.ResolveMember<UFunction>(GetBlueprintClassFromNode()) != nullptr)
+	{
+		InOutBindingsInfo.Emplace(CategoryName, GET_MEMBER_NAME_CHECKED(UAnimGraphNode_StateResult, StateFullyBlendedInFunction));
+	}
+
+	if (StateExitFunction.ResolveMember<UFunction>(GetBlueprintClassFromNode()) != nullptr)
+	{
+		InOutBindingsInfo.Emplace(CategoryName, GET_MEMBER_NAME_CHECKED(UAnimGraphNode_StateResult, StateExitFunction));
+	}
+
+	if (StateFullyBlendedOutFunction.ResolveMember<UFunction>(GetBlueprintClassFromNode()) != nullptr)
+	{
+		InOutBindingsInfo.Emplace(CategoryName, GET_MEMBER_NAME_CHECKED(UAnimGraphNode_StateResult, StateFullyBlendedOutFunction));
+	}
+}
+
+bool UAnimGraphNode_StateResult::ReferencesFunction(const FName& InFunctionName, const UStruct* InScope) const
+{
+	return Super::ReferencesFunction(InFunctionName, InScope)
+	|| Node.StateEntryFunction.GetFunctionName() == InFunctionName
+	|| Node.StateFullyBlendedInFunction.GetFunctionName() == InFunctionName
+	|| Node.StateExitFunction.GetFunctionName() == InFunctionName
+	|| Node.StateFullyBlendedOutFunction.GetFunctionName() == InFunctionName;
 }
 
 #undef LOCTEXT_NAMESPACE
diff --git a/Engine/Source/Editor/AnimGraph/Public/AnimGraphNode_Base.h b/Engine/Source/Editor/AnimGraph/Public/AnimGraphNode_Base.h
--- a/Engine/Source/Editor/AnimGraph/Public/AnimGraphNode_Base.h
+++ b/Engine/Source/Editor/AnimGraph/Public/AnimGraphNode_Base.h
@@ -210,15 +210,15 @@ class ANIMGRAPH_API UAnimGraphNode_Base : public UK2Node
 	EBlueprintUsage BlueprintUsage;
 
 	// Function called before the node is updated for the first time
-	UPROPERTY(EditAnywhere, Category = Functions, meta=(FunctionReference, AllowFunctionLibraries, PrototypeFunction="/Script/AnimGraphRuntime.AnimExecutionContextLibrary.Prototype_ThreadSafeAnimUpdateCall"), DisplayName="On Initial Update")
+	UPROPERTY(EditAnywhere, Category = Functions, meta=(FunctionReference, AllowFunctionLibraries, PrototypeFunction="/Script/AnimGraphRuntime.AnimExecutionContextLibrary.Prototype_ThreadSafeAnimUpdateCall", DefaultBindingName="OnInitialUpdate"), DisplayName="On Initial Update")
 	FMemberReference InitialUpdateFunction;
 
-	// Function called when the node becomes relevant
-	UPROPERTY(EditAnywhere, Category = Functions, meta=(FunctionReference, AllowFunctionLibraries, PrototypeFunction="/Script/AnimGraphRuntime.AnimExecutionContextLibrary.Prototype_ThreadSafeAnimUpdateCall"), DisplayName="On Become Relevant")
+	// Function called when the node becomes relevant, meaning it goes from having no weight to any weight.
+	UPROPERTY(EditAnywhere, Category = Functions, meta=(FunctionReference, AllowFunctionLibraries, PrototypeFunction="/Script/AnimGraphRuntime.AnimExecutionContextLibrary.Prototype_ThreadSafeAnimUpdateCall", DefaultBindingName="OnBecomeRelevant"), DisplayName="On Become Relevant")
 	FMemberReference BecomeRelevantFunction;
 
 	// Function called when the node is updated
-	UPROPERTY(EditAnywhere, Category = Functions, meta=(FunctionReference, AllowFunctionLibraries, PrototypeFunction="/Script/AnimGraphRuntime.AnimExecutionContextLibrary.Prototype_ThreadSafeAnimUpdateCall"), DisplayName="On Update")
+	UPROPERTY(EditAnywhere, Category = Functions, meta=(FunctionReference, AllowFunctionLibraries, PrototypeFunction="/Script/AnimGraphRuntime.AnimExecutionContextLibrary.Prototype_ThreadSafeAnimUpdateCall", DefaultBindingName="OnUpdate"), DisplayName="On Update")
 	FMemberReference UpdateFunction;
 
 private:
@@ -534,6 +534,10 @@ public:
 	// or in a paused PIE world
 	void PostEditRefreshDebuggedComponent();
 
+	// Gets editor information for all the bound anim node functions (category metadata string, member variable's name)
+	// Used by SAnimGraphNode to display all bound functions of an anim node.
+	virtual void GetBoundFunctionsInfo(TArray<TPair<FName, FName>> & InOutBindingsInfo);
+	
 	// Check if a specified function reference appears to be valid by inspecting only the validity of the name and guid
 	static bool IsPotentiallyBoundFunction(const FMemberReference& FunctionReference);
 
@@ -624,6 +628,9 @@ protected:
 	/** @return the current object being debugged from the blueprint for this node. Can be nullptr. */
 	UObject* GetObjectBeingDebugged() const;
 
+	/** Helper function used to validate anim node function references */
+	void ValidateFunctionRef(FName InPropertyName, const FMemberReference& InRef, const FText& InFunctionName, FCompilerResultsLog& MessageLog);
+	
 protected:
 	// Old shown pins. Needs to be a member variable to track pin visibility changes between Pre and PostEditChange 
 	TArray<FName> OldShownPins;
diff --git a/Engine/Source/Editor/AnimGraph/Public/AnimGraphNode_StateResult.h b/Engine/Source/Editor/AnimGraph/Public/AnimGraphNode_StateResult.h
--- a/Engine/Source/Editor/AnimGraph/Public/AnimGraphNode_StateResult.h
+++ b/Engine/Source/Editor/AnimGraph/Public/AnimGraphNode_StateResult.h
@@ -6,7 +6,7 @@
 #include "CoreMinimal.h"
 #include "UObject/ObjectMacros.h"
 #include "AnimGraphNode_Base.h"
-#include "AnimNodes/AnimNode_StateResult.h"
+#include "Animation/AnimNode_StateResult.h"
 #include "AnimGraphNode_StateResult.generated.h"
 
 UCLASS(MinimalAPI)
@@ -17,6 +17,9 @@ class UAnimGraphNode_StateResult : public UAnimGraphNode_Base
 	UPROPERTY(EditAnywhere, Category=Settings)
 	FAnimNode_StateResult Node;
 
+	// UObject interface
+	virtual void PostEditChangeProperty(FPropertyChangedEvent& PropertyChangedEvent) override;
+
 	//~ Begin UEdGraphNode Interface.
 	virtual FLinearColor GetNodeTitleColor() const override;
 	virtual FText GetTooltipText() const override;
@@ -30,9 +33,46 @@ class UAnimGraphNode_StateResult : public UAnimGraphNode_Base
 	//~ Begin UAnimGraphNode_Base Interface
 	virtual bool IsSinkNode() const override;
 	virtual void OnProcessDuringCompilation(IAnimBlueprintCompilationContext& InCompilationContext, IAnimBlueprintGeneratedClassCompiledData& OutCompiledData) override;
+	virtual void ValidateAnimNodeDuringCompilation(USkeleton* ForSkeleton, FCompilerResultsLog& MessageLog) override;
+	virtual void AddSearchMetaDataInfo(TArray<struct FSearchTagDataPair>& OutTaggedMetaData) const override;
+	virtual void GetBoundFunctionsInfo(TArray<TPair<FName, FName>>& InOutBindingsInfo) override;
+	
+	// UK2Node interface
+	virtual bool ReferencesFunction(const FName& InFunctionName, const UStruct* InScope) const override;
 
 	// Get the link to the documentation
 	virtual FString GetDocumentationLink() const override;
 
 	//~ End UAnimGraphNode_Base Interface
+
+	/** Function called when the owning state is entered, meaning it becomes the state machine's current state. */
+	UPROPERTY(EditAnywhere, Category = "Functions|State", meta=(FunctionReference, AllowFunctionLibraries, PrototypeFunction="/Script/AnimGraphRuntime.AnimExecutionContextLibrary.Prototype_ThreadSafeAnimUpdateCall", DefaultBindingName="OnStateEntry"), DisplayName="On State Entry")
+	FMemberReference StateEntryFunction;
+
+	/** Function called when the owning state is fully blended in.
+	 *
+	 * Notes:
+	 * - This is only called for the state machine's current state since its the most recent transition's target state.
+	 * - This will not be called if the state is skipped. This can happen when the flag bSkipFirstUpdateTransition on the state machine node is set to true.
+	 */
+	UPROPERTY(EditAnywhere, Category = "Functions|State", meta=(FunctionReference, AllowFunctionLibraries, PrototypeFunction="/Script/AnimGraphRuntime.AnimExecutionContextLibrary.Prototype_ThreadSafeAnimUpdateCall", DefaultBindingName="OnStateFullyBlendedIn"), DisplayName="On State Fully Blended In")
+	FMemberReference StateFullyBlendedInFunction;
+
+	/** Function called when the owning state is exited, meaning it stops being the state machine's current state.
+	 *
+	 * Notes:
+	 * - This will not be called if the state machine node loses relevancy. Please use "On State Interrupt" for that case.
+	 */
+	UPROPERTY(EditAnywhere, Category = "Functions|State", meta=(FunctionReference, AllowFunctionLibraries, PrototypeFunction="/Script/AnimGraphRuntime.AnimExecutionContextLibrary.Prototype_ThreadSafeAnimUpdateCall", DefaultBindingName="OnStateExit"), DisplayName="On State Exit")
+	FMemberReference StateExitFunction;
+
+	/** Function called when the owning state is fully blended out.
+	 *
+	 * Notes:
+	 * - This will be called for any states that had weight.
+	 * - This will not be called if the state is skipped. This can happen when the flag bSkipFirstUpdateTransition on the state machine node is set to true.
+	 * - This will not be called if the state machine node loses relevancy. Please use "On State Interrupt" for that case.
+	 */
+	UPROPERTY(EditAnywhere, Category = "Functions|State", meta=(FunctionReference, AllowFunctionLibraries, PrototypeFunction="/Script/AnimGraphRuntime.AnimExecutionContextLibrary.Prototype_ThreadSafeAnimUpdateCall", DefaultBindingName="OnStateFullyBlendedOut"), DisplayName="On State Fully Blended Out")
+	FMemberReference StateFullyBlendedOutFunction;
 };
diff --git a/Engine/Source/Editor/AnimationBlueprintEditor/Private/AnimationNodes/SAnimationGraphNode.cpp b/Engine/Source/Editor/AnimationBlueprintEditor/Private/AnimationNodes/SAnimationGraphNode.cpp
--- a/Engine/Source/Editor/AnimationBlueprintEditor/Private/AnimationNodes/SAnimationGraphNode.cpp
+++ b/Engine/Source/Editor/AnimationBlueprintEditor/Private/AnimationNodes/SAnimationGraphNode.cpp
@@ -272,6 +272,37 @@ public:
 	
 	SLATE_END_ARGS()
 
+	void FindFunctionBindingPropertyNode(UAnimGraphNode_Base * InNode, TSharedPtr<IPropertyHandle>& OutPropertyHandle, TSharedPtr<IDetailTreeNode>& OutDetailTreeNode, const TSharedRef<IDetailTreeNode>& InRootNode, FName InCategory, FName InMemberName)
+	{
+		// We already found the property we are looking for
+		if (OutPropertyHandle.IsValid() || OutDetailTreeNode.IsValid())
+		{
+			return;
+		}
+		
+		TArray<TSharedRef<IDetailTreeNode>> Children;
+		InRootNode->GetChildren(Children);
+
+		// Parse children
+		for (const TSharedRef<IDetailTreeNode>& ChildNode : Children)
+		{
+			const TSharedPtr<IPropertyHandle> ChildPropertyHandle = ChildNode->CreatePropertyHandle();
+				
+			// Try to match node
+			if (ChildPropertyHandle.IsValid() && ChildPropertyHandle->GetProperty()->GetFName() == InMemberName && ChildPropertyHandle->GetMetaData(TEXT("Category")) == InCategory)
+			{
+				OutDetailTreeNode = ChildNode;
+				OutPropertyHandle = ChildPropertyHandle;
+				
+				// We found our function binding property.
+				return;
+			}
+
+			// Look into sub categories
+			FindFunctionBindingPropertyNode(InNode, OutPropertyHandle, OutDetailTreeNode, ChildNode, InCategory, InMemberName);
+		}
+	}
+	
 	void Construct(const FArguments& InArgs, UAnimGraphNode_Base* InNode)
 	{
 		UseLowDetail = InArgs._UseLowDetail;
@@ -291,46 +322,45 @@ public:
 		GridPanel->SetVisibility(EVisibility::Collapsed);
 
 		int32 RowIndex = 0;
-		
-		// Add bound functions
-		auto AddFunctionBindingWidget = [this, InNode, &GridPanel, &RowIndex](FName InCategory, FName InMemberName)
+
+		// Get anim node functions that are bound
+		TArray<TPair<FName, FName>> FunctionBindingsInfo;
+		InNode->GetBoundFunctionsInfo(FunctionBindingsInfo);
+
+		// Build anim node widgets
+		for (const TPair<FName, FName> & FunctionBindingInfo : FunctionBindingsInfo)
 		{
-			GridPanel->SetVisibility(EVisibility::Visible);
-			
-			// Find row
+			FName Category = FunctionBindingInfo.Key;
+			FName MemberName = FunctionBindingInfo.Value;
+
 			TSharedPtr<IPropertyHandle> PropertyHandle;
 			TSharedPtr<IDetailTreeNode> DetailTreeNode;
+			
+			GridPanel->SetVisibility(EVisibility::Visible);
 
+			// Find row for function binding property
 			for (const TSharedRef<IDetailTreeNode>& RootTreeNode : PropertyRowGenerator->GetRootTreeNodes())
 			{
-				if(RootTreeNode->GetNodeName() == InCategory)
+				if (Category.ToString().Contains(RootTreeNode->GetNodeName().ToString()))
 				{
-					TArray<TSharedRef<IDetailTreeNode>> Children;
-					RootTreeNode->GetChildren(Children);
-
-					for (int32 ChildIdx = 0; ChildIdx < Children.Num(); ChildIdx++)
-					{
-						TSharedPtr<IPropertyHandle> ChildPropertyHandle = Children[ChildIdx]->CreatePropertyHandle();
-						if (ChildPropertyHandle.IsValid() && ChildPropertyHandle->GetProperty() && ChildPropertyHandle->GetProperty()->GetFName() == InMemberName)
-						{
-							DetailTreeNode = Children[ChildIdx];
-							PropertyHandle = ChildPropertyHandle;
-							PropertyHandle->SetOnPropertyValueChanged(FSimpleDelegate::CreateLambda([InNode]()
-							{
-								InNode->ReconstructNode();
-							}));
-							break;
-						}
-					}
+					FindFunctionBindingPropertyNode(InNode, PropertyHandle, DetailTreeNode, RootTreeNode, Category, MemberName);
 				}
 			}
-
-			if(DetailTreeNode.IsValid() && PropertyHandle.IsValid())
+			
+			// Build function binding widget
+			if (DetailTreeNode.IsValid() && PropertyHandle.IsValid())
 			{
-				DetailNodes.Add(DetailTreeNode);
+				// Ensure anim node is rebuilt if any function binding changes
+				PropertyHandle->SetOnPropertyValueChanged(FSimpleDelegate::CreateLambda([InNode]()
+				{
+					InNode->ReconstructNode();
+				}));
 				
-				FNodeWidgets NodeWidgets = DetailTreeNode->CreateNodeWidgets();
+				DetailNodes.Add(DetailTreeNode);
+
+				const FNodeWidgets NodeWidgets = DetailTreeNode->CreateNodeWidgets();
 
+				// Binding variable name
 				GridPanel->AddSlot(0, RowIndex)
 				.HAlign(HAlign_Left)
 				.VAlign(VAlign_Center)
@@ -349,6 +379,7 @@ public:
 					]
 				];
 
+				// Function name
 				GridPanel->AddSlot(1, RowIndex)
 				.HAlign(HAlign_Left)
 				.VAlign(VAlign_Center)
@@ -369,22 +400,9 @@ public:
 
 				RowIndex++;
 			}
-		};
-
-		if(InNode->InitialUpdateFunction.ResolveMember<UFunction>(InNode->GetBlueprintClassFromNode()) != nullptr)
-		{
-			AddFunctionBindingWidget("Functions", GET_MEMBER_NAME_CHECKED(UAnimGraphNode_Base, InitialUpdateFunction));
-		}
-		if(InNode->BecomeRelevantFunction.ResolveMember<UFunction>(InNode->GetBlueprintClassFromNode()) != nullptr)
-		{
-			AddFunctionBindingWidget("Functions", GET_MEMBER_NAME_CHECKED(UAnimGraphNode_Base, BecomeRelevantFunction));
-		}
-		if(InNode->UpdateFunction.ResolveMember<UFunction>(InNode->GetBlueprintClassFromNode()) != nullptr)
-		{
-			AddFunctionBindingWidget("Functions", GET_MEMBER_NAME_CHECKED(UAnimGraphNode_Base, UpdateFunction));
 		}
 		
-		if(DetailNodes.Num() == 0)
+		if (DetailNodes.Num() == 0)
 		{
 			// If we didnt add a function binding, remove the row generator as we dont need it and its expensive (as it ticks)
 			PropertyRowGenerator.Reset();
diff --git a/Engine/Source/Editor/DetailCustomizations/Private/AnimStateNodeDetails.cpp b/Engine/Source/Editor/DetailCustomizations/Private/AnimStateNodeDetails.cpp
--- a/Engine/Source/Editor/DetailCustomizations/Private/AnimStateNodeDetails.cpp
+++ b/Engine/Source/Editor/DetailCustomizations/Private/AnimStateNodeDetails.cpp
@@ -11,6 +11,9 @@
 #include "Misc/Attribute.h"
 #include "Widgets/DeclarativeSyntaxSupport.h"
 #include "Widgets/Text/STextBlock.h"
+#include "Widgets/Images/SImage.h"
+#include "Widgets/Layout/SBox.h"
+#include "Widgets/SToolTip.h"
 
 class IDetailCustomization;
 
@@ -27,43 +30,46 @@ TSharedRef<IDetailCustomization> FAnimStateNodeDetails::MakeInstance()
 void FAnimStateNodeDetails::CustomizeDetails( IDetailLayoutBuilder& DetailBuilder )
 {
 	IDetailCategoryBuilder& SegmentCategory = DetailBuilder.EditCategory("Animation State", LOCTEXT("AnimationStateCategoryTitle", "Animation State") );
+	
+	GenerateAnimationStateEventRow(SegmentCategory, LOCTEXT("EnteredAnimationStateEventLabel", "Entered State Event"), TEXT("StateEntered"));
+	GenerateAnimationStateEventRow(SegmentCategory, LOCTEXT("ExitedAnimationStateEventLabel", "Left State Event"), TEXT("StateLeft"));
+	GenerateAnimationStateEventRow(SegmentCategory, LOCTEXT("FullyBlendedAnimationStateEventLabel", "Fully Blended State Event"), TEXT("StateFullyBlended"));
 
-	{
-		SegmentCategory.AddCustomRow( LOCTEXT("EnteredAnimationStateEventLabel", "Entered State Event") )
-		[
-			SNew( STextBlock )
-			.Text( LOCTEXT("EnteredAnimationStateEventLabel", "Entered State Event") )
-			.Font( IDetailLayoutBuilder::GetDetailFontBold() )
-		];
-		CreateTransitionEventPropertyWidgets(SegmentCategory, TEXT("StateEntered"));
-	}
-
-
-	{
-		SegmentCategory.AddCustomRow( LOCTEXT("ExitedAnimationStateEventLabel", "Left State Event") )
-		[
-			SNew( STextBlock )
-			.Text( LOCTEXT("ExitedAnimationStateEventLabel", "Left State Event")  )
-			.Font( IDetailLayoutBuilder::GetDetailFontBold() )
-		];
-		CreateTransitionEventPropertyWidgets(SegmentCategory, TEXT("StateLeft"));
-	}
+	DetailBuilder.HideProperty("StateEntered");
+	DetailBuilder.HideProperty("StateLeft");
+	DetailBuilder.HideProperty("StateFullyBlended");
+}
 
-	{
+void FAnimStateNodeDetails::GenerateAnimationStateEventRow(IDetailCategoryBuilder& SegmentCategory,const FText& StateEventLabel, const FString& TransitionName)
+{
+	const FSlateBrush* WarningIcon = FAppStyle::GetBrush("Icons.WarningWithColor");
 
-		SegmentCategory.AddCustomRow( LOCTEXT("FullyBlendedAnimationStateEventLabel", "Fully Blended State Event") )
+	SegmentCategory.AddCustomRow( StateEventLabel)
 		[
-			SNew( STextBlock )
-			.Text( LOCTEXT("FullyBlendedAnimationStateEventLabel", "Fully Blended State Event") )
-			.Font( IDetailLayoutBuilder::GetDetailFontBold() )
+			SNew(SHorizontalBox)
+			.ToolTipText(LOCTEXT("AnimStateEventDeprecatedTooltip", "This is no longer recommended. The recommended approach is to use the anim node function versions."))
+			.Visibility(EVisibility::Visible)
+			
+			+SHorizontalBox::Slot()
+			.VAlign(VAlign_Center)
+			.AutoWidth()
+			[
+				SNew( STextBlock )
+				.Text( StateEventLabel )
+				.Font( IDetailLayoutBuilder::GetDetailFontBold() )
+			]
+
+			+SHorizontalBox::Slot()
+			.VAlign(VAlign_Center)
+			.AutoWidth()
+			.HAlign(HAlign_Right)
+			.Padding(4.0f, 0.0f, 0.0f, 0.0f)
+			[	SNew(SImage)
+				.Image(WarningIcon)
+			]
 		];
-		CreateTransitionEventPropertyWidgets(SegmentCategory, TEXT("StateFullyBlended"));
-	}
-
-
-	DetailBuilder.HideProperty("StateEntered");
-	DetailBuilder.HideProperty("StateLeft");
-	DetailBuilder.HideProperty("StateFullyBlended");
+		
+	CreateTransitionEventPropertyWidgets(SegmentCategory, TransitionName);
 }
 
 #undef LOCTEXT_NAMESPACE
diff --git a/Engine/Source/Editor/DetailCustomizations/Private/AnimStateNodeDetails.h b/Engine/Source/Editor/DetailCustomizations/Private/AnimStateNodeDetails.h
--- a/Engine/Source/Editor/DetailCustomizations/Private/AnimStateNodeDetails.h
+++ b/Engine/Source/Editor/DetailCustomizations/Private/AnimStateNodeDetails.h
@@ -16,5 +16,9 @@ public:
 
 	/** IDetailCustomization interface */
 	virtual void CustomizeDetails( IDetailLayoutBuilder& DetailBuilder ) override;
+
+private:
+	
+	void GenerateAnimationStateEventRow(IDetailCategoryBuilder& SegmentCategory, const FText & StateEventLabel, const FString & TransitionName);
 };
 
diff --git a/Engine/Source/Editor/Kismet/Private/BlueprintMemberReferenceCustomization.cpp b/Engine/Source/Editor/Kismet/Private/BlueprintMemberReferenceCustomization.cpp
--- a/Engine/Source/Editor/Kismet/Private/BlueprintMemberReferenceCustomization.cpp
+++ b/Engine/Source/Editor/Kismet/Private/BlueprintMemberReferenceCustomization.cpp
@@ -87,9 +87,14 @@ void FBlueprintMemberReferenceDetails::CustomizeHeader(TSharedRef<IPropertyHandl
 		const FString& PrototypeFunctionName = InStructPropertyHandle->GetMetaData("PrototypeFunction");
 		UFunction* PrototypeFunction = PrototypeFunctionName.IsEmpty() ? nullptr : FindObject<UFunction>(nullptr, *PrototypeFunctionName);
 
-		auto OnGenerateBindingName = []() -> FString
+		FString DefaultBindingName = InStructPropertyHandle->HasMetaData("DefaultBindingName") ? InStructPropertyHandle->GetMetaData("DefaultBindingName") : TEXT("NewFunction");
+
+		// Binding widget re-adds the 'On' prefix because of legacy support for UMG events, so we remove it here
+		DefaultBindingName.RemoveFromStart(TEXT("On"));
+
+		auto OnGenerateBindingName = [DefaultBindingName]() -> FString
 		{
-			return TEXT("NewFunction");
+			return DefaultBindingName;
 		};
 		
 		auto OnCanBindProperty = [](FProperty* InProperty)
diff --git a/Engine/Source/Runtime/AnimGraphRuntime/Private/AnimationStateMachineLibrary.cpp b/Engine/Source/Runtime/AnimGraphRuntime/Private/AnimationStateMachineLibrary.cpp
--- a/Engine/Source/Runtime/AnimGraphRuntime/Private/AnimationStateMachineLibrary.cpp
+++ b/Engine/Source/Runtime/AnimGraphRuntime/Private/AnimationStateMachineLibrary.cpp
@@ -3,7 +3,7 @@
 #include "AnimationStateMachineLibrary.h"
 
 #include "Animation/AnimNode_StateMachine.h"
-#include "AnimNodes/AnimNode_StateResult.h"
+#include "Animation/AnimNode_StateResult.h"
 #include "Animation/AnimInstanceProxy.h"
 
 #include UE_INLINE_GENERATED_CPP_BY_NAME(AnimationStateMachineLibrary)
diff --git a/Engine/Source/Runtime/AnimGraphRuntime/Public/AnimNodes/AnimNode_StateResult.cpp b/Engine/Source/Runtime/AnimGraphRuntime/Public/AnimNodes/AnimNode_StateResult.cpp
deleted file mode 100644
diff --git a/Engine/Source/Runtime/AnimGraphRuntime/Public/AnimNodes/AnimNode_StateResult.h b/Engine/Source/Runtime/AnimGraphRuntime/Public/AnimNodes/AnimNode_StateResult.h
deleted file mode 100644
diff --git a/Engine/Source/Runtime/Engine/Classes/Animation/AnimExecutionContext.h b/Engine/Source/Runtime/Engine/Classes/Animation/AnimExecutionContext.h
--- a/Engine/Source/Runtime/Engine/Classes/Animation/AnimExecutionContext.h
+++ b/Engine/Source/Runtime/Engine/Classes/Animation/AnimExecutionContext.h
@@ -30,6 +30,8 @@ public:
 	struct FData
 	{
 	public:
+		FData(const FAnimationBaseContext& InContext);
+		
 		FData(const FAnimationInitializeContext& InContext);
 
 		FData(const FAnimationUpdateContext& InContext);
@@ -48,6 +50,7 @@ public:
 		enum class EContextType
 		{
 			None,
+			Base,
 			Initialize,
 			Update,
 			Pose,
diff --git a/Engine/Source/Runtime/Engine/Classes/Animation/AnimNodeFunctionRef.h b/Engine/Source/Runtime/Engine/Classes/Animation/AnimNodeFunctionRef.h
--- a/Engine/Source/Runtime/Engine/Classes/Animation/AnimNodeFunctionRef.h
+++ b/Engine/Source/Runtime/Engine/Classes/Animation/AnimNodeFunctionRef.h
@@ -5,10 +5,12 @@
 #include "CoreMinimal.h"
 #include "AnimNodeFunctionRef.generated.h"
 
+struct FAnimNode_StateMachine;
 struct FPoseLink;
 struct FPoseLinkBase;
 struct FComponentSpacePoseLink;
 struct FAnimNode_Base;
+struct FAnimationBaseContext;
 struct FAnimationInitializeContext;
 struct FAnimationUpdateContext;
 struct FPoseContext;
@@ -88,6 +90,7 @@ private:
 	friend struct ::FPoseLink;
 	friend struct ::FComponentSpacePoseLink;
 	friend struct ::FAnimInstanceProxy;
+	friend struct ::FAnimNode_StateMachine;
 	
 	// Call the InitialUpdate function of this node
 	static void InitialUpdate(const FAnimationUpdateContext& InContext, FAnimNode_Base& InNode);
@@ -97,6 +100,9 @@ private:
 
 	// Call the Update function of this node
 	static void Update(const FAnimationUpdateContext& InContext, FAnimNode_Base& InNode);
+	
+	// Call a generic function for this node
+	static void CallFunction(const FAnimNodeFunctionRef& InFunction, const FAnimationBaseContext& InContext, FAnimNode_Base& InNode);
 };
 
 }}
diff --git a/Engine/Source/Runtime/Engine/Private/Animation/AnimExecutionContext.cpp b/Engine/Source/Runtime/Engine/Private/Animation/AnimExecutionContext.cpp
--- a/Engine/Source/Runtime/Engine/Private/Animation/AnimExecutionContext.cpp
+++ b/Engine/Source/Runtime/Engine/Private/Animation/AnimExecutionContext.cpp
@@ -5,6 +5,12 @@
 
 #include UE_INLINE_GENERATED_CPP_BY_NAME(AnimExecutionContext)
 
+FAnimExecutionContext::FData::FData(const FAnimationBaseContext& InContext)
+{
+	Context = const_cast<FAnimationBaseContext*>(&InContext);
+	ContextType = EContextType::Base;
+}
+
 FAnimExecutionContext::FData::FData(const FAnimationInitializeContext& InContext)
 {
 	Context = const_cast<FAnimationInitializeContext*>(&InContext);
diff --git a/Engine/Source/Runtime/Engine/Private/Animation/AnimNodeFunctionRef.cpp b/Engine/Source/Runtime/Engine/Private/Animation/AnimNodeFunctionRef.cpp
--- a/Engine/Source/Runtime/Engine/Private/Animation/AnimNodeFunctionRef.cpp
+++ b/Engine/Source/Runtime/Engine/Private/Animation/AnimNodeFunctionRef.cpp
@@ -113,7 +113,7 @@ void FNodeFunctionCaller::InitialUpdate(const FAnimationUpdateContext& InContext
 		}
 	}
 }
-	
+
 void FNodeFunctionCaller::BecomeRelevant(const FAnimationUpdateContext& InContext, FAnimNode_Base& InNode)
 {
 	if(InNode.NodeData != nullptr && InNode.NodeData->HasNodeAnyFlags(EAnimNodeDataFlags::HasBecomeRelevantFunction))
@@ -130,7 +130,7 @@ void FNodeFunctionCaller::BecomeRelevant(const FAnimationUpdateContext& InContex
 		}
 	}
 }
-	
+
 void FNodeFunctionCaller::Update(const FAnimationUpdateContext& InContext, FAnimNode_Base& InNode)
 {
 	if(InNode.NodeData != nullptr && InNode.NodeData->HasNodeAnyFlags(EAnimNodeDataFlags::HasUpdateFunction))
@@ -139,4 +139,9 @@ void FNodeFunctionCaller::Update(const FAnimationUpdateContext& InContext, FAnim
 	}
 }
 
+void FNodeFunctionCaller::CallFunction(const FAnimNodeFunctionRef& InFunction, const FAnimationBaseContext& InContext, FAnimNode_Base& InNode)
+{
+	CallFunctionHelper<FAnimExecutionContext>(InFunction, InContext, InNode);
+}
+
 }}
diff --git a/Engine/Source/Runtime/Engine/Private/Animation/AnimNode_StateMachine.cpp b/Engine/Source/Runtime/Engine/Private/Animation/AnimNode_StateMachine.cpp
--- a/Engine/Source/Runtime/Engine/Private/Animation/AnimNode_StateMachine.cpp
+++ b/Engine/Source/Runtime/Engine/Private/Animation/AnimNode_StateMachine.cpp
@@ -14,6 +14,7 @@
 #include "Animation/ExposedValueHandler.h"
 #include "Logging/TokenizedMessage.h"
 #include "Animation/AnimInertializationSyncScope.h"
+#include "Animation/AnimNode_StateResult.h"
 
 #include UE_INLINE_GENERATED_CPP_BY_NAME(AnimNode_StateMachine)
 
@@ -479,6 +480,9 @@ void FAnimNode_StateMachine::Update_AnyThread(const FAnimationUpdateContext& Con
 	// Tick the individual state/states that are active
 	if (ActiveTransitionArray.Num() > 0)
 	{
+		// Keep track of states that have blended out to avoid recalling their anim node state functions.
+		TSet<int32> BlendedOutStates;
+		
 		for (int32 Index = 0; Index < ActiveTransitionArray.Num(); ++Index)
 		{
 			// The custom graph will tick the needed states
@@ -489,11 +493,57 @@ void FAnimNode_StateMachine::Update_AnyThread(const FAnimationUpdateContext& Con
 			
 			if (bFinishedTrans)
 			{
+				// Trigger "Fully Blended Out" of any states that have any weight to them.
+				{
+					const int32 PreviousStateIndex = ActiveTransitionArray[Index].PreviousState;
+					const bool bAttemptingToTransitionOutOfCurrentState = PreviousStateIndex == CurrentState;
+					
+					if (!bAttemptingToTransitionOutOfCurrentState && StatePoseLinks.IsValidIndex(PreviousStateIndex) && StatePoseLinks[PreviousStateIndex].LinkID != INDEX_NONE)
+					{
+						if (!BlendedOutStates.Contains(PreviousStateIndex) && FMath::IsNearlyZero(GetStateWeight(PreviousStateIndex)))
+						{
+							UE::Anim::FNodeFunctionCaller::CallFunction(static_cast<FAnimNode_StateResult*>(StatePoseLinks[PreviousStateIndex].GetLinkNode())->GetStateFullyBlendedOutFunction(), Context, *this);
+							BlendedOutStates.Add(PreviousStateIndex);
+						}
+					}
+				}
+				
 				// only play these events if it is the last transition (most recent, going to current state)
 				if (Index == (ActiveTransitionArray.Num() - 1))
 				{
 					Context.AnimInstanceProxy->AddAnimNotifyFromGeneratedClass(ActiveTransitionArray[Index].EndNotify);
 					Context.AnimInstanceProxy->AddAnimNotifyFromGeneratedClass(GetStateInfo().FullyBlendedNotify);
+
+					// Handle events/calls during "Fully Blended In".
+					{
+						const int32 NextStateIndex = ActiveTransitionArray[Index].NextState;
+						
+						if (StatePoseLinks.IsValidIndex(NextStateIndex))
+						{
+							// If our most recent state has fully blended in that means that all our previous states must also blend out and their transitions will be forced
+							// to complete, (removed from active transition array), therefore we need to call their fully blended out functions if they haven't been called already. 
+							if (ActiveTransitionArray.Num() > 1)
+							{
+								for (int32 BlendOutIndex = 0; BlendOutIndex < Index; ++BlendOutIndex)
+								{
+									const int32 PreviousStateIndex = ActiveTransitionArray[BlendOutIndex].PreviousState;
+									const bool bAttemptingToTransitionOutOfCurrentState = PreviousStateIndex == CurrentState;
+
+									if (!bAttemptingToTransitionOutOfCurrentState && StatePoseLinks.IsValidIndex(PreviousStateIndex) && StatePoseLinks[PreviousStateIndex].LinkID != INDEX_NONE && !BlendedOutStates.Contains(PreviousStateIndex))
+									{
+										UE::Anim::FNodeFunctionCaller::CallFunction(static_cast<FAnimNode_StateResult*>(StatePoseLinks[PreviousStateIndex].GetLinkNode())->GetStateFullyBlendedOutFunction(), Context, *this);
+										BlendedOutStates.Add(PreviousStateIndex);
+									}
+								}
+							}
+
+							// Now that all possible state fully blended out functions have been called, perform state fully blended in call.
+							if (StatePoseLinks.IsValidIndex(NextStateIndex) && StatePoseLinks[NextStateIndex].LinkID != INDEX_NONE)
+							{
+								UE::Anim::FNodeFunctionCaller::CallFunction(static_cast<FAnimNode_StateResult*>(StatePoseLinks[NextStateIndex].GetLinkNode())->GetStateFullyBlendedInFunction(), Context, *this);
+							}
+						}
+					}
 				}
 
 				// we were the last active transition and used inertialization
@@ -1034,6 +1084,11 @@ void FAnimNode_StateMachine::SetState(const FAnimationBaseContext& Context, int3
 			OnGraphStatesExited[CurrentState].ExecuteIfBound(*this, CurrentState, NewStateIndex);
 		}
 
+		if (StatePoseLinks.IsValidIndex(CurrentState) && StatePoseLinks[CurrentState].LinkID != INDEX_NONE)
+		{
+			UE::Anim::FNodeFunctionCaller::CallFunction(static_cast<FAnimNode_StateResult*>(StatePoseLinks[CurrentState].GetLinkNode())->GetStateExitFunction(), Context, *this);
+		}
+
 		bool bForceReset = false;
 
 		if(PRIVATE_MachineDescription->States.IsValidIndex(NewStateIndex))
@@ -1103,6 +1158,11 @@ void FAnimNode_StateMachine::SetState(const FAnimationBaseContext& Context, int3
 			}
 		}
 
+		if (StatePoseLinks.IsValidIndex(CurrentState) && StatePoseLinks[CurrentState].LinkID != INDEX_NONE)
+		{
+			UE::Anim::FNodeFunctionCaller::CallFunction(static_cast<FAnimNode_StateResult*>(StatePoseLinks[CurrentState].GetLinkNode())->GetStateEntryFunction(), Context, *this);
+		}
+
 		if(CurrentState != INDEX_NONE && CurrentState < OnGraphStatesEntered.Num())
 		{
 			OnGraphStatesEntered[CurrentState].ExecuteIfBound(*this, PrevStateIndex, CurrentState);
diff --git a/Engine/Source/Runtime/Engine/Private/Animation/AnimNode_StateResult.cpp b/Engine/Source/Runtime/Engine/Private/Animation/AnimNode_StateResult.cpp
new file mode 100644
--- /dev/null
+++ b/Engine/Source/Runtime/Engine/Private/Animation/AnimNode_StateResult.cpp
@@ -0,0 +1,44 @@
+// Copyright Epic Games, Inc. All Rights Reserved.
+
+#include "Animation/AnimNode_StateResult.h"
+
+#include UE_INLINE_GENERATED_CPP_BY_NAME(AnimNode_StateResult)
+
+bool FAnimNode_StateResult::SerializeFromMismatchedTag(const FPropertyTag& Tag, FStructuredArchive::FSlot Slot)
+{
+	if(Tag.Type == NAME_StructProperty && Tag.StructName == FAnimNode_Root::StaticStruct()->GetFName())
+	{
+		FAnimNode_Root OldValue;
+		FAnimNode_Root::StaticStruct()->SerializeItem(Slot, &OldValue, nullptr);
+		*static_cast<FAnimNode_Root*>(this) = OldValue;
+
+		return true;
+	}
+
+	return false;
+}
+
+int32 FAnimNode_StateResult::GetStateIndex() const
+{
+	return GET_ANIM_NODE_DATA(int32, StateIndex);
+}
+
+const FAnimNodeFunctionRef& FAnimNode_StateResult::GetStateEntryFunction() const
+{
+	return GET_ANIM_NODE_DATA(FAnimNodeFunctionRef, StateEntryFunction);
+}
+
+const FAnimNodeFunctionRef& FAnimNode_StateResult::GetStateFullyBlendedInFunction() const
+{
+	return GET_ANIM_NODE_DATA(FAnimNodeFunctionRef, StateFullyBlendedInFunction);
+}
+
+const FAnimNodeFunctionRef& FAnimNode_StateResult::GetStateExitFunction() const
+{
+	return GET_ANIM_NODE_DATA(FAnimNodeFunctionRef, StateExitFunction);
+}
+
+const FAnimNodeFunctionRef& FAnimNode_StateResult::GetStateFullyBlendedOutFunction() const
+{
+	return GET_ANIM_NODE_DATA(FAnimNodeFunctionRef, StateFullyBlendedOutFunction);
+}
diff --git a/Engine/Source/Runtime/Engine/Public/Animation/AnimNode_StateResult.h b/Engine/Source/Runtime/Engine/Public/Animation/AnimNode_StateResult.h
new file mode 100644
--- /dev/null
+++ b/Engine/Source/Runtime/Engine/Public/Animation/AnimNode_StateResult.h
@@ -0,0 +1,84 @@
+// Copyright Epic Games, Inc. All Rights Reserved.
+
+#pragma once
+
+#include "CoreMinimal.h"
+#include "UObject/ObjectMacros.h"
+#include "Animation/AnimNode_Root.h"
+#include "Animation/AnimNodeFunctionRef.h"
+#include "AnimNode_StateResult.generated.h"
+
+struct FAnimNode_StateMachine;
+class UAnimGraphNode_StateResult;
+
+// Root node of an state machine state (sink node).
+// We dont use AnimNode_Root to let us distinguish these nodes in the property list at link time.
+USTRUCT(BlueprintInternalUseOnly)
+struct ENGINE_API FAnimNode_StateResult : public FAnimNode_Root
+{
+	GENERATED_USTRUCT_BODY()
+
+	/** Used to upgrade old FAnimNode_Roots to FAnimNode_StateResult */
+	bool SerializeFromMismatchedTag(const FPropertyTag& Tag, FStructuredArchive::FSlot Slot);
+
+#if WITH_EDITORONLY_DATA
+protected:
+
+	/** The index of the state this node belongs to. Filled in during the owning state machine's compilation. */
+	UPROPERTY(meta = (FoldProperty))
+	int32 StateIndex = -1;
+#endif
+
+public:
+#if WITH_EDITORONLY_DATA
+	void SetStateIndex(int32 InStateIndex) { StateIndex = InStateIndex; }	
+#endif
+
+	int32 GetStateIndex() const;
+	
+	/** Gets the anim node function called on state entry, state became the current state of its state machine */
+	const FAnimNodeFunctionRef& GetStateEntryFunction() const;
+
+	/** Gets the anim node function called on state fully blended in */
+	const FAnimNodeFunctionRef& GetStateFullyBlendedInFunction() const;
+
+	/** Gets the anim node function called on state exit, state stopped being the current state of its state machine  */
+	const FAnimNodeFunctionRef& GetStateExitFunction() const;
+
+	/** Gets the anim node function called on state fully blended out */
+	const FAnimNodeFunctionRef& GetStateFullyBlendedOutFunction() const;
+
+private:
+	friend struct FAnimNode_StateMachine;
+	friend class UAnimGraphNode_StateResult;
+
+#if WITH_EDITORONLY_DATA
+	
+	/** The function called on state entry */
+	UPROPERTY(meta=(FoldProperty))
+	FAnimNodeFunctionRef StateEntryFunction;
+
+	/** The function called on state fully blended in */
+	UPROPERTY(meta=(FoldProperty))
+	FAnimNodeFunctionRef StateFullyBlendedInFunction;
+
+	/** The function called on state exit */
+	UPROPERTY(meta=(FoldProperty))
+	FAnimNodeFunctionRef StateExitFunction;
+
+	/** The function called on state fully blended out */
+	UPROPERTY(meta=(FoldProperty))
+	FAnimNodeFunctionRef StateFullyBlendedOutFunction;
+
+#endif
+};
+
+template<>
+struct TStructOpsTypeTraits<FAnimNode_StateResult>
+	: public TStructOpsTypeTraitsBase2<FAnimNode_StateResult>
+{
+	enum
+	{
+		WithStructuredSerializeFromMismatchedTag = true,
+	};
+};
