Non unity fix
#rnx

diff --git a/Engine/Source/Editor/AnimGraph/Private/AnimGraphModule.cpp b/Engine/Source/Editor/AnimGraph/Private/AnimGraphModule.cpp
--- a/Engine/Source/Editor/AnimGraph/Private/AnimGraphModule.cpp
+++ b/Engine/Source/Editor/AnimGraph/Private/AnimGraphModule.cpp
@@ -27,6 +27,7 @@
 #include "AnimBlueprintCompiler.h"
 #include "Animation/AnimAttributes.h"
 #include "AnimGraphAttributes.h"
+#include "Animation/AnimInstance.h"
 #include "Animation/AnimSync.h"
 #include "Animation/AnimNode_Inertialization.h"
 #include "Animation/AnimRootMotionProvider.h"
