Fixes for automatic transitions rules and added the feature to manually set their trigger time.

- Added support for explicitly specifying the transitions trigger time when using state machine automatic transition rules.
- Fixed state machine automatic transitions rules not triggering from a looping animation.
- Fixed state machine automatic transitions rules that have their next state as a conduit state triggering before reaching the end of the animation.
- Consolidate all the ways to check if asset in asset player is looping into a single api call, IsLooping().
- Added getter for a AssetPlayer's DeltaTimeRecord.
- Added warning when a asset player is using looping animations with automatic rule based transitions.

diff --git a/Engine/Plugins/Animation/AnimationLocomotionLibrary/Source/Runtime/Private/AnimDistanceMatchingLibrary.cpp b/Engine/Plugins/Animation/AnimationLocomotionLibrary/Source/Runtime/Private/AnimDistanceMatchingLibrary.cpp
--- a/Engine/Plugins/Animation/AnimationLocomotionLibrary/Source/Runtime/Private/AnimDistanceMatchingLibrary.cpp
+++ b/Engine/Plugins/Animation/AnimationLocomotionLibrary/Source/Runtime/Private/AnimDistanceMatchingLibrary.cpp
@@ -164,7 +164,7 @@ FSequenceEvaluatorReference UAnimDistanceMatchingLibrary::AdvanceTimeByDistanceM
 					{
 						const float CurrentTime = InSequenceEvaluator.GetExplicitTime();
 						const float CurrentAssetLength = InSequenceEvaluator.GetCurrentAssetLength();
-						const bool bAllowLooping = InSequenceEvaluator.GetShouldLoop();
+						const bool bAllowLooping = InSequenceEvaluator.IsLooping();
 
 						float TimeAfterDistanceTraveled = UE::Anim::DistanceMatchingUtility::GetTimeAfterDistanceTraveled(AnimSequence, CurrentTime, DistanceTraveled, DistanceCurveName, bAllowLooping);
 
diff --git a/Engine/Source/Editor/AnimGraph/Private/AnimBlueprintPostCompileValidation.cpp b/Engine/Source/Editor/AnimGraph/Private/AnimBlueprintPostCompileValidation.cpp
--- a/Engine/Source/Editor/AnimGraph/Private/AnimBlueprintPostCompileValidation.cpp
+++ b/Engine/Source/Editor/AnimGraph/Private/AnimBlueprintPostCompileValidation.cpp
@@ -146,7 +146,7 @@ void UAnimBlueprintPostCompileValidation::PCV_GatherAnimSequencesFromGraph(TArra
 			if (const FAnimNode_BlendSpacePlayer* const BlendSpacePlayer = Property->ContainerPtrToValuePtr<FAnimNode_BlendSpacePlayer>(PCV_Params.DefaultAnimInstance))
 			{
 				const bool bPassSyncGroupFilter = !GatherParams.bFilterBySyncGroup || (BlendSpacePlayer->GetGroupName() == GatherParams.SyncGroupName);
-				const bool bPassLoopingFilter = !GatherParams.bFilterByLoopingCondition || (BlendSpacePlayer->GetLoop() == GatherParams.bLoopingCondition);
+				const bool bPassLoopingFilter = !GatherParams.bFilterByLoopingCondition || (BlendSpacePlayer->IsLooping() == GatherParams.bLoopingCondition);
 				if (bPassSyncGroupFilter && bPassLoopingFilter)
 				{
 					PCV_GatherAnimSequences(OutAnimSequences, BlendSpacePlayer->GetBlendSpace());
@@ -158,7 +158,7 @@ void UAnimBlueprintPostCompileValidation::PCV_GatherAnimSequencesFromGraph(TArra
 			if (const FAnimNode_SequencePlayer* const SequencePlayer = Property->ContainerPtrToValuePtr<FAnimNode_SequencePlayer>(PCV_Params.DefaultAnimInstance))
 			{
 				const bool bPassSyncGroupFilter = !GatherParams.bFilterBySyncGroup || (SequencePlayer->GetGroupName() == GatherParams.SyncGroupName);
-				const bool bPassLoopingFilter = !GatherParams.bFilterByLoopingCondition || (SequencePlayer->GetLoopAnimation() == GatherParams.bLoopingCondition);
+				const bool bPassLoopingFilter = !GatherParams.bFilterByLoopingCondition || (SequencePlayer->IsLooping() == GatherParams.bLoopingCondition);
 				if (bPassSyncGroupFilter && bPassLoopingFilter)
 				{
 					PCV_GatherAnimSequences(OutAnimSequences, SequencePlayer->GetSequence());
diff --git a/Engine/Source/Editor/AnimGraph/Private/AnimGraphNode_StateMachineBase.cpp b/Engine/Source/Editor/AnimGraph/Private/AnimGraphNode_StateMachineBase.cpp
--- a/Engine/Source/Editor/AnimGraph/Private/AnimGraphNode_StateMachineBase.cpp
+++ b/Engine/Source/Editor/AnimGraph/Private/AnimGraphNode_StateMachineBase.cpp
@@ -698,6 +698,7 @@ void UAnimGraphNode_StateMachineBase::OnProcessDuringCompilation(IAnimBlueprintC
 
 			// Handle automatic time remaining rules
 			Rule.bAutomaticRemainingTimeRule = TransitionNode->bAutomaticRuleBasedOnSequencePlayerInState;
+			Rule.AutomaticRuleTriggerTime = TransitionNode->AutomaticRuleTriggerTime;
 			Rule.SyncGroupNameToRequireValidMarkersRule = TransitionNode->SyncGroupNameToRequireValidMarkersRule;
 
 			// Handle custom transition graphs
diff --git a/Engine/Source/Editor/AnimGraph/Private/AnimStateTransitionNode.cpp b/Engine/Source/Editor/AnimGraph/Private/AnimStateTransitionNode.cpp
--- a/Engine/Source/Editor/AnimGraph/Private/AnimStateTransitionNode.cpp
+++ b/Engine/Source/Editor/AnimGraph/Private/AnimStateTransitionNode.cpp
@@ -7,11 +7,16 @@
 #include "AnimStateTransitionNode.h"
 #include "Kismet2/BlueprintEditorUtils.h"
 #include "Animation/AnimInstance.h"
+#include "AnimationStateGraph.h"
 #include "AnimationTransitionGraph.h"
 #include "AnimationTransitionSchema.h"
 #include "AnimationCustomTransitionGraph.h"
 #include "AnimationCustomTransitionSchema.h"
+#include "AnimGraphNode_BlendSpacePlayer.h"
+#include "AnimGraphNode_SequencePlayer.h"
+#include "AnimGraphNode_StateResult.h"
 #include "AnimGraphNode_TransitionResult.h"
+#include "AnimStateConduitNode.h"
 #include "Kismet2/CompilerResultsLog.h"
 #include "EdGraphUtilities.h"
 #include "Kismet2/Kismet2NameValidators.h"
@@ -73,6 +78,8 @@ UAnimStateTransitionNode::UAnimStateTransitionNode(const FObjectInitializer& Obj
 
 	CrossfadeDuration = 0.2f;
 	BlendMode = EAlphaBlendOption::HermiteCubic;
+	bAutomaticRuleBasedOnSequencePlayerInState = false;
+	AutomaticRuleTriggerTime = -1.f;
 	bSharedRules = false;
 	SharedRulesGuid.Invalidate();
 	bSharedCrossfade = false;
@@ -684,7 +691,41 @@ void UAnimStateTransitionNode::ValidateNodeDuringCompilation(class FCompilerResu
 		UAnimGraphNode_TransitionResult* ResultNode = TransGraph->GetResultNode();
 		check(ResultNode);
 
-		if (ResultNode->PropertyBindings.Num() > 0 && ResultNode->PropertyBindings.CreateIterator()->Value.bIsBound)
+		if (bAutomaticRuleBasedOnSequencePlayerInState)
+		{
+			// Check for automatic transition rules that are being triggered from looping asset players, as these can often be symptomatic of logic errors
+			if (UAnimStateNodeBase* PreviousState = GetPreviousState())
+			{
+				if (UEdGraph* PreviousStateGraph = PreviousState->GetBoundGraph())
+				{
+					if (UAnimGraphNode_StateResult* PreviousStateGraphResultNode = CastChecked<UAnimationStateGraph>(PreviousStateGraph)->GetResultNode())
+					{
+						for (UEdGraphPin* TestPin : PreviousStateGraphResultNode->Pins)
+						{
+							// Warn for the trivial but common case of a looping asset player connected directly to the result node
+							if ((TestPin->Direction == EGPD_Input) && (TestPin->LinkedTo.Num() == 1))
+							{
+								if (UAnimGraphNode_SequencePlayer* SequencePlayer = Cast<UAnimGraphNode_SequencePlayer>(TestPin->LinkedTo[0]->GetOwningNode()))
+								{
+									if (SequencePlayer->Node.IsLooping())
+									{
+										MessageLog.Warning(TEXT("Transition @@ is using an automatic transition rule but the source @@ is set as looping.  Please clear the 'Loop Animation' flag"), this, SequencePlayer);
+									}
+								}
+								else if (UAnimGraphNode_BlendSpacePlayer* BlendSpacePlayer = Cast<UAnimGraphNode_BlendSpacePlayer>(TestPin->LinkedTo[0]->GetOwningNode()))
+								{
+									if (BlendSpacePlayer->Node.IsLooping())
+									{
+										MessageLog.Warning(TEXT("Transition @@ is using an automatic transition rule but the source @@ is set as looping.  Please clear the 'Loop' flag"), this, BlendSpacePlayer);
+									}
+								}
+							}
+						}
+					}
+				}
+			}
+		}
+		else if (ResultNode->PropertyBindings.Num() > 0 && ResultNode->PropertyBindings.CreateIterator()->Value.bIsBound)
 		{
 			// Rule is bound so nothing more to check
 		}
diff --git a/Engine/Source/Editor/AnimGraph/Public/AnimStateTransitionNode.h b/Engine/Source/Editor/AnimGraph/Public/AnimStateTransitionNode.h
--- a/Engine/Source/Editor/AnimGraph/Public/AnimStateTransitionNode.h
+++ b/Engine/Source/Editor/AnimGraph/Public/AnimStateTransitionNode.h
@@ -52,10 +52,16 @@ class UAnimStateTransitionNode : public UAnimStateNodeBase
 	UPROPERTY(EditAnywhere, Category=Transition, meta=(UseAsBlendProfile=true))
 	TObjectPtr<UBlendProfile> BlendProfile;
 
-	// Try setting the rule automatically based on most relevant player node's remaining time and the CrossfadeDuration of the transition, ignoring the internal time
+	// Try setting the rule automatically based on most relevant asset player node's remaining time and the Automatic Rule Trigger Time of the transition, ignoring the internal time
 	UPROPERTY(EditAnywhere, Category=Transition)
 	bool bAutomaticRuleBasedOnSequencePlayerInState;
 
+	// When should the automatic transition rule trigger relative to the time remaining on the relevant asset player:
+	//  < 0 means trigger the transition 'Crossfade Duration' seconds before the end of the asset player, so a standard blend would finish just as the asset player ends
+	// >= 0 means trigger the transition 'Automatic Rule Trigger Time' seconds before the end of the asset player
+	UPROPERTY(EditAnywhere, Category=Transition)
+	float AutomaticRuleTriggerTime;
+
 	// If SyncGroupName is specified, Transition will only be taken if the SyncGroup has valid markers (other transition rules still apply in addition to that).
 	UPROPERTY(EditAnywhere, Category = Transition)
 	FName SyncGroupNameToRequireValidMarkersRule;
diff --git a/Engine/Source/Runtime/AnimGraphRuntime/Private/AnimNodes/AnimNode_BlendSpacePlayer.cpp b/Engine/Source/Runtime/AnimGraphRuntime/Private/AnimNodes/AnimNode_BlendSpacePlayer.cpp
--- a/Engine/Source/Runtime/AnimGraphRuntime/Private/AnimNodes/AnimNode_BlendSpacePlayer.cpp
+++ b/Engine/Source/Runtime/AnimGraphRuntime/Private/AnimNodes/AnimNode_BlendSpacePlayer.cpp
@@ -92,7 +92,7 @@ void FAnimNode_BlendSpacePlayerBase::UpdateInternal(const FAnimationUpdateContex
 		UE::Anim::FAnimSyncGroupScope& SyncScope = Context.GetMessageChecked<UE::Anim::FAnimSyncGroupScope>();
 
 		FAnimTickRecord TickRecord(
-			CurrentBlendSpace, Position, BlendSampleDataCache, BlendFilter, GetLoop(), GetPlayRate(), ShouldTeleportToTime(), 
+			CurrentBlendSpace, Position, BlendSampleDataCache, BlendFilter, IsLooping(), GetPlayRate(), ShouldTeleportToTime(),
 			IsEvaluator(), Context.GetFinalBlendWeight(), /*inout*/ InternalTimeAccumulator, MarkerTickRecord);
 		TickRecord.RootMotionWeightModifier = Context.GetRootMotionWeightModifier();
 		TickRecord.DeltaTimeRecord = &DeltaTimeRecord;
@@ -129,7 +129,7 @@ void FAnimNode_BlendSpacePlayerBase::Evaluate_AnyThread(FPoseContext& Output)
 	if (CurrentBlendSpace != nullptr && CurrentBlendSpace->GetSkeleton() != nullptr)
 	{
 		FAnimationPoseData AnimationPoseData(Output);
-		CurrentBlendSpace->GetAnimationPose(BlendSampleDataCache, FAnimExtractContext(static_cast<double>(InternalTimeAccumulator), Output.AnimInstanceProxy->ShouldExtractRootMotion(), DeltaTimeRecord, GetLoop()), AnimationPoseData);
+		CurrentBlendSpace->GetAnimationPose(BlendSampleDataCache, FAnimExtractContext(static_cast<double>(InternalTimeAccumulator), Output.AnimInstanceProxy->ShouldExtractRootMotion(), DeltaTimeRecord, IsLooping()), AnimationPoseData);
 	}
 	else
 	{
@@ -367,7 +367,7 @@ bool FAnimNode_BlendSpacePlayer::SetPlayRate(float InPlayRate)
 	return false;
 }
 
-bool FAnimNode_BlendSpacePlayer::GetLoop() const
+bool FAnimNode_BlendSpacePlayer::IsLooping() const
 {
 	return GET_ANIM_NODE_DATA(bool, bLoop);
 }
diff --git a/Engine/Source/Runtime/AnimGraphRuntime/Private/AnimNodes/AnimNode_RandomPlayer.cpp b/Engine/Source/Runtime/AnimGraphRuntime/Private/AnimNodes/AnimNode_RandomPlayer.cpp
--- a/Engine/Source/Runtime/AnimGraphRuntime/Private/AnimNodes/AnimNode_RandomPlayer.cpp
+++ b/Engine/Source/Runtime/AnimGraphRuntime/Private/AnimNodes/AnimNode_RandomPlayer.cpp
@@ -382,6 +382,16 @@ void FAnimNode_RandomPlayer::ClearCachedBlendWeight()
 	BlendWeight = 0.f;
 }
 
+const FDeltaTimeRecord* FAnimNode_RandomPlayer::GetDeltaTimeRecord() const
+{
+	if (ValidEntries.Num() > 0)
+	{
+		const FRandomAnimPlayData& CurrentPlayData = GetPlayData(ERandomDataIndexType::Current);
+		return &CurrentPlayData.DeltaTimeRecord;
+	}
+	return nullptr;
+}
+
 int32 FAnimNode_RandomPlayer::GetNextValidEntryIndex()
 {
 	check(ValidEntries.Num() > 0);
diff --git a/Engine/Source/Runtime/AnimGraphRuntime/Private/AnimNodes/AnimNode_SequenceEvaluator.cpp b/Engine/Source/Runtime/AnimGraphRuntime/Private/AnimNodes/AnimNode_SequenceEvaluator.cpp
--- a/Engine/Source/Runtime/AnimGraphRuntime/Private/AnimNodes/AnimNode_SequenceEvaluator.cpp
+++ b/Engine/Source/Runtime/AnimGraphRuntime/Private/AnimNodes/AnimNode_SequenceEvaluator.cpp
@@ -67,12 +67,12 @@ void FAnimNode_SequenceEvaluatorBase::UpdateAssetPlayer(const FAnimationUpdateCo
 			const float DeltaTime = Context.GetDeltaTime();
 			const float RateScale = CurrentSequence->RateScale;
 			const float PlayRate = FMath::IsNearlyZero(DeltaTime) || FMath::IsNearlyZero(RateScale) ? 0.f : (TimeJump / (DeltaTime * RateScale));
-			CreateTickRecordForNode(Context, CurrentSequence, GetShouldLoop(), PlayRate, true);
+			CreateTickRecordForNode(Context, CurrentSequence, IsLooping(), PlayRate, true);
 		}
 		else
 		{
 			InternalTimeAccumulator = CurrentExplicitTime;
-			CreateTickRecordForNode(Context, CurrentSequence, GetShouldLoop(), 0, true);
+			CreateTickRecordForNode(Context, CurrentSequence, IsLooping(), 0, true);
 		}
 	}
 
@@ -92,7 +92,7 @@ void FAnimNode_SequenceEvaluatorBase::Evaluate_AnyThread(FPoseContext& Output)
 	if (CurrentSequence != nullptr && CurrentSequence->GetSkeleton() != nullptr)
 	{
 		FAnimationPoseData AnimationPoseData(Output);
-		CurrentSequence->GetAnimationPose(AnimationPoseData, FAnimExtractContext(static_cast<double>(InternalTimeAccumulator), Output.AnimInstanceProxy->ShouldExtractRootMotion(), DeltaTimeRecord, GetShouldLoop()));
+		CurrentSequence->GetAnimationPose(AnimationPoseData, FAnimExtractContext(static_cast<double>(InternalTimeAccumulator), Output.AnimInstanceProxy->ShouldExtractRootMotion(), DeltaTimeRecord, IsLooping()));
 	}
 	else
 	{
@@ -113,7 +113,7 @@ float FAnimNode_SequenceEvaluatorBase::GetEffectiveDeltaTime(float ExplicitTime,
 {
 	float DeltaTime = ExplicitTime - PrevExplicitTime;
 
-	if (GetShouldLoop())
+	if (IsLooping())
 	{
 		if (FMath::Abs(DeltaTime) > (GetSequence()->GetPlayLength() * 0.5f))
 		{
@@ -187,7 +187,7 @@ bool FAnimNode_SequenceEvaluator::SetShouldLoop(bool bInShouldLoop)
 	return false;
 }
 
-bool FAnimNode_SequenceEvaluator::GetShouldLoop() const
+bool FAnimNode_SequenceEvaluator::IsLooping() const
 {
 	return GET_ANIM_NODE_DATA(bool, bShouldLoop);
 }
diff --git a/Engine/Source/Runtime/AnimGraphRuntime/Private/BlendSpacePlayerLibrary.cpp b/Engine/Source/Runtime/AnimGraphRuntime/Private/BlendSpacePlayerLibrary.cpp
--- a/Engine/Source/Runtime/AnimGraphRuntime/Private/BlendSpacePlayerLibrary.cpp
+++ b/Engine/Source/Runtime/AnimGraphRuntime/Private/BlendSpacePlayerLibrary.cpp
@@ -176,7 +176,7 @@ bool UBlendSpacePlayerLibrary::GetLoop(const FBlendSpacePlayerReference& BlendSp
 		TEXT("GetLoop"),
 		[&bLoop](FAnimNode_BlendSpacePlayer& InBlendSpacePlayer)
 		{
-			bLoop = InBlendSpacePlayer.GetLoop();
+			bLoop = InBlendSpacePlayer.IsLooping();
 		});
 	return bLoop;
 }
diff --git a/Engine/Source/Runtime/AnimGraphRuntime/Private/SequenceEvaluatorLibrary.cpp b/Engine/Source/Runtime/AnimGraphRuntime/Private/SequenceEvaluatorLibrary.cpp
--- a/Engine/Source/Runtime/AnimGraphRuntime/Private/SequenceEvaluatorLibrary.cpp
+++ b/Engine/Source/Runtime/AnimGraphRuntime/Private/SequenceEvaluatorLibrary.cpp
@@ -40,7 +40,7 @@ FSequenceEvaluatorReference USequenceEvaluatorLibrary::AdvanceTime(const FAnimUp
 			if (const FAnimationUpdateContext* AnimationUpdateContext = UpdateContext.GetContext())
 			{
 				float ExplicitTime = InSequenceEvaluator.GetExplicitTime();
-				FAnimationRuntime::AdvanceTime(InSequenceEvaluator.GetShouldLoop(), AnimationUpdateContext->GetDeltaTime() * PlayRate, ExplicitTime, InSequenceEvaluator.GetCurrentAssetLength());
+				FAnimationRuntime::AdvanceTime(InSequenceEvaluator.IsLooping(), AnimationUpdateContext->GetDeltaTime() * PlayRate, ExplicitTime, InSequenceEvaluator.GetCurrentAssetLength());
 
 				if (!InSequenceEvaluator.SetExplicitTime(ExplicitTime))
 				{
diff --git a/Engine/Source/Runtime/AnimGraphRuntime/Private/SequencePlayerLibrary.cpp b/Engine/Source/Runtime/AnimGraphRuntime/Private/SequencePlayerLibrary.cpp
--- a/Engine/Source/Runtime/AnimGraphRuntime/Private/SequencePlayerLibrary.cpp
+++ b/Engine/Source/Runtime/AnimGraphRuntime/Private/SequencePlayerLibrary.cpp
@@ -183,7 +183,7 @@ bool USequencePlayerLibrary::GetLoopAnimation(const FSequencePlayerReference& Se
 		TEXT("GetLoopAnimation"),
 		[&bLoopAnimation](FAnimNode_SequencePlayer& InSequencePlayer)
 		{
-			bLoopAnimation = InSequencePlayer.GetLoopAnimation();
+			bLoopAnimation = InSequencePlayer.IsLooping();
 		});
 
 	return bLoopAnimation;
diff --git a/Engine/Source/Runtime/AnimGraphRuntime/Public/AnimNodes/AnimNode_BlendSpacePlayer.h b/Engine/Source/Runtime/AnimGraphRuntime/Public/AnimNodes/AnimNode_BlendSpacePlayer.h
--- a/Engine/Source/Runtime/AnimGraphRuntime/Public/AnimNodes/AnimNode_BlendSpacePlayer.h
+++ b/Engine/Source/Runtime/AnimGraphRuntime/Public/AnimNodes/AnimNode_BlendSpacePlayer.h
@@ -70,7 +70,8 @@ public:
 	virtual float GetPlayRate() const PURE_VIRTUAL(FAnimNode_BlendSpacePlayerBase::GetPlayRate, return 1.0f;);
 
 	// Should the animation loop back to the start when it reaches the end?
-	virtual bool GetLoop() const PURE_VIRTUAL(FAnimNode_BlendSpacePlayerBase::GetLoop, return true;);
+	UE_DEPRECATED(5.3, "Please use IsLooping instead.")
+	virtual bool GetLoop() const final { return IsLooping(); }
 
 	// Get whether we should reset the current play time when the blend space changes
 	virtual bool ShouldResetPlayTimeWhenBlendSpaceChanges() const PURE_VIRTUAL(FAnimNode_BlendSpacePlayerBase::ShouldResetPlayTimeWhenBlendSpaceChanges, return true;);
@@ -182,6 +183,7 @@ public:
 	virtual EAnimGroupRole::Type GetGroupRole() const override;
 	virtual EAnimSyncMethod GetGroupMethod() const override;
 	virtual bool GetIgnoreForRelevancyTest() const override;
+	virtual bool IsLooping() const override;
 	virtual bool SetGroupName(FName InGroupName) override;
 	virtual bool SetGroupRole(EAnimGroupRole::Type InRole) override;
 	virtual bool SetGroupMethod(EAnimSyncMethod InMethod) override;
@@ -193,7 +195,6 @@ public:
 	virtual FVector GetPosition() const override;
 	virtual float GetStartPosition() const override;
 	virtual float GetPlayRate() const override;
-	virtual bool GetLoop() const override;
 	virtual bool ShouldResetPlayTimeWhenBlendSpaceChanges() const override;
 	virtual bool SetResetPlayTimeWhenBlendSpaceChanges(bool bReset) override;
 	virtual bool SetBlendSpace(UBlendSpace* InBlendSpace) override;
@@ -262,6 +263,7 @@ public:
 	virtual EAnimGroupRole::Type GetGroupRole() const override { return GroupRole; }
 	virtual EAnimSyncMethod GetGroupMethod() const override { return Method; }
 	virtual bool GetIgnoreForRelevancyTest() const override { return bIgnoreForRelevancyTest; }
+	virtual bool IsLooping() const override { return bLoop; }
 	virtual bool SetGroupName(FName InGroupName) override { GroupName = InGroupName; return true; }
 	virtual bool SetGroupRole(EAnimGroupRole::Type InRole) override { GroupRole = InRole; return true; }
 	virtual bool SetGroupMethod(EAnimSyncMethod InMethod) override { Method = InMethod; return true; }
@@ -273,7 +275,6 @@ public:
 	virtual FVector GetPosition() const override { return FVector(X, Y, 0.0); }
 	virtual float GetStartPosition() const override { return StartPosition; }
 	virtual float GetPlayRate() const override { return PlayRate; }
-	virtual bool GetLoop() const override { return bLoop; }
 	virtual bool ShouldResetPlayTimeWhenBlendSpaceChanges() const override { return bResetPlayTimeWhenBlendSpaceChanges; }
 	virtual bool SetResetPlayTimeWhenBlendSpaceChanges(bool bReset) override { bResetPlayTimeWhenBlendSpaceChanges = bReset; return true; }
 	virtual bool SetBlendSpace(UBlendSpace* InBlendSpace) override { BlendSpace = InBlendSpace; return true; }
diff --git a/Engine/Source/Runtime/AnimGraphRuntime/Public/AnimNodes/AnimNode_RandomPlayer.h b/Engine/Source/Runtime/AnimGraphRuntime/Public/AnimNodes/AnimNode_RandomPlayer.h
--- a/Engine/Source/Runtime/AnimGraphRuntime/Public/AnimNodes/AnimNode_RandomPlayer.h
+++ b/Engine/Source/Runtime/AnimGraphRuntime/Public/AnimNodes/AnimNode_RandomPlayer.h
@@ -118,6 +118,7 @@ public:
 	virtual bool SetIgnoreForRelevancyTest(bool bInIgnoreForRelevancyTest) override;
 	virtual float GetCachedBlendWeight() const override;
 	virtual void ClearCachedBlendWeight() override;
+	virtual const FDeltaTimeRecord* GetDeltaTimeRecord() const override;
 	// End of FAnimNode_RelevantAssetPlayerBase
 
 private:
diff --git a/Engine/Source/Runtime/AnimGraphRuntime/Public/AnimNodes/AnimNode_SequenceEvaluator.h b/Engine/Source/Runtime/AnimGraphRuntime/Public/AnimNodes/AnimNode_SequenceEvaluator.h
--- a/Engine/Source/Runtime/AnimGraphRuntime/Public/AnimNodes/AnimNode_SequenceEvaluator.h
+++ b/Engine/Source/Runtime/AnimGraphRuntime/Public/AnimNodes/AnimNode_SequenceEvaluator.h
@@ -77,7 +77,8 @@ public:
 	virtual float GetExplicitTime() const { return 0.0f; }
 
 	/** This only works if bTeleportToExplicitTime is false OR this node is set to use SyncGroup */
-	virtual bool GetShouldLoop() const { return true; }
+	UE_DEPRECATED(5.3, "Please use IsLooping instead.")
+	virtual bool GetShouldLoop() const final { return IsLooping(); }
 
 	// Set the animation to continue looping when it reaches the end
 	virtual bool SetShouldLoop(bool bInShouldLoop) { return false; }
@@ -153,7 +154,6 @@ public:
 	virtual UAnimSequenceBase* GetSequence() const override;
 	virtual float GetExplicitTime() const override;
 	virtual bool SetExplicitTime(float InTime) override;
-	virtual bool GetShouldLoop() const override;
 	virtual bool SetShouldLoop(bool bInShouldLoop) override;
 	virtual bool GetTeleportToExplicitTime() const override;
 	virtual TEnumAsByte<ESequenceEvalReinit::Type> GetReinitializationBehavior() const override;
@@ -164,6 +164,7 @@ public:
 	virtual EAnimGroupRole::Type GetGroupRole() const override;
 	virtual EAnimSyncMethod GetGroupMethod() const override;
 	virtual bool GetIgnoreForRelevancyTest() const override;
+	virtual bool IsLooping() const override;
 	virtual bool SetGroupName(FName InGroupName) override;
 	virtual bool SetGroupRole(EAnimGroupRole::Type InRole) override;
 	virtual bool SetGroupMethod(EAnimSyncMethod InMethod) override;
@@ -227,7 +228,6 @@ public:
 	virtual void SetReinitializationBehavior(TEnumAsByte<ESequenceEvalReinit::Type> InBehavior) override { ReinitializationBehavior = InBehavior; }
 	virtual UAnimSequenceBase* GetSequence() const override { return Sequence; }
 	virtual float GetExplicitTime() const override { return ExplicitTime; }
-	virtual bool GetShouldLoop() const override { return bShouldLoop; }
 	virtual bool GetTeleportToExplicitTime() const override { return bTeleportToExplicitTime; }
 	virtual TEnumAsByte<ESequenceEvalReinit::Type> GetReinitializationBehavior() const override { return ReinitializationBehavior; }
 	virtual float GetStartPosition() const override { return StartPosition; }
@@ -237,6 +237,7 @@ public:
 	virtual EAnimGroupRole::Type GetGroupRole() const override { return GroupRole; }
 	virtual EAnimSyncMethod GetGroupMethod() const override { return Method; }
 	virtual bool GetIgnoreForRelevancyTest() const override { return bIgnoreForRelevancyTest; }
+	virtual bool IsLooping() const override { return bShouldLoop; }
 	virtual bool SetGroupName(FName InGroupName) override { GroupName = InGroupName; return true; }
 	virtual bool SetGroupRole(EAnimGroupRole::Type InRole) override { GroupRole = InRole; return true; }
 	virtual bool SetGroupMethod(EAnimSyncMethod InMethod) override { Method = InMethod; return true; }
diff --git a/Engine/Source/Runtime/Engine/Classes/Animation/AnimNode_AssetPlayerBase.h b/Engine/Source/Runtime/Engine/Classes/Animation/AnimNode_AssetPlayerBase.h
--- a/Engine/Source/Runtime/Engine/Classes/Animation/AnimNode_AssetPlayerBase.h
+++ b/Engine/Source/Runtime/Engine/Classes/Animation/AnimNode_AssetPlayerBase.h
@@ -55,7 +55,8 @@ struct ENGINE_API FAnimNode_AssetPlayerBase : public FAnimNode_AssetPlayerReleva
 	virtual void SetAccumulatedTime(float NewTime) override;
 	virtual float GetCachedBlendWeight() const override;
 	virtual void ClearCachedBlendWeight() override;
-	virtual float GetCurrentAssetTimePlayRateAdjusted() const;
+	virtual float GetCurrentAssetTimePlayRateAdjusted() const override;
+	virtual const FDeltaTimeRecord* GetDeltaTimeRecord() const override;
 	// --- End of FAnimNode_RelevantAssetPlayerBase ---
 
 private:
diff --git a/Engine/Source/Runtime/Engine/Classes/Animation/AnimNode_RelevantAssetPlayerBase.h b/Engine/Source/Runtime/Engine/Classes/Animation/AnimNode_RelevantAssetPlayerBase.h
--- a/Engine/Source/Runtime/Engine/Classes/Animation/AnimNode_RelevantAssetPlayerBase.h
+++ b/Engine/Source/Runtime/Engine/Classes/Animation/AnimNode_RelevantAssetPlayerBase.h
@@ -29,6 +29,9 @@ public:
 	virtual float GetCurrentAssetTime() const;
 	virtual float GetCurrentAssetTimePlayRateAdjusted() const;
 
+	// Does this asset player loop back to the start when it reaches the end?
+	virtual bool IsLooping() const;
+
 	// Check whether this node should be ignored when testing for relevancy in state machines
 	virtual bool GetIgnoreForRelevancyTest() const;
 
@@ -40,4 +43,7 @@ public:
 
 	/** Set the cached blendweight to zero */
 	virtual void ClearCachedBlendWeight();
+
+	/** Get the delta time record owned by this asset player (or null) */
+	virtual const FDeltaTimeRecord* GetDeltaTimeRecord() const;
 };
diff --git a/Engine/Source/Runtime/Engine/Classes/Animation/AnimNode_SequencePlayer.h b/Engine/Source/Runtime/Engine/Classes/Animation/AnimNode_SequencePlayer.h
--- a/Engine/Source/Runtime/Engine/Classes/Animation/AnimNode_SequencePlayer.h
+++ b/Engine/Source/Runtime/Engine/Classes/Animation/AnimNode_SequencePlayer.h
@@ -44,6 +44,10 @@ public:
 	// The animation sequence asset to play
 	virtual UAnimSequenceBase* GetSequence() const { return nullptr; }
 
+	// Should the animation loop back to the start when it reaches the end?
+	UE_DEPRECATED(5.3, "Please use IsLooping instead.")
+	virtual bool GetLoopAnimation() const final { return IsLooping(); }
+
 protected:
 	// Set the animation sequence asset to play
 	virtual bool SetSequence(UAnimSequenceBase* InSequence) { return false; }
@@ -65,9 +69,6 @@ protected:
 	// The start position [range: 0 - sequence length] to use when initializing. When looping, play will still jump back to the beginning when reaching the end.
 	virtual float GetStartPosition() const { return 0.0f; }
 
-	// Should the animation loop back to the start when it reaches the end?
-	virtual bool GetLoopAnimation() const { return true; }
-
 	// Use pose matching to choose the start position. Requires experimental PoseSearch plugin.
 	virtual bool GetStartFromMatchingPose() const { return false; }
 
@@ -151,7 +152,6 @@ public:
 	virtual float GetPlayRate() const override;
 	virtual const FInputScaleBiasClampConstants& GetPlayRateScaleBiasClampConstants() const override;
 	virtual float GetStartPosition() const override;
-	virtual bool GetLoopAnimation() const override;
 	virtual bool GetStartFromMatchingPose() const override;
 	virtual bool SetStartPosition(float InStartPosition) override;
 	virtual bool SetPlayRate(float InPlayRate) override;
@@ -160,6 +160,7 @@ public:
 	virtual FName GetGroupName() const override;
 	virtual EAnimGroupRole::Type GetGroupRole() const override;
 	virtual EAnimSyncMethod GetGroupMethod() const override;
+	virtual bool IsLooping() const override;
 	virtual bool GetIgnoreForRelevancyTest() const override;
 	virtual bool SetGroupName(FName InGroupName) override;
 	virtual bool SetGroupRole(EAnimGroupRole::Type InRole) override;
@@ -229,13 +230,13 @@ public:
 	virtual float GetPlayRate() const override { return PlayRate; }
 	virtual const FInputScaleBiasClampConstants& GetPlayRateScaleBiasClampConstants() const override { return PlayRateScaleBiasClampConstants; }
 	virtual float GetStartPosition() const override { return StartPosition; }
-	virtual bool GetLoopAnimation() const override { return bLoopAnimation; }
 	virtual bool GetStartFromMatchingPose() const override { return bStartFromMatchingPose; }
 
 	// FAnimNode_AssetPlayerBase interface
 	virtual FName GetGroupName() const override { return GroupName; }
 	virtual EAnimGroupRole::Type GetGroupRole() const override { return GroupRole; }
 	virtual EAnimSyncMethod GetGroupMethod() const override { return Method; }
+	virtual bool IsLooping() const override { return bLoopAnimation; }
 	virtual bool GetIgnoreForRelevancyTest() const override { return bIgnoreForRelevancyTest; }
 	virtual bool SetGroupName(FName InGroupName) override { GroupName = InGroupName; return true; }
 	virtual bool SetGroupRole(EAnimGroupRole::Type InRole) override { GroupRole = InRole; return true; }
diff --git a/Engine/Source/Runtime/Engine/Classes/Animation/AnimStateMachineTypes.h b/Engine/Source/Runtime/Engine/Classes/Animation/AnimStateMachineTypes.h
--- a/Engine/Source/Runtime/Engine/Classes/Animation/AnimStateMachineTypes.h
+++ b/Engine/Source/Runtime/Engine/Classes/Animation/AnimStateMachineTypes.h
@@ -267,7 +267,13 @@ struct FBakedStateExitTransition
 	// Automatic Transition Rule based on animation remaining time.
 	UPROPERTY()
 	bool bAutomaticRemainingTimeRule;
-	
+
+	// Automatic Transition Rule triggering time:
+	//  < 0 means trigger the transition 'Crossfade Duration' seconds before the end of the asset player, so a standard blend would finish just as the asset player ends
+	// >= 0 means trigger the transition 'Automatic Rule Trigger Time' seconds before the end of the asset player
+	UPROPERTY()
+	float AutomaticRuleTriggerTime;
+
 	// Additional rule around SyncGroup requiring Valid Markers
 	UPROPERTY()
 	FName SyncGroupNameToRequireValidMarkersRule;
@@ -281,6 +287,7 @@ struct FBakedStateExitTransition
 		, TransitionIndex(INDEX_NONE)
 		, bDesiredTransitionReturnValue(true)
 		, bAutomaticRemainingTimeRule(false)
+		, AutomaticRuleTriggerTime(-1.f)
 		, SyncGroupNameToRequireValidMarkersRule(NAME_None)
 	{
 	}
diff --git a/Engine/Source/Runtime/Engine/Private/Animation/AnimNode_AssetPlayerBase.cpp b/Engine/Source/Runtime/Engine/Private/Animation/AnimNode_AssetPlayerBase.cpp
--- a/Engine/Source/Runtime/Engine/Private/Animation/AnimNode_AssetPlayerBase.cpp
+++ b/Engine/Source/Runtime/Engine/Private/Animation/AnimNode_AssetPlayerBase.cpp
@@ -98,5 +98,10 @@ void FAnimNode_AssetPlayerBase::ClearCachedBlendWeight()
 
 float FAnimNode_AssetPlayerBase::GetCurrentAssetTimePlayRateAdjusted() const
 {
-	return  GetCurrentAssetTime();
+	return GetCurrentAssetTime();
+}
+
+const FDeltaTimeRecord* FAnimNode_AssetPlayerBase::GetDeltaTimeRecord() const
+{
+	return &DeltaTimeRecord;
 }
diff --git a/Engine/Source/Runtime/Engine/Private/Animation/AnimNode_RelevantAssetPlayerBase.cpp b/Engine/Source/Runtime/Engine/Private/Animation/AnimNode_RelevantAssetPlayerBase.cpp
--- a/Engine/Source/Runtime/Engine/Private/Animation/AnimNode_RelevantAssetPlayerBase.cpp
+++ b/Engine/Source/Runtime/Engine/Private/Animation/AnimNode_RelevantAssetPlayerBase.cpp
@@ -33,6 +33,11 @@ float FAnimNode_AssetPlayerRelevancyBase::GetCurrentAssetTimePlayRateAdjusted()
 	return 0.f;
 }
 
+bool FAnimNode_AssetPlayerRelevancyBase::IsLooping() const
+{
+	return true;
+}
+
 bool FAnimNode_AssetPlayerRelevancyBase::GetIgnoreForRelevancyTest() const
 {
 	return false;
@@ -51,3 +56,8 @@ float FAnimNode_AssetPlayerRelevancyBase::GetCachedBlendWeight() const
 void FAnimNode_AssetPlayerRelevancyBase::ClearCachedBlendWeight()
 {
 }
+
+const FDeltaTimeRecord* FAnimNode_AssetPlayerRelevancyBase::GetDeltaTimeRecord() const
+{
+	return nullptr;
+}
diff --git a/Engine/Source/Runtime/Engine/Private/Animation/AnimNode_SequencePlayer.cpp b/Engine/Source/Runtime/Engine/Private/Animation/AnimNode_SequencePlayer.cpp
--- a/Engine/Source/Runtime/Engine/Private/Animation/AnimNode_SequencePlayer.cpp
+++ b/Engine/Source/Runtime/Engine/Private/Animation/AnimNode_SequencePlayer.cpp
@@ -97,7 +97,7 @@ void FAnimNode_SequencePlayerBase::UpdateAssetPlayer(const FAnimationUpdateConte
 
 		InternalTimeAccumulator = FMath::Clamp(InternalTimeAccumulator, 0.f, CurrentSequence->GetPlayLength());
 		const float AdjustedPlayRate = PlayRateScaleBiasClampState.ApplyTo(GetPlayRateScaleBiasClampConstants(), FMath::IsNearlyZero(CurrentPlayRateBasis) ? 0.f : (CurrentPlayRate / CurrentPlayRateBasis), Context.GetDeltaTime());
-		CreateTickRecordForNode(Context, CurrentSequence, GetLoopAnimation(), AdjustedPlayRate, false);
+		CreateTickRecordForNode(Context, CurrentSequence, IsLooping(), AdjustedPlayRate, false);
 	}
 
 #if WITH_EDITORONLY_DATA
@@ -130,7 +130,7 @@ void FAnimNode_SequencePlayerBase::Evaluate_AnyThread(FPoseContext& Output)
 		}
 
 		FAnimationPoseData AnimationPoseData(Output);
-		CurrentSequence->GetAnimationPose(AnimationPoseData, FAnimExtractContext(static_cast<double>(InternalTimeAccumulator), Output.AnimInstanceProxy->ShouldExtractRootMotion(), DeltaTimeRecord, GetLoopAnimation()));
+		CurrentSequence->GetAnimationPose(AnimationPoseData, FAnimExtractContext(static_cast<double>(InternalTimeAccumulator), Output.AnimInstanceProxy->ShouldExtractRootMotion(), DeltaTimeRecord, IsLooping()));
 	}
 	else
 	{
@@ -204,7 +204,7 @@ float FAnimNode_SequencePlayer::GetStartPosition() const
 	return GET_ANIM_NODE_DATA(float, StartPosition);
 }
 
-bool FAnimNode_SequencePlayer::GetLoopAnimation() const
+bool FAnimNode_SequencePlayer::IsLooping() const
 {
 	return GET_ANIM_NODE_DATA(bool, bLoopAnimation);
 }
diff --git a/Engine/Source/Runtime/Engine/Private/Animation/AnimNode_StateMachine.cpp b/Engine/Source/Runtime/Engine/Private/Animation/AnimNode_StateMachine.cpp
--- a/Engine/Source/Runtime/Engine/Private/Animation/AnimNode_StateMachine.cpp
+++ b/Engine/Source/Runtime/Engine/Private/Animation/AnimNode_StateMachine.cpp
@@ -762,12 +762,44 @@ bool FAnimNode_StateMachine::FindValidTransition(const FAnimationUpdateContext&
 			{
 				if (const UAnimationAsset* AnimAsset = RelevantPlayer->GetAnimAsset())
 				{
-					const float AnimTimeRemaining = AnimAsset->GetPlayLength() - RelevantPlayer->GetAccumulatedTime();
-					const FAnimationTransitionBetweenStates& TransitionInfo = GetTransitionInfo(TransitionRule.TransitionIndex);
-					CrossfadeTimeAdjustment = TransitionInfo.CrossfadeDuration - AnimTimeRemaining;
-					bCanEnterTransition = (CrossfadeTimeAdjustment >= 0.f);
+					float AnimTimeRemaining = AnimAsset->GetPlayLength() - RelevantPlayer->GetAccumulatedTime();
+
+					// If the asset player is a looping player and has looped, then we enter the transition as though the asset player is a non-looping player that has reached its end
+					if (RelevantPlayer->IsLooping() && AnimTimeRemaining > 0.f)
+					{
+						if (const FDeltaTimeRecord* DeltaTimeRecord = RelevantPlayer->GetDeltaTimeRecord())
+						{
+							if (DeltaTimeRecord->IsPreviousValid())
+							{
+								// If the player's time difference is in the opposite direction of the time delta, then the player must have looped
+								const float DeltaTimeFromPlayer = RelevantPlayer->GetAccumulatedTime() - DeltaTimeRecord->GetPrevious();
+								const float DeltaTimeFromRecord = DeltaTimeRecord->Delta;
+								if (DeltaTimeFromPlayer * DeltaTimeFromRecord < 0.f)
+								{
+									AnimTimeRemaining = 0.f;
+								}
+							}
+						}
+					}
+
+					// Determine if we can automatically begin transitioning.
+					{
+						const FAnimationTransitionBetweenStates& TransitionInfo = GetTransitionInfo(TransitionRule.TransitionIndex);
+					
+						// For transitions that go to a conduit the user is not able to edit the transition's cross fade duration,
+						// therefore we force the cross fade duration to zero to ensure the transition is always triggerred upon reaching the end of the animation.
+						const float CrossfadeDuration = GetStateInfo(TransitionInfo.NextState).bIsAConduit ? 0.0f : TransitionInfo.CrossfadeDuration;
+
+						// Allow for used to determine the automatic transition trigger time or fallback to using cross fade duration.
+						const float TransitionTriggerTime = (TransitionRule.AutomaticRuleTriggerTime >= 0.0f) ? TransitionRule.AutomaticRuleTriggerTime : CrossfadeDuration;
+						CrossfadeTimeAdjustment = TransitionTriggerTime - AnimTimeRemaining;
+
+						// Trigger transition only if we have "CrossfadeTimeAdjustment" seconds left before reaching animation end boundary.
+						bCanEnterTransition = (CrossfadeTimeAdjustment >= 0.f);
+					}
 				}
 			}
+			
 			ResultNode->bCanEnterTransition = bCanEnterTransition;
 		}			
 		else 
