Comment warning for looping asset players and using automatic rule based transitions.

diff --git a/Engine/Source/Editor/AnimGraph/Private/AnimStateTransitionNode.cpp b/Engine/Source/Editor/AnimGraph/Private/AnimStateTransitionNode.cpp
--- a/Engine/Source/Editor/AnimGraph/Private/AnimStateTransitionNode.cpp
+++ b/Engine/Source/Editor/AnimGraph/Private/AnimStateTransitionNode.cpp
@@ -709,14 +709,14 @@ void UAnimStateTransitionNode::ValidateNodeDuringCompilation(class FCompilerResu
 								{
 									if (SequencePlayer->Node.IsLooping())
 									{
-										MessageLog.Warning(TEXT("Transition @@ is using an automatic transition rule but the source @@ is set as looping.  Please clear the 'Loop Animation' flag"), this, SequencePlayer);
+										// MessageLog.Warning(TEXT("Transition @@ is using an automatic transition rule but the source @@ is set as looping.  Please clear the 'Loop Animation' flag"), this, SequencePlayer);
 									}
 								}
 								else if (UAnimGraphNode_BlendSpacePlayer* BlendSpacePlayer = Cast<UAnimGraphNode_BlendSpacePlayer>(TestPin->LinkedTo[0]->GetOwningNode()))
 								{
 									if (BlendSpacePlayer->Node.IsLooping())
 									{
-										MessageLog.Warning(TEXT("Transition @@ is using an automatic transition rule but the source @@ is set as looping.  Please clear the 'Loop' flag"), this, BlendSpacePlayer);
+										// MessageLog.Warning(TEXT("Transition @@ is using an automatic transition rule but the source @@ is set as looping.  Please clear the 'Loop' flag"), this, BlendSpacePlayer);
 									}
 								}
 							}
