Player nodes in blendspace graphs are now added to GraphAssetPlayerInformation
# jira UE-188902

diff --git a/Engine/Source/Editor/AnimGraph/Private/AnimGraphNode_BlendSpaceGraphBase.cpp b/Engine/Source/Editor/AnimGraph/Private/AnimGraphNode_BlendSpaceGraphBase.cpp
--- a/Engine/Source/Editor/AnimGraph/Private/AnimGraphNode_BlendSpaceGraphBase.cpp
+++ b/Engine/Source/Editor/AnimGraph/Private/AnimGraphNode_BlendSpaceGraphBase.cpp
@@ -29,7 +29,9 @@
 #include "BlueprintEditor.h"
 #include "Animation/AnimSequence.h"
 #include "AnimBlueprintExtension_BlendSpaceGraph.h"
+#include "AnimGraphNode_RandomPlayer.h"
 #include "BlueprintNodeTemplateCache.h"
+#include "Animation/AnimLayerInterface.h"
 
 #define LOCTEXT_NAMESPACE "UAnimGraphNode_BlendSpaceGraphBase"
 
@@ -121,6 +123,35 @@ UAnimGraphNode_Base* UAnimGraphNode_BlendSpaceGraphBase::ExpandGraphAndProcessNo
 
 void UAnimGraphNode_BlendSpaceGraphBase::OnProcessDuringCompilation(IAnimBlueprintCompilationContext& InCompilationContext, IAnimBlueprintGeneratedClassCompiledData& OutCompiledData)
 {
+	auto ProcessGraph = [this, &OutCompiledData](UEdGraph* Graph, TArray<UAnimGraphNode_AssetPlayerBase*> AssetPlayerNodes, TArray<UAnimGraphNode_RandomPlayer*> AssetRandomPlayerPlayerNodes)
+	{
+		FString FunctionGraphName = Graph->GetName();
+		// Also make sure we do not process any empty stub graphs
+		if (!FunctionGraphName.Contains(ANIM_FUNC_DECORATOR))
+		{
+			if (Graph->Nodes.ContainsByPredicate([this, &OutCompiledData](UEdGraphNode* Node) { return Node && Node->NodeGuid == NodeGuid; }))
+			{
+				for (UAnimGraphNode_AssetPlayerBase* Node : AssetPlayerNodes)
+				{
+					if (int32* IndexPtr = OutCompiledData.GetAnimBlueprintDebugData().NodeGuidToIndexMap.Find(Node->NodeGuid))
+					{
+						FGraphAssetPlayerInformation& Info = OutCompiledData.GetGraphAssetPlayerInformation().FindOrAdd(FName(*FunctionGraphName));
+						Info.PlayerNodeIndices.AddUnique(*IndexPtr);
+					}
+				}
+
+				for (UAnimGraphNode_RandomPlayer* Node : AssetRandomPlayerPlayerNodes)
+				{
+					if (int32* IndexPtr = OutCompiledData.GetAnimBlueprintDebugData().NodeGuidToIndexMap.Find(Node->NodeGuid))
+					{
+						FGraphAssetPlayerInformation& Info = OutCompiledData.GetGraphAssetPlayerInformation().FindOrAdd(FName(*FunctionGraphName));
+						Info.PlayerNodeIndices.AddUnique(*IndexPtr);
+					}
+				}
+			}
+		}
+	};
+	
 	FStructProperty* NodeProperty = GetFNodeProperty();
 	FArrayProperty* PoseLinksProperty = CastFieldChecked<FArrayProperty>(NodeProperty->Struct->FindPropertyByName(GET_MEMBER_NAME_CHECKED(FAnimNode_BlendSpaceGraphBase, SamplePoseLinks)));
 
@@ -128,12 +159,38 @@ void UAnimGraphNode_BlendSpaceGraphBase::OnProcessDuringCompilation(IAnimBluepri
 	FAnimNode_BlendSpaceGraphBase* AnimNode = NodeProperty->ContainerPtrToValuePtr<FAnimNode_BlendSpaceGraphBase>(this);
 	AnimNode->SamplePoseLinks.SetNum(Graphs.Num());
 
+	TArray<UAnimGraphNode_AssetPlayerBase*> AssetPlayerNodes;
+	TArray<UAnimGraphNode_RandomPlayer*> AssetRandomPlayerPlayerNodes;
+
 	for(int32 PoseIndex = 0; PoseIndex < Graphs.Num(); ++PoseIndex)
 	{
 		UAnimationBlendSpaceSampleGraph* SampleGraph = CastChecked<UAnimationBlendSpaceSampleGraph>(Graphs[PoseIndex]);
 		UAnimGraphNode_Base* RootNode = ExpandGraphAndProcessNodes(SampleGraph, SampleGraph->ResultNode, InCompilationContext, OutCompiledData);
 
 		InCompilationContext.AddPoseLinkMappingRecord(FPoseLinkMappingRecord::MakeFromArrayEntry(this, RootNode, PoseLinksProperty, PoseIndex));
+
+		SampleGraph->GetNodesOfClass(AssetPlayerNodes);
+		SampleGraph->GetNodesOfClass(AssetRandomPlayerPlayerNodes);
+	}
+
+	// Append player nodes to the owning graph's list
+	UBlueprint* Blueprint = GetBlueprint();
+	for (UEdGraph* FunctionGraph : Blueprint->FunctionGraphs)
+	{
+		ProcessGraph(FunctionGraph, AssetPlayerNodes, AssetRandomPlayerPlayerNodes);
+	}
+
+	// Now do the same for implemented layer interfaces
+	for (FBPInterfaceDescription& InterfaceDesc : Blueprint->ImplementedInterfaces)
+	{
+		// Only process Anim Layer interfaces
+		if (InterfaceDesc.Interface->IsChildOf<UAnimLayerInterface>())
+		{
+			for (UEdGraph* Graph : InterfaceDesc.Graphs)
+			{
+				ProcessGraph(Graph, AssetPlayerNodes, AssetRandomPlayerPlayerNodes);
+			}
+		}
 	}
 }
 
