AnimNext params in Anim BPs

diff --git a/Engine/Plugins/Experimental/AnimNext/Source/AnimNextUncookedOnly/Private/Graph/AnimGraph/AnimBlueprintExtension_AnimNextParameters.cpp b/Engine/Plugins/Experimental/AnimNext/Source/AnimNextUncookedOnly/Private/Graph/AnimGraph/AnimBlueprintExtension_AnimNextParameters.cpp
new file mode 100644
--- /dev/null
+++ b/Engine/Plugins/Experimental/AnimNext/Source/AnimNextUncookedOnly/Private/Graph/AnimGraph/AnimBlueprintExtension_AnimNextParameters.cpp
@@ -0,0 +1,202 @@
+// Copyright Epic Games, Inc. All Rights Reserved.
+
+#include "Graph/AnimGraph/AnimBlueprintExtension_AnimNextParameters.h"
+#include "Graph/AnimGraph/AnimGraphNodeBinding_AnimNextParameters.h"
+#include "AnimBlueprintExtension_PropertyAccess.h"
+#include "IAnimBlueprintCompilationContext.h"
+#include "IAnimBlueprintCompilationBracketContext.h"
+#include "Graph/AnimGraph/AnimNodeExposedValueHandler_AnimNextParameters.h"
+#include "Features/IModularFeatures.h"
+#include "IPropertyAccessEditor.h"
+
+void UAnimBlueprintExtension_AnimNextParameters::ProcessNodePins(UAnimGraphNode_Base* InNode, IAnimBlueprintCompilationContext& InCompilationContext, IAnimBlueprintGeneratedClassCompiledData& OutCompiledData)
+{
+	if (UAnimGraphNodeBinding_AnimNextParameters* Binding = Cast<UAnimGraphNodeBinding_AnimNextParameters>(InNode->GetMutableBinding()))
+	{
+		if(Binding->PropertyBindings.Num() > 0)
+		{
+			FStructProperty* NodeProperty = CastFieldChecked<FStructProperty>(InCompilationContext.GetAllocatedPropertiesByNode().FindChecked(InNode));
+
+			BindingRecordIndexMap.Add(InNode, BindingRecords.Num());
+			FBindingRecord& NewRecord = BindingRecords.AddDefaulted_GetRef();
+			NewRecord.AnimGraphNode = InNode;
+			NewRecord.NodeVariableProperty = NodeProperty;
+
+			for (TPair<FName, FAnimNextAnimGraphNodeParameterBinding>& PropertyBindingPair : Binding->PropertyBindings)
+			{
+				FName BindingName = PropertyBindingPair.Key;
+				FAnimNextAnimGraphNodeParameterBinding& ParameterBinding = PropertyBindingPair.Value;
+
+				// for array properties we need to account for the extra FName number 
+				FName ComparisonName = BindingName;
+				ComparisonName.SetNumber(0);
+
+				if (FProperty* Property = FindFProperty<FProperty>(NewRecord.NodeVariableProperty->Struct, ComparisonName))
+				{
+					NewRecord.NodeVariableProperty = NodeProperty;
+					NewRecord.bServicesNodeProperties = true;
+					NewRecord.RegisterPropertyBinding(Property, ParameterBinding);
+				}
+				else if (FProperty* ClassProperty = FindFProperty<FProperty>(InCompilationContext.GetBlueprint()->SkeletonGeneratedClass, ComparisonName))
+				{
+					NewRecord.NodeVariableProperty = NodeProperty;
+					NewRecord.bServicesInstanceProperties = true;
+					NewRecord.RegisterPropertyBinding(ClassProperty, ParameterBinding);
+				}
+				else
+				{
+					InCompilationContext.GetMessageLog().Warning(*FString::Printf(TEXT("ICE: @@ Failed to find a property '%s'"), *ComparisonName.ToString()), InNode);
+				}
+			}
+		}
+	}
+}
+
+void UAnimBlueprintExtension_AnimNextParameters::HandleStartCompilingClass(const UClass* InClass, IAnimBlueprintCompilationBracketContext& InCompilationContext, IAnimBlueprintGeneratedClassCompiledData& OutCompiledData)
+{
+	BindingRecords.Empty();
+	BindingRecordIndexMap.Empty();
+
+	UAnimBlueprintExtension_PropertyAccess* PropertyAccessExtension = UAnimBlueprintExtension::GetExtension<UAnimBlueprintExtension_PropertyAccess>(GetAnimBlueprint());
+	if (PropertyAccessExtension)
+	{
+		PreLibraryCompiledDelegateHandle = PropertyAccessExtension->OnPreLibraryCompiled().AddLambda([this, PropertyAccessExtension, InClass]()
+		{
+			if (IModularFeatures::Get().IsModularFeatureAvailable("PropertyAccessEditor"))
+			{
+				IPropertyAccessEditor& PropertyAccessEditor = IModularFeatures::Get().GetModularFeature<IPropertyAccessEditor>("PropertyAccessEditor");
+
+				// Build the classes property access library before the library is compiled
+				for (FBindingRecord& BindingRecord : BindingRecords)
+				{
+					for (TPair<FName, FSinglePropertyRecord>& PropertyRecord : BindingRecord.ServicedProperties)
+					{
+						PropertyRecord.Value.LibraryHandle = PropertyAccessExtension->AddAccess(PropertyRecord.Value.PropertyPath, BindingRecord.AnimGraphNode);
+					}
+				}
+			}
+
+			PropertyAccessExtension->OnPreLibraryCompiled().Remove(PreLibraryCompiledDelegateHandle);
+		});
+
+		PostLibraryCompiledDelegateHandle = PropertyAccessExtension->OnPostLibraryCompiled().AddLambda([this, PropertyAccessExtension, InClass](IAnimBlueprintCompilationBracketContext& InCompilationContext, IAnimBlueprintGeneratedClassCompiledData& OutCompiledData)
+		{
+			// Get access index for runtime usage
+			for (FBindingRecord& BindingRecord : BindingRecords)
+			{
+				UAnimGraphNode_Base* OriginalNode = Cast<UAnimGraphNode_Base>(InCompilationContext.GetMessageLog().FindSourceObject(BindingRecord.AnimGraphNode));
+
+				for (TPair<FName, FSinglePropertyRecord>& PropertyRecord : BindingRecord.ServicedProperties)
+				{
+					PropertyRecord.Value.LibraryCompiledHandle = PropertyAccessExtension->GetCompiledHandle(PropertyRecord.Value.LibraryHandle);
+					PropertyRecord.Value.LibraryCompiledAccessType = PropertyAccessExtension->GetCompiledHandleAccessType(PropertyRecord.Value.LibraryHandle);
+				}
+			}
+
+			PatchEvaluationHandlers(InClass, InCompilationContext, OutCompiledData);
+
+			PropertyAccessExtension->OnPostLibraryCompiled().Remove(PostLibraryCompiledDelegateHandle);
+		});
+	}
+}
+
+void UAnimBlueprintExtension_AnimNextParameters::PatchEvaluationHandlers(const UClass* InClass, IAnimBlueprintCompilationBracketContext& InCompilationContext, IAnimBlueprintGeneratedClassCompiledData& OutCompiledData)
+{
+	for (FBindingRecord& BindingRecord : BindingRecords)
+	{
+		// Indices here are in reverse order with respect to iterated properties as properties are prepended to the linked list when they are added
+		if (const int32* NodePropertyIndex = InCompilationContext.GetAllocatedAnimNodeIndices().Find(BindingRecord.AnimGraphNode))
+		{
+			FStructProperty* HandlerProperty = CastFieldChecked<FStructProperty>(InCompilationContext.GetAllocatedHandlerPropertiesByNode().FindChecked(BindingRecord.AnimGraphNode));
+			if (HandlerProperty->Struct == FAnimNodeExposedValueHandler_AnimNextParameters::StaticStruct())
+			{
+				UAnimBlueprintGeneratedClass* AnimClass = Cast<UAnimBlueprintGeneratedClass>(const_cast<UClass*>(InClass));
+				check(HandlerProperty->GetOwner<UScriptStruct>() == AnimClass->GetSparseClassDataStruct());
+				void* ConstantNodeData = const_cast<void*>(AnimClass->GetSparseClassData(EGetSparseClassDataMethod::ReturnIfNull));
+				check(ConstantNodeData);
+				FAnimNodeExposedValueHandler_AnimNextParameters* Handler = HandlerProperty->ContainerPtrToValuePtr<FAnimNodeExposedValueHandler_AnimNextParameters>(ConstantNodeData);
+
+				Handler->Entries.Empty();
+
+				for (TPair<FName, FSinglePropertyRecord>& PropertyRecord : BindingRecord.ServicedProperties)
+				{
+					if (PropertyRecord.Value.LibraryCompiledHandle.IsValid())
+					{
+						FAnimNodeExposedValueHandler_AnimNextParameters_Entry& NewEntry = Handler->Entries.AddDefaulted_GetRef();
+						NewEntry.ParameterName = PropertyRecord.Value.ParameterName;
+						NewEntry.AccessIndex = PropertyRecord.Value.LibraryCompiledHandle.GetId();
+						NewEntry.PropertyParamType = UE::AnimNext::FParamTypeHandle::FromProperty(PropertyRecord.Value.Property).GetType();
+						NewEntry.AccessType = PropertyRecord.Value.LibraryCompiledAccessType;
+					}
+				}
+			}
+		}
+	}
+}
+
+void UAnimBlueprintExtension_AnimNextParameters::FBindingRecord::RegisterPropertyBinding(FProperty* InProperty, const FAnimNextAnimGraphNodeParameterBinding& InBinding)
+{
+	FSinglePropertyRecord& PropertyRecord = ServicedProperties.FindOrAdd(InProperty->GetFName());
+
+	// Prepend the destination property with the node's member property if the property is not on a UClass
+	if (Cast<UClass>(InProperty->Owner.ToUObject()) == nullptr)
+	{
+		PropertyRecord.PropertyPath.Add(NodeVariableProperty->GetName());
+	}
+
+	if (InBinding.ArrayIndex != INDEX_NONE)
+	{
+		PropertyRecord.PropertyPath.Add(FString::Printf(TEXT("%s[%d]"), *InProperty->GetName(), InBinding.ArrayIndex));
+	}
+	else
+	{
+		PropertyRecord.PropertyPath.Add(InProperty->GetName());
+	}
+
+	PropertyRecord.ParameterName = InBinding.ParameterName;
+	PropertyRecord.Property = InProperty;
+	PropertyRecord.ArrayIndex = InBinding.ArrayIndex;
+}
+
+void UAnimBlueprintExtension_AnimNextParameters::RedirectPropertyPaths(IAnimBlueprintCompilationContext& InCompilationContext, UAnimGraphNode_Base* InNode)
+{
+	if(const int32* IndexPtr = BindingRecordIndexMap.Find(InNode))
+	{
+		FBindingRecord& BindingRecord = BindingRecords[*IndexPtr];
+
+		const FStructProperty* MutableDataProperty = InCompilationContext.GetMutableDataProperty();
+		if (InCompilationContext.IsAnimGraphNodeFolded(InNode) && MutableDataProperty != nullptr)
+		{
+			for (TPair<FName, FSinglePropertyRecord>& SinglePropertyPair : BindingRecord.ServicedProperties)
+			{
+				if (const IAnimBlueprintCompilationContext::FFoldedPropertyRecord* FoldedPropertyRecord = InCompilationContext.GetFoldedPropertyRecord(InNode, SinglePropertyPair.Key))
+				{
+					if (SinglePropertyPair.Value.PropertyPath.Num() > 1)
+					{
+						// If this record writes to the node, switch it to the mutable data's property instead
+						if (SinglePropertyPair.Value.PropertyPath[0] == BindingRecord.NodeVariableProperty->GetName())
+						{
+							SinglePropertyPair.Value.PropertyPath[0] = MutableDataProperty->GetName();
+
+							FString PropertyPathTail = SinglePropertyPair.Value.PropertyPath[1];
+							FString PropertyPathWithoutArray = PropertyPathTail;
+							FString ArrayIndex;
+							int32 ArrayDelim = INDEX_NONE;
+							if (PropertyPathTail.FindChar(TEXT('['), ArrayDelim))
+							{
+								PropertyPathWithoutArray = PropertyPathTail.Left(ArrayDelim);
+								ArrayIndex = PropertyPathTail.RightChop(ArrayDelim);
+							}
+
+							// Switch the destination property from the node's property to the generated one
+							if (PropertyPathWithoutArray == FoldedPropertyRecord->Property->GetName())
+							{
+								SinglePropertyPair.Value.PropertyPath[1] = FoldedPropertyRecord->GeneratedProperty->GetName() + ArrayIndex;
+							}
+						}
+					}
+				}
+			}
+		}
+	}
+}
\ No newline at end of file
diff --git a/Engine/Plugins/Experimental/AnimNext/Source/AnimNextUncookedOnly/Private/Graph/AnimGraph/AnimGraphNodeBinding_AnimNextParameters.cpp b/Engine/Plugins/Experimental/AnimNext/Source/AnimNextUncookedOnly/Private/Graph/AnimGraph/AnimGraphNodeBinding_AnimNextParameters.cpp
new file mode 100644
--- /dev/null
+++ b/Engine/Plugins/Experimental/AnimNext/Source/AnimNextUncookedOnly/Private/Graph/AnimGraph/AnimGraphNodeBinding_AnimNextParameters.cpp
@@ -0,0 +1,574 @@
+// Copyright Epic Games, Inc. All Rights Reserved.
+
+#include "Graph/AnimGraph/AnimGraphNodeBinding_AnimNextParameters.h"
+#include "Graph/AnimGraph/AnimNodeExposedValueHandler_AnimNextParameters.h"
+#include "Graph/AnimGraph/AnimBlueprintExtension_AnimNextParameters.h"
+#include "Widgets/SNullWidget.h"
+#include "UncookedOnlyUtils.h"
+#include "AnimBlueprintExtension_Base.h"
+
+#if WITH_EDITOR
+#include "AnimNextEditorModule.h"
+#include "Param/ParameterPickerArgs.h"
+#include "Widgets/Input/SComboButton.h"
+#include "Widgets/Layout/SBox.h"
+#include "Framework/Application/SlateApplication.h"
+#include "Param/ParamUtils.h"
+#include "Param/ParamCompatibility.h"
+#include "FindInBlueprintManager.h"
+#include "Kismet2/BlueprintEditorUtils.h"
+#include "ScopedTransaction.h"
+#include "Widgets/Layout/SSpacer.h"
+#include "Framework/MultiBox/MultiBoxBuilder.h"
+#endif
+
+#define LOCTEXT_NAMESPACE "AnimGraphNodeBinding_AnimNextParameters"
+
+UScriptStruct* UAnimGraphNodeBinding_AnimNextParameters::GetAnimNodeHandlerStruct() const
+{
+	return FAnimNodeExposedValueHandler_AnimNextParameters::StaticStruct();
+}
+
+void UAnimGraphNodeBinding_AnimNextParameters::OnExpandNode(IAnimBlueprintCompilationContext& InCompilationContext, UAnimGraphNode_Base* InNode, UEdGraph* InSourceGraph)
+{
+	if (UAnimGraphNode_Base* Node = GetTypedOuter<UAnimGraphNode_Base>())
+	{
+		if (UAnimBlueprint* AnimBlueprint = Node->GetTypedOuter<UAnimBlueprint>())
+		{
+			UAnimBlueprintExtension_AnimNextParameters* AnimNextExtension = UAnimBlueprintExtension::GetExtension<UAnimBlueprintExtension_AnimNextParameters>(AnimBlueprint);
+			AnimNextExtension->RedirectPropertyPaths(InCompilationContext, InNode);
+		}
+	}
+}
+
+bool UAnimGraphNodeBinding_AnimNextParameters::HasBinding(FName InBindingName, bool bCheckArrayIndexName) const
+{
+	if (bCheckArrayIndexName)
+	{
+		// This checks for the 'pin name' of an array element
+		return PropertyBindings.Find(InBindingName) != nullptr;
+	}
+	else
+	{
+		// Comparison without name index to deal with arrays
+		// This checks for any pin of an array element and the array itself
+		const FName ComparisonName = FName(InBindingName, 0);
+
+		for (const TPair<FName, FAnimNextAnimGraphNodeParameterBinding>& BindingPair : PropertyBindings)
+		{
+			if (ComparisonName == FName(BindingPair.Key, 0))
+			{
+				return true;
+			}
+		}
+
+		return false;
+	}
+}
+
+void UAnimGraphNodeBinding_AnimNextParameters::RemoveBindings(FName InBindingName)
+{
+	const FName ComparisonName = FName(InBindingName, 0);
+
+	for (auto Iter = PropertyBindings.CreateIterator(); Iter; ++Iter)
+	{
+		if (ComparisonName == FName(Iter.Key(), 0))
+		{
+			Iter.RemoveCurrent();
+		}
+	}
+}
+
+void UAnimGraphNodeBinding_AnimNextParameters::AddPinSearchMetaDataInfo(const UEdGraphPin* InPin, FName InBindingName, TArray<struct FSearchTagDataPair>& OutTaggedMetaData) const
+{
+	if (const FAnimNextAnimGraphNodeParameterBinding* BindingInfo = PropertyBindings.Find(InBindingName))
+	{
+		OutTaggedMetaData.Add(FSearchTagDataPair(FText::FromString(TEXT("Parameter Binding")), BindingInfo->CachedParameterNameText));
+	}
+}
+
+void UAnimGraphNodeBinding_AnimNextParameters::UpdateBindingNames(TFunctionRef<FString(const FString& InOldName)> InModifierFunction)
+{
+	if (UAnimGraphNode_Base* Node = GetTypedOuter<UAnimGraphNode_Base>())
+	{
+		TMap<FName, FAnimNextAnimGraphNodeParameterBinding> NewBindings;
+
+		for (const TPair<FName, FAnimNextAnimGraphNodeParameterBinding>& BindingPair : PropertyBindings)
+		{
+			FString BindingNameString = BindingPair.Key.ToString();
+
+			FString NewNameString = InModifierFunction(BindingNameString);
+
+			if (NewNameString != BindingNameString)
+			{
+				FName NewName = *NewNameString;
+				FAnimNextAnimGraphNodeParameterBinding NewBinding = BindingPair.Value;
+				NewBinding.BindingName = NewName;
+				NewBindings.Add(NewName, NewBinding);
+			}
+		}
+
+		if (NewBindings.Num() > 0)
+		{
+			PropertyBindings = NewBindings;
+		}
+	}
+}
+
+void UAnimGraphNodeBinding_AnimNextParameters::GetRequiredExtensions(TArray<TSubclassOf<UAnimBlueprintExtension>>& OutExtensions) const
+{
+	OutExtensions.Add(UAnimBlueprintExtension_AnimNextParameters::StaticClass());
+}
+
+void UAnimGraphNodeBinding_AnimNextParameters::ProcessDuringCompilation(IAnimBlueprintCompilationContext& InCompilationContext, IAnimBlueprintGeneratedClassCompiledData& OutCompiledData)
+{
+	if (UAnimGraphNode_Base* Node = GetTypedOuter<UAnimGraphNode_Base>())
+	{
+		if (UAnimBlueprint* AnimBlueprint = Node->GetTypedOuter<UAnimBlueprint>())
+		{
+			UAnimBlueprintExtension_AnimNextParameters* AnimNextExtension = UAnimBlueprintExtension::GetExtension<UAnimBlueprintExtension_AnimNextParameters>(AnimBlueprint);
+			// Record pose pins for later patchup and gather pins that have an associated evaluation handler
+			AnimNextExtension->ProcessNodePins(Node, InCompilationContext, OutCompiledData);
+
+			// Call base handler to allow BP logic to be wired up as well as AnimNext parameters
+			UAnimBlueprintExtension_Base* BaseExtension = UAnimBlueprintExtension::GetExtension<UAnimBlueprintExtension_Base>(AnimBlueprint);
+			BaseExtension->ProcessNonPosePins(Node, InCompilationContext, OutCompiledData, UAnimBlueprintExtension_Base::EPinProcessingFlags::Constants | UAnimBlueprintExtension_Base::EPinProcessingFlags::BlueprintHandlers);
+		}
+	}
+}
+
+#if WITH_EDITOR
+TSharedRef<SWidget> UAnimGraphNodeBinding_AnimNextParameters::MakePropertyBindingWidget(const UAnimGraphNode_Base::FAnimPropertyBindingWidgetArgs& InArgs)
+{
+	using namespace UE::AnimNext;
+	using namespace UE::AnimNext::Editor;
+	using namespace UE::AnimNext::UncookedOnly;
+
+	enum class ECurrentValueType : int32
+	{
+		None,
+		Pin,
+		Binding,
+		Dynamic,
+		MultipleValues,
+	};
+
+	UAnimGraphNode_Base* FirstAnimGraphNode = InArgs.Nodes[0];
+	const bool bMultiSelect = InArgs.Nodes.Num() > 1;
+
+	if (FirstAnimGraphNode->HasValidBlueprint())
+	{
+		UBlueprint* Blueprint = FirstAnimGraphNode->GetAnimBlueprint();
+
+		int32 PinArrayIndex = InArgs.PinName.GetNumber() - 1;
+		const bool bIsArrayOrArrayElement = InArgs.PinProperty->IsA<FArrayProperty>();
+		const bool bIsArrayElement = bIsArrayOrArrayElement && PinArrayIndex != INDEX_NONE && InArgs.bPropertyIsOnFNode;
+		const bool bIsArray = bIsArrayOrArrayElement && PinArrayIndex == INDEX_NONE && InArgs.bPropertyIsOnFNode;
+
+		FProperty* PropertyToBindTo = bIsArrayElement ? CastField<FArrayProperty>(InArgs.PinProperty)->Inner : InArgs.PinProperty;
+
+		// Properties could potentially be removed underneath this widget, so keep a TFieldPath reference to them
+		TFieldPath<FProperty> BindingPropertyPath(PropertyToBindTo);
+		TFieldPath<FProperty> PinPropertyPath(InArgs.PinProperty);
+
+		return SNew(SComboButton)
+			.ToolTipText_Lambda([InArgs]()
+			{
+				ECurrentValueType CurrentValueType = ECurrentValueType::None;
+
+				const FText MultipleValues = LOCTEXT("MultipleValuesToolTip", "Bindings Have Multiple Values");
+				const FText ExposedAsPin = LOCTEXT("ExposedAsPinToolTip", "Exposed As a Pin on the Node");
+				const FText BindValue = LOCTEXT("BindValueToolTip", "Bind This Value");
+				const FText DynamicValue = LOCTEXT("DynamicValueToolTip", "Dynamic value that can be set externally");
+				FText CurrentValue;
+
+				auto SetAssignValue = [&CurrentValueType, &CurrentValue, &MultipleValues](const FText& InValue, ECurrentValueType InType)
+				{
+					if (CurrentValueType != ECurrentValueType::MultipleValues)
+					{
+						if (CurrentValueType == ECurrentValueType::None)
+						{
+							CurrentValueType = InType;
+							CurrentValue = InValue;
+						}
+						else if (CurrentValueType == InType)
+						{
+							if (!CurrentValue.EqualTo(InValue))
+							{
+								CurrentValueType = ECurrentValueType::MultipleValues;
+								CurrentValue = MultipleValues;
+							}
+						}
+						else
+						{
+							CurrentValueType = ECurrentValueType::MultipleValues;
+							CurrentValue = MultipleValues;
+						}
+					}
+				};
+
+				const FName ComparisonName = FName(InArgs.PinName, 0);
+
+				for (UAnimGraphNode_Base* AnimGraphNode : InArgs.Nodes)
+				{
+					if(const UAnimGraphNodeBinding_AnimNextParameters* Binding = Cast<UAnimGraphNodeBinding_AnimNextParameters>(AnimGraphNode->GetBinding()))
+					{
+						if (const FAnimNextAnimGraphNodeParameterBinding* BindingPtr = Binding->PropertyBindings.Find(InArgs.BindingName))
+						{
+							if (BindingPtr->CachedParameterNameText.IsEmpty())
+							{
+								SetAssignValue(BindValue, ECurrentValueType::Binding);
+							}
+							else
+							{
+								SetAssignValue(FText::Format(LOCTEXT("BindingToolTipFormat", "Pin is bound to parameter '{0}'"), BindingPtr->CachedParameterNameText), ECurrentValueType::Binding);
+							}
+						}
+						else if (AnimGraphNode->AlwaysDynamicProperties.Find(ComparisonName))
+						{
+							SetAssignValue(DynamicValue, ECurrentValueType::Dynamic);
+						}
+						else
+						{
+							TArrayView<FOptionalPinFromProperty> OptionalPins;
+							InArgs.OnGetOptionalPins.ExecuteIfBound(AnimGraphNode, OptionalPins);
+							if (OptionalPins.IsValidIndex(InArgs.OptionalPinIndex) && OptionalPins[InArgs.OptionalPinIndex].bShowPin)
+							{
+								SetAssignValue(InArgs.bOnGraphNode ? BindValue : ExposedAsPin, ECurrentValueType::Pin);
+							}
+							else
+							{
+								SetAssignValue(BindValue, ECurrentValueType::None);
+							}
+						}
+					}
+				}
+
+				return CurrentValue;
+			})
+			.ButtonContent()
+			[
+				SNew(STextBlock)
+				.TextStyle(&FCoreStyle::Get().GetWidgetStyle<FTextBlockStyle>("SmallText"))
+				.Text_Lambda([InArgs]()
+				{
+					ECurrentValueType CurrentValueType = ECurrentValueType::None;
+
+					const FText MultipleValues = LOCTEXT("MultipleValuesLabel", "Multiple Values");
+					const FText Bind = LOCTEXT("BindLabel", "Param");
+					const FText ExposedAsPin = LOCTEXT("ExposedAsPinLabel", "Pin");
+					const FText Dynamic = LOCTEXT("DynamicLabel", "Dynamic");
+					FText CurrentValue = Bind;
+
+					auto SetAssignValue = [&CurrentValueType, &CurrentValue, &MultipleValues](const FText& InValue, ECurrentValueType InType)
+					{
+						if (CurrentValueType != ECurrentValueType::MultipleValues)
+						{
+							if (CurrentValueType == ECurrentValueType::None)
+							{
+								CurrentValueType = InType;
+								CurrentValue = InValue;
+							}
+							else if (CurrentValueType == InType)
+							{
+								if (!CurrentValue.EqualTo(InValue))
+								{
+									CurrentValueType = ECurrentValueType::MultipleValues;
+									CurrentValue = MultipleValues;
+								}
+							}
+							else
+							{
+								CurrentValueType = ECurrentValueType::MultipleValues;
+								CurrentValue = MultipleValues;
+							}
+						}
+					};
+
+					const FName ComparisonName = FName(InArgs.PinName, 0);
+
+					for (UAnimGraphNode_Base* AnimGraphNode : InArgs.Nodes)
+					{
+						if(const UAnimGraphNodeBinding_AnimNextParameters* Binding = Cast<UAnimGraphNodeBinding_AnimNextParameters>(AnimGraphNode->GetBinding()))
+						{
+							if (const FAnimNextAnimGraphNodeParameterBinding* BindingPtr = Binding->PropertyBindings.Find(InArgs.BindingName))
+							{
+								SetAssignValue(BindingPtr->CachedParameterNameText, ECurrentValueType::Binding);
+							}
+							else if (AnimGraphNode->AlwaysDynamicProperties.Find(ComparisonName))
+							{
+								SetAssignValue(Dynamic, ECurrentValueType::Dynamic);
+							}
+							else
+							{
+								TArrayView<FOptionalPinFromProperty> OptionalPins;
+								InArgs.OnGetOptionalPins.ExecuteIfBound(AnimGraphNode, OptionalPins);
+								if (OptionalPins.IsValidIndex(InArgs.OptionalPinIndex) && OptionalPins[InArgs.OptionalPinIndex].bShowPin)
+								{
+									SetAssignValue(InArgs.bOnGraphNode ? Bind : ExposedAsPin, ECurrentValueType::Pin);
+								}
+								else
+								{
+									SetAssignValue(Bind, ECurrentValueType::None);
+								}
+							}
+						}
+					}
+
+					return CurrentValue;
+				})
+			]
+			.OnGetMenuContent(FOnGetContent::CreateLambda([InArgs, PropertyToBindTo, BindingPropertyPath, PinArrayIndex, bIsArray, bIsArrayElement, Blueprint]()
+			{
+				FMenuBuilder MenuBuilder(true, nullptr);
+				MenuBuilder.BeginSection("Pins", LOCTEXT("Pin", "Pin"));
+				{
+					auto ExposeAsPin = [InArgs, Blueprint]()
+					{
+						bool bHasBinding = false;
+
+						for (UAnimGraphNode_Base* AnimGraphNode : InArgs.Nodes)
+						{
+							if (AnimGraphNode->HasBinding(InArgs.BindingName))
+							{
+								bHasBinding = true;
+								break;
+							}
+						}
+
+						{
+							FScopedTransaction Transaction(LOCTEXT("PinExposure", "Pin Exposure"));
+
+							// Switching from non-pin to pin, remove any bindings
+							for (UAnimGraphNode_Base* AnimGraphNode : InArgs.Nodes)
+							{
+								if(UAnimGraphNodeBinding_AnimNextParameters* Binding = CastChecked<UAnimGraphNodeBinding_AnimNextParameters>(AnimGraphNode->GetMutableBinding()))
+								{
+									Binding->Modify();
+									AnimGraphNode->Modify();
+
+									TArrayView<FOptionalPinFromProperty> OptionalPins;
+									InArgs.OnGetOptionalPins.ExecuteIfBound(AnimGraphNode, OptionalPins);
+									const bool bVisible = OptionalPins.IsValidIndex(InArgs.OptionalPinIndex) && OptionalPins[InArgs.OptionalPinIndex].bShowPin;
+									InArgs.OnSetPinVisibility.ExecuteIfBound(AnimGraphNode, !bVisible || bHasBinding, InArgs.OptionalPinIndex);
+
+									// Remove all bindings that match the property, array or array elements
+									const FName ComparisonName = FName(InArgs.BindingName, 0);
+									for (auto It = Binding->PropertyBindings.CreateIterator(); It; ++It)
+									{
+										if (ComparisonName == FName(It.Key(), 0))
+										{
+											It.RemoveCurrent();
+										}
+									}
+								}
+							}
+
+							FBlueprintEditorUtils::MarkBlueprintAsModified(Blueprint);
+						}
+					};
+
+					auto GetExposedAsPinCheckState = [InArgs]()
+					{
+						bool bPinShown = false;
+						bool bHasBinding = false;
+
+						for (UAnimGraphNode_Base* AnimGraphNode : InArgs.Nodes)
+						{
+							if (AnimGraphNode->HasBinding(InArgs.BindingName))
+							{
+								bHasBinding = true;
+								break;
+							}
+
+							TArrayView<FOptionalPinFromProperty> OptionalPins;
+							InArgs.OnGetOptionalPins.ExecuteIfBound(AnimGraphNode, OptionalPins);
+							bPinShown |= OptionalPins.IsValidIndex(InArgs.OptionalPinIndex) && OptionalPins[InArgs.OptionalPinIndex].bShowPin;
+						}
+
+						// Pins are exposed if we have a binding or not, so treat as unchecked only if we have
+						// no binding
+						return bPinShown && !bHasBinding ? ECheckBoxState::Checked : ECheckBoxState::Unchecked;
+					};
+
+					MenuBuilder.AddMenuEntry(
+						LOCTEXT("ExposeAsPin", "Expose As Pin"),
+						LOCTEXT("ExposeAsPinTooltip", "Show/hide this property as a pin on the node"),
+						FSlateIcon(FAppStyle::GetAppStyleSetName(), "GraphEditor.PinIcon"),
+						FUIAction(
+							FExecuteAction::CreateLambda(ExposeAsPin),
+							FCanExecuteAction(),
+							FGetActionCheckState::CreateLambda(GetExposedAsPinCheckState)
+						),
+						NAME_None,
+						EUserInterfaceActionType::Check
+					);
+
+					if (InArgs.bPropertyIsOnFNode)
+					{
+						auto MakeDynamic = [InArgs, Blueprint]()
+						{
+							// Comparison without name index to deal with arrays
+							const FName ComparisonName = FName(InArgs.PinName, 0);
+
+							bool bIsAlwaysDynamic = false;
+							for (UAnimGraphNode_Base* AnimGraphNode : InArgs.Nodes)
+							{
+								if (AnimGraphNode->AlwaysDynamicProperties.Contains(ComparisonName))
+								{
+									bIsAlwaysDynamic = true;
+									break;
+								}
+							}
+
+							{
+								FScopedTransaction Transaction(LOCTEXT("AlwaysDynamic", "Always Dynamic"));
+
+								for (UAnimGraphNode_Base* AnimGraphNode : InArgs.Nodes)
+								{
+									AnimGraphNode->Modify();
+
+									if (bIsAlwaysDynamic)
+									{
+										AnimGraphNode->AlwaysDynamicProperties.Remove(ComparisonName);
+									}
+									else
+									{
+										AnimGraphNode->AlwaysDynamicProperties.Add(ComparisonName);
+									}
+								}
+
+								FBlueprintEditorUtils::MarkBlueprintAsModified(Blueprint);
+							}
+						};
+
+						auto GetDynamicCheckState = [InArgs]()
+						{
+							// Comparison without name index to deal with arrays
+							const FName ComparisonName = FName(InArgs.PinName, 0);
+
+							bool bIsAlwaysDynamic = false;
+							for (UAnimGraphNode_Base* AnimGraphNode : InArgs.Nodes)
+							{
+								for (const FName& AlwaysDynamicPropertyName : AnimGraphNode->AlwaysDynamicProperties)
+								{
+									if (ComparisonName == FName(AlwaysDynamicPropertyName, 0))
+									{
+										bIsAlwaysDynamic = true;
+										break;
+									}
+								}
+							}
+
+							return bIsAlwaysDynamic ? ECheckBoxState::Checked : ECheckBoxState::Unchecked;
+						};
+
+						MenuBuilder.AddMenuEntry(
+							LOCTEXT("DynamicValue", "Dynamic Value"),
+							LOCTEXT("DynamicValueTooltip", "Flag this value as dynamic. This way it can be set from functions even when not exposed as a pin."),
+							FSlateIcon(),
+							FUIAction(
+								FExecuteAction::CreateLambda(MakeDynamic),
+								FCanExecuteAction(),
+								FGetActionCheckState::CreateLambda(GetDynamicCheckState)
+							),
+							NAME_None,
+							EUserInterfaceActionType::Check
+						);
+					}
+				}
+				MenuBuilder.EndSection();
+
+				IModule& EditorModule = FModuleManager::Get().LoadModuleChecked<IModule>("AnimNextEditor");
+				FParameterPickerArgs Args;
+				Args.OnParameterPicked = FOnParameterPicked::CreateLambda([InArgs, BindingPropertyPath, PinArrayIndex, bIsArray, bIsArrayElement, Blueprint](const FParameterBindingReference& InBinding)
+				{
+					FSlateApplication::Get().DismissAllMenus();
+
+					for (UAnimGraphNode_Base* AnimGraphNode : InArgs.Nodes)
+					{
+						if (UAnimGraphNodeBinding_AnimNextParameters* Binding = Cast<UAnimGraphNodeBinding_AnimNextParameters>(AnimGraphNode->GetMutableBinding()))
+						{
+							Binding->Modify();
+
+							// Reset to default so that references are not preserved
+							FProperty* BindingProperty = BindingPropertyPath.Get();
+							if (BindingProperty && BindingProperty->GetOwner<UStruct>() && AnimGraphNode->GetFNodeType()->IsChildOf(BindingProperty->GetOwner<UStruct>()) && BindingProperty->IsA<FObjectPropertyBase>())
+							{
+								void* PropertyAddress = BindingProperty->ContainerPtrToValuePtr<void>(AnimGraphNode->GetFNode());
+								BindingProperty->InitializeValue(PropertyAddress);
+							}
+
+							// Pins are exposed if we have a binding or not - and after running this we do.
+							InArgs.OnSetPinVisibility.ExecuteIfBound(AnimGraphNode, true, InArgs.OptionalPinIndex);
+
+							// Need to break all pin links now we have a binding
+							if (UEdGraphPin* Pin = AnimGraphNode->FindPin(InArgs.PinName))
+							{
+								Pin->BreakAllPinLinks();
+							}
+
+							if (bIsArray)
+							{
+								// Remove bindings for array elements if this is an array
+								FName ComparisonName(InArgs.BindingName, 0);
+								for (auto Iter = Binding->PropertyBindings.CreateIterator(); Iter; ++Iter)
+								{
+									if (ComparisonName == FName(Iter.Key(), 0))
+									{
+										Iter.RemoveCurrent();
+									}
+								}
+							}
+							else if (bIsArrayElement)
+							{
+								// If we are an array element, remove only whole-array bindings
+								FName ComparisonName(InArgs.BindingName, 0);
+								for (auto Iter = Binding->PropertyBindings.CreateIterator(); Iter; ++Iter)
+								{
+									if (ComparisonName == Iter.Key())
+									{
+										Iter.RemoveCurrent();
+									}
+								}
+							}
+
+							FAnimNextAnimGraphNodeParameterBinding NewBinding;
+							NewBinding.BindingName = InArgs.BindingName;
+							NewBinding.ParameterName = InBinding.Parameter;
+							NewBinding.CachedParameterNameText = UncookedOnly::FUtils::GetParameterDisplayNameText(NewBinding.ParameterName);
+							NewBinding.ArrayIndex = PinArrayIndex;
+							Binding->PropertyBindings.Add(InArgs.BindingName, NewBinding);
+						}
+					}
+
+					FBlueprintEditorUtils::MarkBlueprintAsModified(Blueprint);
+				});
+				Args.OnFilterParameterType = FOnFilterParameterType::CreateLambda([ParamTypeHandle = FParamTypeHandle::FromProperty(PropertyToBindTo)](const FAnimNextParamType& InParamType)
+				{
+					return (FParamUtils::GetCompatibility(ParamTypeHandle, InParamType.GetHandle()).IsCompatibleWithDataLoss()) ? EFilterParameterResult::Include : EFilterParameterResult::Exclude;
+				});
+				Args.bMultiSelect = false;
+				Args.bShowLibraries = false;
+				Args.bShowBlocks = false;
+			
+				MenuBuilder.AddWidget(
+					SNew(SBox)
+					.WidthOverride(300.0f)
+					.HeightOverride(400.0f)
+					[
+						EditorModule.CreateParameterPicker(Args)
+					],
+					FText::GetEmpty(), true, false);
+
+				return MenuBuilder.MakeWidget();
+			}));
+	}
+	else
+	{
+		return SNew(SSpacer);
+	}
+}
+#endif
+
+#undef LOCTEXT_NAMESPACE
\ No newline at end of file
diff --git a/Engine/Plugins/Experimental/AnimNext/Source/AnimNextUncookedOnly/Private/Graph/AnimGraph/AnimGraphNode_AnimNextParameters.cpp b/Engine/Plugins/Experimental/AnimNext/Source/AnimNextUncookedOnly/Private/Graph/AnimGraph/AnimGraphNode_AnimNextParameters.cpp
new file mode 100644
--- /dev/null
+++ b/Engine/Plugins/Experimental/AnimNext/Source/AnimNextUncookedOnly/Private/Graph/AnimGraph/AnimGraphNode_AnimNextParameters.cpp
@@ -0,0 +1,22 @@
+// Copyright Epic Games, Inc. All Rights Reserved.
+
+#include "Graph/AnimGraph/AnimGraphNode_AnimNextParameters.h"
+
+#define LOCTEXT_NAMESPACE "AnimGraphNode_AnimNextParameters"
+
+FText UAnimGraphNode_AnimNextParameters::GetNodeTitle(ENodeTitleType::Type TitleType) const
+{
+	return LOCTEXT("NodeTitle", "AnimNext Parameter Block");
+}
+
+FText UAnimGraphNode_AnimNextParameters::GetTooltipText() const
+{
+	return LOCTEXT("NodeTooltip", "Pushes parameter blocks that can be used by all nodes leaf-wards");
+}
+
+FText UAnimGraphNode_AnimNextParameters::GetMenuCategory() const
+{
+	return LOCTEXT("Category", "AnimNext");
+}
+
+#undef LOCTEXT_NAMESPACE
\ No newline at end of file
diff --git a/Engine/Plugins/Experimental/AnimNext/Source/AnimNextUncookedOnly/Private/Graph/AnimNextGraph_EditorData.cpp b/Engine/Plugins/Experimental/AnimNext/Source/AnimNextUncookedOnly/Private/Graph/AnimNextGraph_EditorData.cpp
--- a/Engine/Plugins/Experimental/AnimNext/Source/AnimNextUncookedOnly/Private/Graph/AnimNextGraph_EditorData.cpp
+++ b/Engine/Plugins/Experimental/AnimNext/Source/AnimNextUncookedOnly/Private/Graph/AnimNextGraph_EditorData.cpp
@@ -6,7 +6,7 @@
 #include "Graph/AnimNextGraph.h"
 #include "Graph/AnimNextGraph_EdGraphSchema.h"
 #include "UncookedOnlyUtils.h"
-#include "Graph/GraphExecuteContext.h"
+#include "Graph/AnimNextExecuteContext.h"
 #include "Rigs/RigHierarchyPose.h"
 #include "RigVMModel/RigVMFunctionLibrary.h"
 #include "RigVMModel/RigVMNotifications.h"
@@ -25,7 +25,7 @@ UAnimNextGraph_EditorData::UAnimNextGraph_EditorData(const FObjectInitializer& O
 		RigVMClient.AddModel(TEXT("RigVMGraph"), false, &ObjectInitializer);
 		RigVMClient.GetOrCreateFunctionLibrary(false, &ObjectInitializer);
 	}
-	RigVMClient.SetExecuteContextStruct(FAnimNextGraphExecuteContext::StaticStruct());
+	RigVMClient.SetExecuteContextStruct(FAnimNextExecuteContext::StaticStruct());
 	
 	auto MakeEdGraph = [this, &ObjectInitializer](FName InName) -> UAnimNextGraph_EdGraph*
 	{
@@ -260,7 +260,7 @@ void UAnimNextGraph_EditorData::HandleRigVMGraphAdded(const FRigVMClient* InClie
 {
 	if(URigVMGraph* RigVMGraph = InClient->GetModel(InNodePath))
 	{
-		RigVMGraph->SetExecuteContextStruct(FAnimNextGraphExecuteContext::StaticStruct());
+		RigVMGraph->SetExecuteContextStruct(FAnimNextExecuteContext::StaticStruct());
 	}
 }
 
diff --git a/Engine/Plugins/Experimental/AnimNext/Source/AnimNextUncookedOnly/Private/Param/AnimNextParameterBlock_EditorData.cpp b/Engine/Plugins/Experimental/AnimNext/Source/AnimNextUncookedOnly/Private/Param/AnimNextParameterBlock_EditorData.cpp
--- a/Engine/Plugins/Experimental/AnimNext/Source/AnimNextUncookedOnly/Private/Param/AnimNextParameterBlock_EditorData.cpp
+++ b/Engine/Plugins/Experimental/AnimNext/Source/AnimNextUncookedOnly/Private/Param/AnimNextParameterBlock_EditorData.cpp
@@ -6,7 +6,7 @@
 #include "Param/AnimNextParameterBlock_EdGraph.h"
 #include "Param/AnimNextParameterBlock_EdGraphSchema.h"
 #include "UncookedOnlyUtils.h"
-#include "Param/ParametersExecuteContext.h"
+#include "Graph/AnimNextExecuteContext.h"
 #include "Rigs/RigHierarchyPose.h"
 #include "RigVMModel/RigVMFunctionLibrary.h"
 #include "RigVMModel/RigVMNotifications.h"
@@ -20,6 +20,7 @@
 #include "Param/AnimNextParameterBlockBindingReference.h"
 #include "Param/AnimNextParameterLibrary.h"
 #include "Misc/TransactionObjectEvent.h"
+#include "ObjectTools.h"
 
 #if WITH_EDITORONLY_DATA
 
@@ -57,7 +58,7 @@ UAnimNextParameterBlock_EditorData::UAnimNextParameterBlock_EditorData(const FOb
 {
 	RigVMClient.Reset();
 	RigVMClient.SetSchemaClass(UAnimNextParameterBlockLibrary_Schema::StaticClass());
-	RigVMClient.SetExecuteContextStruct(FAnimNextParametersExecuteContext::StaticStruct());
+	RigVMClient.SetExecuteContextStruct(FAnimNextExecuteContext::StaticStruct());
 	RigVMClient.SetOuterClientHost(this, GET_MEMBER_NAME_CHECKED(UAnimNextParameterBlock_EditorData, RigVMClient));
 	{
 		TGuardValue<bool> DisableClientNotifs(RigVMClient.bSuspendNotifications, true);
@@ -132,7 +133,7 @@ UAnimNextParameterBlockBinding* UAnimNextParameterBlock_EditorData::AddBinding(F
 		NewEntry->BindingGraph = NewGraph;
 
 		URigVMController* Controller = RigVMClient.GetController(NewGraph);
-		UE::AnimNext::UncookedOnly::FUtils::SetupBindingGraphForLiteral(Controller, Parameter->GetType());
+		UE::AnimNext::UncookedOnly::FUtils::SetupBindingGraphForLiteral(Controller, InName, Parameter->GetType());
 	}
 
 	BroadcastModified();
@@ -292,15 +293,26 @@ bool UAnimNextParameterBlock_EditorData::RemoveEntry(UAnimNextParameterBlockEntr
 	UAnimNextParameterBlockEntry* EntryToRemove = *EntryToRemovePtr;
 	Entries.Remove(EntryToRemove);
 
+	if (bSetupUndoRedo)
+	{
+		EntryToRemove->Modify();
+	}
+
 	bool bResult = true;
 	if(const IAnimNextParameterBlockGraphInterface* GraphInterface = Cast<IAnimNextParameterBlockGraphInterface>(EntryToRemove))
 	{
 		// Remove any binding graphs
-		TGuardValue<bool> EnablePythonPrint(bSuspendPythonMessagesForRigVMClient, !bPrintPythonCommand);
-		TGuardValue<bool> DisableAutoCompile(bAutoRecompileVM, false);
-		bResult = RigVMClient.RemoveModel(GraphInterface->GetGraph()->GetNodePath(), bSetupUndoRedo);
+		if(GraphInterface->GetGraph())
+		{
+			TGuardValue<bool> EnablePythonPrint(bSuspendPythonMessagesForRigVMClient, !bPrintPythonCommand);
+			TGuardValue<bool> DisableAutoCompile(bAutoRecompileVM, false);
+			bResult = RigVMClient.RemoveModel(GraphInterface->GetGraph()->GetNodePath(), bSetupUndoRedo);
+		}
 	}
 
+	// This will cause any external package to be removed when saved
+	EntryToRemove->MarkAsGarbage();
+	
 	BroadcastModified();
 
 	return bResult;
@@ -620,7 +632,7 @@ void UAnimNextParameterBlock_EditorData::HandleRigVMGraphAdded(const FRigVMClien
 {
 	if(URigVMGraph* RigVMGraph = InClient->GetModel(InNodePath))
 	{
-		RigVMGraph->SetExecuteContextStruct(FAnimNextParametersExecuteContext::StaticStruct());
+		RigVMGraph->SetExecuteContextStruct(FAnimNextExecuteContext::StaticStruct());
 
 		if(!HasAnyFlags(RF_ClassDefaultObject | RF_NeedInitialization | RF_NeedLoad | RF_NeedPostLoad) &&
 			GetOuter() != GetTransientPackage())
@@ -856,6 +868,8 @@ void UAnimNextParameterBlock_EditorData::HandleModifiedEvent(ERigVMGraphNotifTyp
 					break;
 				}
 			}
+
+			RequestAutoVMRecompilation();
 			break;
 		}
 	}
diff --git a/Engine/Plugins/Experimental/AnimNext/Source/AnimNextUncookedOnly/Private/UncookedOnlyUtils.cpp b/Engine/Plugins/Experimental/AnimNext/Source/AnimNextUncookedOnly/Private/UncookedOnlyUtils.cpp
--- a/Engine/Plugins/Experimental/AnimNext/Source/AnimNextUncookedOnly/Private/UncookedOnlyUtils.cpp
+++ b/Engine/Plugins/Experimental/AnimNext/Source/AnimNextUncookedOnly/Private/UncookedOnlyUtils.cpp
@@ -5,13 +5,12 @@
 #include "Graph/AnimNextGraph_EditorData.h"
 #include "RigVMCompiler/RigVMCompiler.h"
 #include "RigVMCore/RigVM.h"
-#include "Graph/GraphExecuteContext.h"
+#include "Graph/AnimNextExecuteContext.h"
 #include "Param/AnimNextParameter.h"
 #include "Param/AnimNextParameterBlock.h"
 #include "Param/AnimNextParameterBlock_EditorData.h"
 #include "Param/AnimNextParameterBlock_EdGraph.h"
-#include "Param/ParametersExecuteContext.h"
-#include "Param/RigUnit_AnimNextParametersBeginExecution.h"
+#include "Graph/RigUnit_AnimNextBeginExecution.h"
 #include "Param/RigVMDispatch_SetParameter.h"
 #include "Param/AnimNextParameterBlockEntry.h"
 #include "Param/IAnimNextParameterBlockBindingInterface.h"
@@ -274,9 +273,11 @@ void FUtils::CompileStruct(UAnimNextParameterBlock* InParameterBlock)
 	{
 		if(const IAnimNextParameterBlockBindingInterface* Binding = Cast<IAnimNextParameterBlockBindingInterface>(Entry))
 		{
-			const FAnimNextParamType& Type = Binding->GetParamType();
-			const UAnimNextParameter* Parameter = Binding->GetParameter();
-			PropertyDescs.Emplace(Parameter->GetFName(), Type.GetContainerType(), Type.GetValueType(), Type.GetValueTypeObject());
+			if(const UAnimNextParameter* Parameter = Binding->GetParameter())
+			{
+				const FAnimNextParamType& Type = Binding->GetParamType();
+				PropertyDescs.Emplace(Parameter->GetFName(), Type.GetContainerType(), Type.GetValueType(), Type.GetValueTypeObject());
+			}
 		}
 	}
 
@@ -292,14 +293,6 @@ void FUtils::Compile(UAnimNextParameterBlock* InParameterBlock)
 {
 	check(InParameterBlock);
 
-	UAnimNextParameterBlock_EditorData* EditorData = GetEditorData(InParameterBlock);
-	if(EditorData->bIsCompiling)
-	{
-		return;
-	}
-
-	TGuardValue<bool> CompilingGuard(EditorData->bIsCompiling, true);
-
 	CompileStruct(InParameterBlock);
 	CompileVM(InParameterBlock);
 }
@@ -613,7 +606,7 @@ FRigVMTemplateArgumentType FUtils::GetRigVMArgTypeFromParamType(const FAnimNextP
 	return ArgType;
 }
 
-void FUtils::SetupBindingGraphForLiteral(URigVMController* InController, const FAnimNextParamType& InParamType)
+void FUtils::SetupBindingGraphForLiteral(URigVMController* InController, FName InParameterName, const FAnimNextParamType& InParamType)
 {
 	FRigVMTemplateArgumentType ArgType = GetRigVMArgTypeFromParamType(InParamType);
 	TRigVMTypeIndex TypeIndex = FRigVMRegistry::Get().GetTypeIndex(ArgType);
@@ -622,14 +615,22 @@ void FUtils::SetupBindingGraphForLiteral(URigVMController* InController, const F
 	InController->RemoveNodes(InController->GetGraph()->GetNodes());
 
 	// Add new nodes for a simple literal binding
-	URigVMUnitNode* EntryPointNode = InController->AddUnitNode(FRigUnit_AnimNextParametersBeginExecution::StaticStruct(), FRigUnit::GetMethodName(), FVector2D(-200.0f, 0.0f), FString(), false);
+	URigVMUnitNode* EntryPointNode = InController->AddUnitNode(FRigUnit_AnimNextBeginExecution::StaticStruct(), FRigUnit::GetMethodName(), FVector2D(-200.0f, 0.0f), FString(), false);
 
 	const FName FactoryName = FRigVMDispatch_SetParameter().GetFactoryName();
 	FRigVMDispatch_SetParameter* Factory = static_cast<FRigVMDispatch_SetParameter*>(FRigVMRegistry::Get().FindDispatchFactory(FactoryName));
 	URigVMTemplateNode* SetParameterNode = InController->AddTemplateNode(Factory->GetTemplate()->GetNotation(), FVector2D(200.0f, 0.0f));
+	InController->SetPinDefaultValue(SetParameterNode->FindPin(FRigVMDispatch_SetParameter::ParameterName.ToString())->GetPinPath(), InParameterName.ToString());
 	InController->ResolveWildCardPin(SetParameterNode->FindPin(FRigVMDispatch_SetParameter::ValueName.ToString()), TypeIndex);
 
-	InController->AddLink(EntryPointNode->FindPin(GET_MEMBER_NAME_STRING_CHECKED(FRigUnit_AnimNextParametersBeginExecution, ExecuteContext)), SetParameterNode->FindPin(FRigVMDispatch_SetParameter::ExecuteContextName.ToString()));
+	InController->AddLink(EntryPointNode->FindPin(GET_MEMBER_NAME_STRING_CHECKED(FRigUnit_AnimNextBeginExecution, ExecuteContext)), SetParameterNode->FindPin(FRigVMDispatch_SetParameter::ExecuteContextName.ToString()));
+}
+
+FText FUtils::GetParameterDisplayNameText(FName InParameterName)
+{
+	FString NameAsString = InParameterName.ToString();
+	NameAsString.ReplaceCharInline(TEXT('_'), TEXT('.'));
+	return FText::FromString(NameAsString);
 }
 
 }
diff --git a/Engine/Plugins/Runtime/PropertyAccess/Source/PropertyAccessEditor/Private/PropertyAccessEditor.cpp b/Engine/Plugins/Runtime/PropertyAccess/Source/PropertyAccessEditor/Private/PropertyAccessEditor.cpp
--- a/Engine/Plugins/Runtime/PropertyAccess/Source/PropertyAccessEditor/Private/PropertyAccessEditor.cpp
+++ b/Engine/Plugins/Runtime/PropertyAccess/Source/PropertyAccessEditor/Private/PropertyAccessEditor.cpp
@@ -341,6 +341,49 @@ struct FPropertyAccessEditorSystem
 		return Result;
 	}
 	
+	static EPropertyAccessCopyType GetAccessType(const FPropertyAccessSegment& InSegment)
+	{
+		FProperty* Property = InSegment.Property.Get();
+		if (FArrayProperty* ArrayProperty = CastField<FArrayProperty>(Property))
+		{
+			// use the array's inner property if we are not trying to copy the whole array
+			if (InSegment.ArrayIndex != INDEX_NONE)
+			{
+				Property = ArrayProperty->Inner;
+			}
+		}
+
+		if (CastField<FNameProperty>(Property))
+		{
+			return EPropertyAccessCopyType::Name;
+		}
+		else if (CastField<FBoolProperty>(Property))
+		{
+			return EPropertyAccessCopyType::Bool;
+		}
+		else if (CastField<FStructProperty>(Property))
+		{
+			return EPropertyAccessCopyType::Struct;
+		}
+		else if (CastField<FObjectPropertyBase>(Property))
+		{
+			return EPropertyAccessCopyType::Object;
+		}
+		else if (CastField<FArrayProperty>(Property) && Property->HasAnyPropertyFlags(CPF_EditFixedSize))
+		{
+			// only apply array copying rules if the destination array is fixed size, otherwise it will be 'complex'
+			return EPropertyAccessCopyType::Array;
+		}
+		else if (Property->PropertyFlags & CPF_IsPlainOldData)
+		{
+			return EPropertyAccessCopyType::Plain;
+		}
+		else
+		{
+			return EPropertyAccessCopyType::Complex;
+		}
+	}
+
 	static EPropertyAccessCopyType GetCopyType(const FPropertyAccessSegment& InSrcSegment, const FPropertyAccessSegment& InDestSegment, FText& OutErrorMessage)
 	{
 		FProperty* SrcProperty = InSrcSegment.Property.Get();
@@ -370,35 +413,7 @@ struct FPropertyAccessEditorSystem
 		EPropertyAccessCompatibility Compatibility = PropertyAccess::GetPropertyCompatibility(SrcProperty, DestProperty);
 		if(Compatibility == EPropertyAccessCompatibility::Compatible)
 		{
-			if (CastField<FNameProperty>(DestProperty))
-			{
-				return EPropertyAccessCopyType::Name;
-			}
-			else if (CastField<FBoolProperty>(DestProperty))
-			{
-				return EPropertyAccessCopyType::Bool;
-			}
-			else if (CastField<FStructProperty>(DestProperty))
-			{
-				return EPropertyAccessCopyType::Struct;
-			}
-			else if (CastField<FObjectPropertyBase>(DestProperty))
-			{
-				return EPropertyAccessCopyType::Object;
-			}
-			else if (CastField<FArrayProperty>(DestProperty) && DestProperty->HasAnyPropertyFlags(CPF_EditFixedSize))
-			{
-				// only apply array copying rules if the destination array is fixed size, otherwise it will be 'complex'
-				return EPropertyAccessCopyType::Array;
-			}
-			else if(DestProperty->PropertyFlags & CPF_IsPlainOldData)
-			{
-				return EPropertyAccessCopyType::Plain;
-			}
-			else
-			{
-				return EPropertyAccessCopyType::Complex;
-			}
+			return GetAccessType(InDestSegment);
 		}
 		else if(Compatibility == EPropertyAccessCompatibility::Promotable)
 		{
@@ -582,6 +597,32 @@ struct FPropertyAccessEditorSystem
 
 		return false;
 	}
+
+	static bool CompileAccess(const UStruct* InStruct, FPropertyAccessLibrary& InLibrary, FPropertyAccessLibraryCompiler::FQueuedAccess& OutAccess)
+	{
+		FPropertyAccessPath AccessPath;
+
+		FResolveSegmentsContext Context(InStruct, OutAccess.Path, AccessPath);
+		Context.bPerformValidation = false;
+
+		OutAccess.Result = ResolveSegments(Context);
+		OutAccess.ErrorText = Context.ErrorMessage;
+
+		if (OutAccess.Result != EPropertyAccessResolveResult::Failed)
+		{
+			OutAccess.AccessIndex = InLibrary.DestPaths.Num();
+			OutAccess.AccessType = GetAccessType(Context.Segments.Last());
+
+			AccessPath.PathSegmentStartIndex = InLibrary.PathSegments.Num();
+			AccessPath.PathSegmentCount = Context.Segments.Num();
+			InLibrary.DestPaths.Add(AccessPath);
+			InLibrary.PathSegments.Append(Context.Segments);
+
+			return true;
+		}
+
+		return false;
+	}
 };
 
 namespace PropertyAccess
@@ -915,7 +956,18 @@ FPropertyAccessHandle FPropertyAccessLibraryCompiler::AddCopy(TArrayView<FString
 
 	QueuedCopies.Add(MoveTemp(QueuedCopy));
 
-	return FPropertyAccessHandle(QueuedCopies.Num() - 1);
+	return FPropertyAccessHandle(QueuedCopies.Num() - 1, EPropertyAccessHandleType::Copy);
+}
+
+FPropertyAccessHandle FPropertyAccessLibraryCompiler::AddAccess(TArrayView<FString> InPath, UObject* InAssociatedObject)
+{
+	FQueuedAccess QueuedAccess;
+	QueuedAccess.Path = InPath;
+	QueuedAccess.AssociatedObject = InAssociatedObject;
+
+	QueuedAccesses.Add(MoveTemp(QueuedAccess));
+
+	return FPropertyAccessHandle(QueuedAccesses.Num() - 1, EPropertyAccessHandleType::Access);
 }
 
 bool FPropertyAccessLibraryCompiler::FinishCompilation()
@@ -927,7 +979,14 @@ bool FPropertyAccessLibraryCompiler::FinishCompilation()
 		{
 			FQueuedCopy& Copy = QueuedCopies[CopyIndex];
 			bResult &= ::FPropertyAccessEditorSystem::CompileCopy(Class, OnDetermineBatchId, *Library, Copy);
-			CopyMap.Add(FPropertyAccessHandle(CopyIndex), FCompiledPropertyAccessHandle(Copy.BatchIndex, Copy.BatchId));
+			CopyMap.Add(FPropertyAccessHandle(CopyIndex, EPropertyAccessHandleType::Copy), FCompiledPropertyAccessHandle(Copy.BatchIndex, Copy.BatchId, EPropertyAccessHandleType::Copy));
+		}
+
+		for (int32 AccessIndex = 0; AccessIndex < QueuedAccesses.Num(); ++AccessIndex)
+		{
+			FQueuedAccess& Access = QueuedAccesses[AccessIndex];
+			bResult &= ::FPropertyAccessEditorSystem::CompileAccess(Class, *Library, Access);
+			AccessMap.Add(FPropertyAccessHandle(AccessIndex, EPropertyAccessHandleType::Access), FCompiledPropertyAccessHandle(Access.AccessIndex, INDEX_NONE, EPropertyAccessHandleType::Access));
 		}
 
 		// Always rebuild the library even if we detected a 'failure'. Otherwise we could fail to copy data for both
@@ -961,12 +1020,35 @@ void FPropertyAccessLibraryCompiler::IterateErrors(TFunctionRef<void(const FText
 
 FCompiledPropertyAccessHandle FPropertyAccessLibraryCompiler::GetCompiledHandle(FPropertyAccessHandle InHandle) const
 {
-	if(const FCompiledPropertyAccessHandle* FoundHandle = CopyMap.Find(InHandle))
+	if (InHandle.GetType() == EPropertyAccessHandleType::Copy)
 	{
-		return *FoundHandle;
+		if (const FCompiledPropertyAccessHandle* FoundHandle = CopyMap.Find(InHandle))
+		{
+			return *FoundHandle;
+		}
+	}
+	else if (InHandle.GetType() == EPropertyAccessHandleType::Access)
+	{
+		if (const FCompiledPropertyAccessHandle* FoundHandle = AccessMap.Find(InHandle))
+		{
+			return *FoundHandle;
+		}
 	}
 
 	return FCompiledPropertyAccessHandle();
 }
 
+EPropertyAccessCopyType FPropertyAccessLibraryCompiler::GetCompiledHandleAccessType(FPropertyAccessHandle InHandle) const
+{
+	if (InHandle.GetType() == EPropertyAccessHandleType::Access)
+	{
+		if (const FCompiledPropertyAccessHandle* FoundHandle = AccessMap.Find(InHandle))
+		{
+			return QueuedAccesses[InHandle.GetId()].AccessType;
+		}
+	}
+
+	return EPropertyAccessCopyType::None;
+}
+
 #undef LOCTEXT_NAMESPACE
diff --git a/Engine/Plugins/Runtime/PropertyAccess/Source/PropertyAccessEditor/Private/PropertyAccessEditor.h b/Engine/Plugins/Runtime/PropertyAccess/Source/PropertyAccessEditor/Private/PropertyAccessEditor.h
--- a/Engine/Plugins/Runtime/PropertyAccess/Source/PropertyAccessEditor/Private/PropertyAccessEditor.h
+++ b/Engine/Plugins/Runtime/PropertyAccess/Source/PropertyAccessEditor/Private/PropertyAccessEditor.h
@@ -4,6 +4,7 @@
 
 #include "IPropertyAccessEditor.h"
 #include "IPropertyAccessCompiler.h"
+#include "PropertyAccess.h"
 
 struct FEdGraphPinType;
 
@@ -40,9 +41,11 @@ public:
 	// IPropertyAccessLibraryCompiler interface
 	virtual void BeginCompilation() override;
 	virtual FPropertyAccessHandle AddCopy(TArrayView<FString> InSourcePath, TArrayView<FString> InDestPath, const FName& InContextId, UObject* InAssociatedObject = nullptr) override;
+	virtual FPropertyAccessHandle AddAccess(TArrayView<FString> InPath, UObject* InAssociatedObject = nullptr) override;
 	virtual bool FinishCompilation() override;
 	virtual void IterateErrors(TFunctionRef<void(const FText&, UObject*)> InFunction) const override;
 	virtual FCompiledPropertyAccessHandle GetCompiledHandle(FPropertyAccessHandle InHandle) const override;
+	virtual EPropertyAccessCopyType GetCompiledHandleAccessType(FPropertyAccessHandle InHandle) const override;
 
 public:
 	// Stored copy info for processing in FinishCompilation()
@@ -79,6 +82,28 @@ public:
 		int32 BatchIndex = INDEX_NONE;
 	};
 
+	// Stored access info for processing in FinishCompilation()
+	struct FQueuedAccess
+	{
+		// The path to access
+		TArray<FString> Path;
+
+		// Error text associated with the path
+		FText ErrorText;
+
+		// Result of the path resolution
+		EPropertyAccessResolveResult Result = EPropertyAccessResolveResult::Failed;
+
+		// Associated object, used to provide context for path resolution
+		UObject* AssociatedObject = nullptr;
+
+		// Compiled access index
+		int32 AccessIndex = INDEX_NONE;
+
+		// Compiled access type
+		EPropertyAccessCopyType AccessType = EPropertyAccessCopyType::None;
+	};
+
 protected:
 	friend struct FPropertyAccessEditorSystem;
 
@@ -94,6 +119,12 @@ protected:
 	// All copies to process in FinishCompilation()
 	TArray<FQueuedCopy> QueuedCopies;
 
-	// Copy map
+	// All accesses to process in FinishCompilation()
+	TArray<FQueuedAccess> QueuedAccesses;
+
+	// Copy map used to resolve handles
 	TMap<FPropertyAccessHandle, FCompiledPropertyAccessHandle> CopyMap;
+
+	// Access map used to resolve handles
+	TMap<FPropertyAccessHandle, FCompiledPropertyAccessHandle> AccessMap;
 };
diff --git a/Engine/Plugins/Tests/EditorTests/Source/EditorTests/Private/Animation/AnimationBlueprintFastPathTest.cpp b/Engine/Plugins/Tests/EditorTests/Source/EditorTests/Private/Animation/AnimationBlueprintFastPathTest.cpp
--- a/Engine/Plugins/Tests/EditorTests/Source/EditorTests/Private/Animation/AnimationBlueprintFastPathTest.cpp
+++ b/Engine/Plugins/Tests/EditorTests/Source/EditorTests/Private/Animation/AnimationBlueprintFastPathTest.cpp
@@ -44,8 +44,6 @@ DEFINE_LATENT_AUTOMATION_COMMAND_ONE_PARAMETER(FCheckFastPathLatentCommand, FChe
 
 bool FCheckFastPathLatentCommand::Update()
 {
-	bool bIsStructTest = Payload.Filename.Contains(TEXT("SubStruct"));
-
 	UAnimBlueprint* AnimBlueprint = LoadObject<UAnimBlueprint>(nullptr, *Payload.Filename);
 	if (AnimBlueprint != nullptr)
 	{
@@ -62,29 +60,34 @@ bool FCheckFastPathLatentCommand::Update()
 						FAnimNode_Base* AnimNode = Property->ContainerPtrToValuePtr<FAnimNode_Base>(DefaultAnimInstance);
 						if (AnimNode)
 						{
-							if (Payload.bCheckEnabled)
+							if (AnimNode->GetEvaluateGraphExposedInputs().GetHandlerStruct() == FAnimNodeExposedValueHandler_PropertyAccess::StaticStruct())
 							{
-								if (AnimNode->GetEvaluateGraphExposedInputs().BoundFunction != NAME_None)
-								{
-									UE_LOG(LogAnimBlueprintFastPathTests, Error, TEXT("Found blueprint VM call (%s) in fast-path Anim Blueprint (%s)"), *AnimNode->GetEvaluateGraphExposedInputs().BoundFunction.ToString(), *AnimBlueprint->GetName());
-								}
-								for (const FExposedValueCopyRecord& CopyRecord : AnimNode->GetEvaluateGraphExposedInputs().CopyRecords)
+								const FAnimNodeExposedValueHandler_PropertyAccess* Handler = static_cast<const FAnimNodeExposedValueHandler_PropertyAccess*>(AnimNode->GetEvaluateGraphExposedInputs().GetHandler());
+
+								if (Payload.bCheckEnabled)
 								{
-									if (CopyRecord.CopyIndex == INDEX_NONE)
+									if (Handler->BoundFunction != NAME_None)
 									{
-										UE_LOG(LogAnimBlueprintFastPathTests, Error, TEXT("Anim blueprint has an invalid copy index (%s)"), *AnimBlueprint->GetName());
+										UE_LOG(LogAnimBlueprintFastPathTests, Error, TEXT("Found blueprint VM call (%s) in fast-path Anim Blueprint (%s)"), *Handler->BoundFunction.ToString(), *AnimBlueprint->GetName());
+									}
+									for (const FExposedValueCopyRecord& CopyRecord : Handler->CopyRecords)
+									{
+										if (CopyRecord.CopyIndex == INDEX_NONE)
+										{
+											UE_LOG(LogAnimBlueprintFastPathTests, Error, TEXT("Anim blueprint has an invalid copy index (%s)"), *AnimBlueprint->GetName());
+										}
 									}
 								}
-							}
-							else
-							{
-								if (AnimNode->GetEvaluateGraphExposedInputs().BoundFunction == NAME_None)
-								{
-									UE_LOG(LogAnimBlueprintFastPathTests, Error, TEXT("No function bound for node evaluation (%s)"), *AnimBlueprint->GetName());
-								}
-								if (AnimNode->GetEvaluateGraphExposedInputs().CopyRecords.Num() > 0)
+								else
 								{
-									UE_LOG(LogAnimBlueprintFastPathTests, Error, TEXT("Found copy records in non-fast-path node evaluator (%s)"), *AnimBlueprint->GetName());
+									if (Handler->BoundFunction == NAME_None)
+									{
+										UE_LOG(LogAnimBlueprintFastPathTests, Error, TEXT("No function bound for node evaluation (%s)"), *AnimBlueprint->GetName());
+									}
+									if (Handler->CopyRecords.Num() > 0)
+									{
+										UE_LOG(LogAnimBlueprintFastPathTests, Error, TEXT("Found copy records in non-fast-path node evaluator (%s)"), *AnimBlueprint->GetName());
+									}
 								}
 							}
 						}
diff --git a/Engine/Source/Editor/AnimGraph/Private/AnimBlueprintCompilationContext.cpp b/Engine/Source/Editor/AnimGraph/Private/AnimBlueprintCompilationContext.cpp
--- a/Engine/Source/Editor/AnimGraph/Private/AnimBlueprintCompilationContext.cpp
+++ b/Engine/Source/Editor/AnimGraph/Private/AnimBlueprintCompilationContext.cpp
@@ -84,6 +84,11 @@ const TMap<UAnimGraphNode_Base*, FProperty*>& FAnimBlueprintCompilationContext::
 	return CompilerContext->AllocatedAnimNodes;
 }
 
+const TMap<UAnimGraphNode_Base*, FProperty*>& FAnimBlueprintCompilationContext::GetAllocatedHandlerPropertiesByNodeImpl() const
+{
+	return CompilerContext->AllocatedAnimNodeHandlers;
+}
+
 void FAnimBlueprintCompilationContext::AddAttributesToNodeImpl(UAnimGraphNode_Base* InNode, TArrayView<const FName> InAttributes) const
 {
 	CompilerContext->AddAttributesToNode(InNode, InAttributes);
@@ -124,6 +129,11 @@ const TMap<UAnimGraphNode_Base*, int32>& FAnimBlueprintCompilationBracketContext
 	return CompilerContext->AllocatedAnimNodeIndices;
 }
 
+const TMap<UAnimGraphNode_Base*, FProperty*>& FAnimBlueprintCompilationBracketContext::GetAllocatedHandlerPropertiesByNodeImpl() const
+{
+	return CompilerContext->AllocatedAnimNodeHandlers;
+}
+
 FCompilerResultsLog& FAnimBlueprintCopyTermDefaultsContext::GetMessageLogImpl() const
 {
 	return CompilerContext->MessageLog;
diff --git a/Engine/Source/Editor/AnimGraph/Private/AnimBlueprintCompilationContext.h b/Engine/Source/Editor/AnimGraph/Private/AnimBlueprintCompilationContext.h
--- a/Engine/Source/Editor/AnimGraph/Private/AnimBlueprintCompilationContext.h
+++ b/Engine/Source/Editor/AnimGraph/Private/AnimBlueprintCompilationContext.h
@@ -39,6 +39,7 @@ private:
 	virtual const TMap<UAnimGraphNode_Base*, UAnimGraphNode_Base*>& GetSourceNodeToProcessedNodeMapImpl() const override;
 	virtual const TMap<int32, FProperty*>& GetAllocatedPropertiesByIndexImpl() const override;
 	virtual const TMap<UAnimGraphNode_Base*, FProperty*>& GetAllocatedPropertiesByNodeImpl() const override;
+	virtual const TMap<UAnimGraphNode_Base*, FProperty*>& GetAllocatedHandlerPropertiesByNodeImpl() const override;
 	virtual void AddAttributesToNodeImpl(UAnimGraphNode_Base* InNode, TArrayView<const FName> InAttributes) const override;
 	virtual TArrayView<const FName> GetAttributesFromNodeImpl(UAnimGraphNode_Base* InNode) const override;
 	virtual bool IsAnimGraphNodeFoldedImpl(UAnimGraphNode_Base* InNode) const override;
@@ -60,6 +61,7 @@ private:
 
 	virtual FCompilerResultsLog& GetMessageLogImpl() const override;
 	virtual const TMap<UAnimGraphNode_Base*, int32>& GetAllocatedAnimNodeIndicesImpl() const override;
+	virtual const TMap<UAnimGraphNode_Base*, FProperty*>& GetAllocatedHandlerPropertiesByNodeImpl() const override;
 	
 	FAnimBlueprintCompilerContext* CompilerContext;
 };
diff --git a/Engine/Source/Editor/AnimGraph/Private/AnimBlueprintCompiler.cpp b/Engine/Source/Editor/AnimGraph/Private/AnimBlueprintCompiler.cpp
--- a/Engine/Source/Editor/AnimGraph/Private/AnimBlueprintCompiler.cpp
+++ b/Engine/Source/Editor/AnimGraph/Private/AnimBlueprintCompiler.cpp
@@ -34,6 +34,7 @@
 #include "AnimBlueprintExtension.h"
 #include "UObject/FrameworkObjectVersion.h"
 #include "UObject/UE5MainStreamObjectVersion.h"
+#include "AnimGraphNodeBinding.h"
 
 #define LOCTEXT_NAMESPACE "AnimBlueprintCompiler"
 
@@ -273,7 +274,29 @@ void FAnimBlueprintCompilerContext::ProcessAnimationNode(UAnimGraphNode_Base* Vi
 
 	if (NewProperty == nullptr)
 	{
-		MessageLog.Error(TEXT("Failed to create node property for @@"), VisualAnimNode);
+		MessageLog.Error(TEXT("ICE: Failed to create node property for @@"), VisualAnimNode);
+	}
+
+	// Create a handler property in constants
+	UScriptStruct* HandlerStruct = nullptr;
+	if(const UAnimGraphNodeBinding* Binding = VisualAnimNode->GetBinding())
+	{
+		HandlerStruct = Binding->GetAnimNodeHandlerStruct();
+		check(HandlerStruct->IsChildOf(FAnimNodeExposedValueHandler::StaticStruct()));
+	}
+	else
+	{
+		HandlerStruct = FAnimNodeExposedValueHandler::StaticStruct();
+	}
+
+	FEdGraphPinType HandlerVariableType;
+	HandlerVariableType.PinCategory = UAnimationGraphSchema::PC_Struct;
+	HandlerVariableType.PinSubCategoryObject = MakeWeakObjectPtr(HandlerStruct);
+
+	FStructProperty* NewHandlerProperty = CastField<FStructProperty>(CreateStructVariable(NewAnimBlueprintConstants, FName(*NodeVariableName), HandlerVariableType));
+	if (NewHandlerProperty == nullptr)
+	{
+		MessageLog.Error(TEXT("ICE: Failed to create node handler property for @@"), VisualAnimNode);
 	}
 
 	GatherFoldRecordsForAnimationNode(NodeType, NewProperty, VisualAnimNode);
@@ -281,6 +304,7 @@ void FAnimBlueprintCompilerContext::ProcessAnimationNode(UAnimGraphNode_Base* Vi
 	// Register this node with the compile-time data structures
 	const int32 AllocatedIndex = AllocateNodeIndexCounter++;
 	AllocatedAnimNodes.Add(VisualAnimNode, NewProperty);
+	AllocatedAnimNodeHandlers.Add(VisualAnimNode, NewHandlerProperty);
 	AllocatedNodePropertiesToNodes.Add(NewProperty, VisualAnimNode);
 	AllocatedAnimNodeIndices.Add(VisualAnimNode, AllocatedIndex);
 	AllocatedPropertiesByIndex.Add(AllocatedIndex, NewProperty);
@@ -846,28 +870,6 @@ void FAnimBlueprintCompilerContext::CopyTermDefaultsToDefaultObject(UObject* Def
 		FAnimBlueprintGeneratedClassCompiledData CompiledData(NewAnimBlueprintClass);
 		FAnimBlueprintCopyTermDefaultsContext CompilerContext(this);
 
-		// Initialize extensions from their templates
-		for (TFieldIterator<FStructProperty> It(DefaultAnimInstance->GetClass(), EFieldIteratorFlags::ExcludeSuper); It; ++It)
-		{
-			FStructProperty* TargetProperty = *It;
-
-			if (UAnimBlueprintExtension* Extension = InstancePropertyToExtensionMap.FindRef(TargetProperty))
-			{
-				uint8* DestinationPtr = TargetProperty->ContainerPtrToValuePtr<uint8>(DefaultAnimInstance);
-				const uint8* SourcePtr = nullptr;
-				const FStructProperty* SourceExtensionProperty = Extension->GetInstanceDataProperty();
-				
-				if(SourceExtensionProperty)
-				{
-					check(TargetProperty->Struct == SourceExtensionProperty->Struct);				
-					SourcePtr = SourceExtensionProperty->ContainerPtrToValuePtr<uint8>(Extension);
-				}
-				
-				FAnimBlueprintExtensionCopyTermDefaultsContext ExtensionContext(DefaultObject, TargetProperty, DestinationPtr, SourcePtr, LinkIndexCount);
-				Extension->CopyTermDefaultsToDefaultObject(DefaultAnimInstance, CompilerContext, ExtensionContext);
-			}
-		}
-		
 		// Initialize animation nodes from their templates
 		for (TFieldIterator<FProperty> It(DefaultAnimInstance->GetClass(), EFieldIteratorFlags::ExcludeSuper); It; ++It)
 		{
@@ -966,6 +968,28 @@ void FAnimBlueprintCompilerContext::CopyTermDefaultsToDefaultObject(UObject* Def
 			PatchDataArea(Mutables, MutablesStruct, MutablePropertyRecords);
 		}
 		
+		// Initialize extensions from their templates
+		for (TFieldIterator<FStructProperty> It(DefaultAnimInstance->GetClass(), EFieldIteratorFlags::ExcludeSuper); It; ++It)
+		{
+			FStructProperty* TargetProperty = *It;
+
+			if (UAnimBlueprintExtension* Extension = InstancePropertyToExtensionMap.FindRef(TargetProperty))
+			{
+				uint8* DestinationPtr = TargetProperty->ContainerPtrToValuePtr<uint8>(DefaultAnimInstance);
+				const uint8* SourcePtr = nullptr;
+				const FStructProperty* SourceExtensionProperty = Extension->GetInstanceDataProperty();
+
+				if (SourceExtensionProperty)
+				{
+					check(TargetProperty->Struct == SourceExtensionProperty->Struct);
+					SourcePtr = SourceExtensionProperty->ContainerPtrToValuePtr<uint8>(Extension);
+				}
+
+				FAnimBlueprintExtensionCopyTermDefaultsContext ExtensionContext(DefaultObject, TargetProperty, DestinationPtr, SourcePtr, LinkIndexCount);
+				Extension->CopyTermDefaultsToDefaultObject(DefaultAnimInstance, CompilerContext, ExtensionContext);
+			}
+		}
+
 		// And wire up node links
 		for (auto PoseLinkIt = ValidPoseLinkList.CreateIterator(); PoseLinkIt; ++PoseLinkIt)
 		{
diff --git a/Engine/Source/Editor/AnimGraph/Private/AnimBlueprintCompiler.h b/Engine/Source/Editor/AnimGraph/Private/AnimBlueprintCompiler.h
--- a/Engine/Source/Editor/AnimGraph/Private/AnimBlueprintCompiler.h
+++ b/Engine/Source/Editor/AnimGraph/Private/AnimBlueprintCompiler.h
@@ -93,6 +93,7 @@ protected:
 
 	// Map of allocated v3 nodes that are members of the class
 	TMap<class UAnimGraphNode_Base*, FProperty*> AllocatedAnimNodes;
+	TMap<class UAnimGraphNode_Base*, FProperty*> AllocatedAnimNodeHandlers;
 	TMap<FProperty*, class UAnimGraphNode_Base*> AllocatedNodePropertiesToNodes;
 	TMap<FProperty*, class UAnimGraphNode_Base*> AllocatedNodeConstantPropertiesToNodes;
 	TMap<int32, FProperty*> AllocatedPropertiesByIndex;
diff --git a/Engine/Source/Editor/AnimGraph/Private/AnimBlueprintExtension.cpp b/Engine/Source/Editor/AnimGraph/Private/AnimBlueprintExtension.cpp
--- a/Engine/Source/Editor/AnimGraph/Private/AnimBlueprintExtension.cpp
+++ b/Engine/Source/Editor/AnimGraph/Private/AnimBlueprintExtension.cpp
@@ -13,6 +13,7 @@
 #include "AnimBlueprintExtension_NodeRelevancy.h"
 #include "AnimBlueprintExtension_SharedLinkedAnimLayers.h"
 #include "AnimBlueprintExtension_Tag.h"
+#include "AnimGraphNodeBinding.h"
 
 // Set used to refresh extensions. Checks that an extension has a reference from an anim node for each refresh.
 static TSet<TSubclassOf<UAnimBlueprintExtension>> RefreshSet;
@@ -101,6 +102,11 @@ void UAnimBlueprintExtension::RequestExtensionsForNode(UAnimGraphNode_Base* InAn
 		}
 		
 		InAnimGraphNode->GetRequiredExtensions(ExtensionClasses);
+
+		if (const UAnimGraphNodeBinding* Binding = InAnimGraphNode->GetBinding())
+		{
+			Binding->GetRequiredExtensions(ExtensionClasses);
+		}
 	
 		for(const TSubclassOf<UAnimBlueprintExtension>& ExtensionClass : ExtensionClasses)
 		{	
diff --git a/Engine/Source/Editor/AnimGraph/Private/AnimBlueprintExtension_Base.cpp b/Engine/Source/Editor/AnimGraph/Private/AnimBlueprintExtension_Base.cpp
--- a/Engine/Source/Editor/AnimGraph/Private/AnimBlueprintExtension_Base.cpp
+++ b/Engine/Source/Editor/AnimGraph/Private/AnimBlueprintExtension_Base.cpp
@@ -31,6 +31,8 @@
 #include "IAnimBlueprintPostExpansionStepContext.h"
 #include "IAnimBlueprintCompilationBracketContext.h"
 #include "Features/IModularFeatures.h"
+#include "Animation/ExposedValueHandler.h"
+#include "AnimGraphNodeBinding_Base.h"
 
 #define LOCTEXT_NAMESPACE "AnimBlueprintExtension_Base"
 
@@ -47,6 +49,8 @@ void UAnimBlueprintExtension_Base::HandleCopyTermDefaultsToDefaultObject(UObject
 
 	if(DefaultAnimInstance)
 	{
+		Subsystem.PatchValueHandlers(DefaultAnimInstance->GetClass());
+
 		for(const FEvaluationHandlerRecord& EvaluationHandler : ValidEvaluationHandlerList)
 		{
 			if(EvaluationHandler.AnimGraphNode)
@@ -59,8 +63,13 @@ void UAnimBlueprintExtension_Base::HandleCopyTermDefaultsToDefaultObject(UObject
 					const FAnimNodeSinglePropertyHandler& Handler = EvaluationHandler.ServicedProperties.CreateConstIterator()->Value;
 					check(Handler.CopyRecords.Num() > 0);
 
-					const FExposedValueHandler& ValueHandler = Subsystem.ExposedValueHandlers[ EvaluationHandler.EvaluationHandlerIdx ];
-					TrueNode->BlueprintUsage = ValueHandler.BoundFunction != NAME_None ? EBlueprintUsage::UsesBlueprint : EBlueprintUsage::DoesNotUseBlueprint; 
+					bool bHandlerUsesBlueprint = false;
+					if(Subsystem.ExposedValueHandlers[EvaluationHandler.EvaluationHandlerIdx].GetHandlerStruct()->IsChildOf(FAnimNodeExposedValueHandler_Base::StaticStruct()))
+					{
+						const FAnimNodeExposedValueHandler_Base* BaseHandler = static_cast<const FAnimNodeExposedValueHandler_Base*>(Subsystem.ExposedValueHandlers[EvaluationHandler.EvaluationHandlerIdx].GetHandler());
+						bHandlerUsesBlueprint = BaseHandler->BoundFunction != NAME_None;
+					}
+					TrueNode->BlueprintUsage = bHandlerUsesBlueprint ? EBlueprintUsage::UsesBlueprint : EBlueprintUsage::DoesNotUseBlueprint;
 
 #if WITH_EDITORONLY_DATA // ANIMINST_PostCompileValidation
 					const bool bWarnAboutBlueprintUsage = InCompilationContext.GetAnimBlueprint()->bWarnAboutBlueprintUsage || DefaultAnimInstance->PCV_ShouldWarnAboutNodesNotUsingFastPath();
@@ -125,14 +134,12 @@ void UAnimBlueprintExtension_Base::HandlePostExpansionStep(const UEdGraph* InGra
 	}
 }
 
-void UAnimBlueprintExtension_Base::PatchEvaluationHandlers(IAnimBlueprintCompilationBracketContext& InCompilationContext, IAnimBlueprintGeneratedClassCompiledData& OutCompiledData)
+void UAnimBlueprintExtension_Base::PatchEvaluationHandlers(const UClass* InClass, IAnimBlueprintCompilationBracketContext& InCompilationContext, IAnimBlueprintGeneratedClassCompiledData& OutCompiledData)
 {
-	TArray<FExposedValueHandler>& ExposedValueHandlers = Subsystem.ExposedValueHandlers;
-
 	// Exposed value handlers indices must match the index of anim node properties, 
 	// so we iterate over anim node properties here when patching up
 	const int32 NumAllocatedNodes = InCompilationContext.GetAllocatedAnimNodeIndices().Num();
-	ExposedValueHandlers.SetNum(NumAllocatedNodes);
+
 	for(const TPair<UAnimGraphNode_Base*, int32>& GraphNodePair : InCompilationContext.GetAllocatedAnimNodeIndices())
 	{
 		if(int32* EvaluationHandlerIndexPtr = ValidEvaluationHandlerMap.Find(GraphNodePair.Key))
@@ -142,7 +149,7 @@ void UAnimBlueprintExtension_Base::PatchEvaluationHandlers(IAnimBlueprintCompila
 
 			FEvaluationHandlerRecord& EvaluationHandlerRecord = ValidEvaluationHandlerList[*EvaluationHandlerIndexPtr];
 			EvaluationHandlerRecord.EvaluationHandlerIdx = NodePropertyIndex;
-			EvaluationHandlerRecord.PatchFunctionNameAndCopyRecordsInto(ExposedValueHandlers[NodePropertyIndex]);
+			EvaluationHandlerRecord.PatchAnimNodeExposedValueHandler(InClass, InCompilationContext);
 		}
 	}
 }
@@ -184,7 +191,7 @@ void UAnimBlueprintExtension_Base::HandleStartCompilingClass(const UClass* InCla
 			PropertyAccessExtension->OnPreLibraryCompiled().Remove(PreLibraryCompiledDelegateHandle);
 		});
 
-		PostLibraryCompiledDelegateHandle = PropertyAccessExtension->OnPostLibraryCompiled().AddLambda([this, PropertyAccessExtension](IAnimBlueprintCompilationBracketContext& InCompilationContext, IAnimBlueprintGeneratedClassCompiledData& OutCompiledData)
+		PostLibraryCompiledDelegateHandle = PropertyAccessExtension->OnPostLibraryCompiled().AddLambda([this, PropertyAccessExtension, InClass](IAnimBlueprintCompilationBracketContext& InCompilationContext, IAnimBlueprintGeneratedClassCompiledData& OutCompiledData)
 		{
 			for(FEvaluationHandlerRecord& HandlerRecord : ValidEvaluationHandlerList)
 			{
@@ -206,17 +213,20 @@ void UAnimBlueprintExtension_Base::HandleStartCompilingClass(const UClass* InCla
 								BindingName.SetNumber(CopyRecord.DestArrayIndex + 1);
 							}
 
-							if(FAnimGraphNodePropertyBinding* Binding = OriginalNode->PropertyBindings.Find(BindingName))
+							if(UAnimGraphNodeBinding_Base* Binding = Cast<UAnimGraphNodeBinding_Base>(OriginalNode->GetMutableBinding()))
 							{
-								if(CopyRecord.LibraryCompiledHandle.IsValid())
+								if(FAnimGraphNodePropertyBinding* PropertyBinding = Binding->PropertyBindings.Find(BindingName))
 								{
-									Binding->CompiledContext = UAnimBlueprintExtension_PropertyAccess::GetCompiledHandleContext(CopyRecord.LibraryCompiledHandle);
-									Binding->CompiledContextDesc = UAnimBlueprintExtension_PropertyAccess::GetCompiledHandleContextDesc(CopyRecord.LibraryCompiledHandle);
-								}
-								else
-								{
-									Binding->CompiledContext = FText::GetEmpty();
-									Binding->CompiledContextDesc = FText::GetEmpty();
+									if(CopyRecord.LibraryCompiledHandle.IsValid())
+									{
+										PropertyBinding->CompiledContext = UAnimBlueprintExtension_PropertyAccess::GetCompiledHandleContext(CopyRecord.LibraryCompiledHandle);
+										PropertyBinding->CompiledContextDesc = UAnimBlueprintExtension_PropertyAccess::GetCompiledHandleContextDesc(CopyRecord.LibraryCompiledHandle);
+									}
+									else
+									{
+										PropertyBinding->CompiledContext = FText::GetEmpty();
+										PropertyBinding->CompiledContextDesc = FText::GetEmpty();
+									}
 								}
 							}
 						}
@@ -224,7 +234,7 @@ void UAnimBlueprintExtension_Base::HandleStartCompilingClass(const UClass* InCla
 				}
 			}
 
-			PatchEvaluationHandlers(InCompilationContext, OutCompiledData);
+			PatchEvaluationHandlers(InClass, InCompilationContext, OutCompiledData);
 
 			PropertyAccessExtension->OnPostLibraryCompiled().Remove(PostLibraryCompiledDelegateHandle);
 		});
@@ -237,11 +247,11 @@ void UAnimBlueprintExtension_Base::HandleFinishCompilingClass(const UClass* InCl
 	if(PropertyAccessExtension == nullptr)
 	{
 		// Without the property access system we need to patch generated function names here
-		PatchEvaluationHandlers(InCompilationContext, OutCompiledData);
+		PatchEvaluationHandlers(InClass, InCompilationContext, OutCompiledData);
 	}
 }
 
-void UAnimBlueprintExtension_Base::ProcessNodePins(UAnimGraphNode_Base* InNode, IAnimBlueprintCompilationContext& InCompilationContext, IAnimBlueprintGeneratedClassCompiledData& OutCompiledData)
+void UAnimBlueprintExtension_Base::ProcessPosePins(UAnimGraphNode_Base* InNode, IAnimBlueprintCompilationContext& InCompilationContext, IAnimBlueprintGeneratedClassCompiledData& OutCompiledData)
 {
 	const UAnimationGraphSchema* AnimGraphDefaultSchema = GetDefault<UAnimationGraphSchema>();
 
@@ -250,7 +260,6 @@ void UAnimBlueprintExtension_Base::ProcessNodePins(UAnimGraphNode_Base* InNode,
 	for (auto SourcePinIt = InNode->Pins.CreateIterator(); SourcePinIt; ++SourcePinIt)
 	{
 		UEdGraphPin* SourcePin = *SourcePinIt;
-		bool bConsumed = false;
 
 		// Register pose links for future use
 		if ((SourcePin->Direction == EGPD_Input) && (AnimGraphDefaultSchema->IsPosePin(SourcePin->PinType)))
@@ -260,10 +269,32 @@ void UAnimBlueprintExtension_Base::ProcessNodePins(UAnimGraphNode_Base* InNode,
 			if (LinkRecord.IsValid())
 			{
 				InCompilationContext.AddPoseLinkMappingRecord(LinkRecord);
-				bConsumed = true;
+			}
+			else
+			{
+				//@TODO: ANIMREFACTOR: It's probably OK to have certain pins ignored eventually, but this is very helpful during development
+				InCompilationContext.GetMessageLog().Note(TEXT("@@ was visible but ignored"), SourcePin);
 			}
 		}
-		else if(!InNode->ShouldCreateStructEvalHandlers())
+	}
+}
+
+void UAnimBlueprintExtension_Base::ProcessNonPosePins(UAnimGraphNode_Base* InNode, IAnimBlueprintCompilationContext& InCompilationContext, IAnimBlueprintGeneratedClassCompiledData& OutCompiledData, EPinProcessingFlags InFlags)
+{
+	const UAnimationGraphSchema* AnimGraphDefaultSchema = GetDefault<UAnimationGraphSchema>();
+
+	FStructProperty* NodeProperty = CastFieldChecked<FStructProperty>(InCompilationContext.GetAllocatedPropertiesByNode().FindChecked(InNode));
+
+	for (auto SourcePinIt = InNode->Pins.CreateIterator(); SourcePinIt; ++SourcePinIt)
+	{
+		UEdGraphPin* SourcePin = *SourcePinIt;
+		bool bConsumed = false;
+
+		if ((SourcePin->Direction == EGPD_Input) && (AnimGraphDefaultSchema->IsPosePin(SourcePin->PinType)))
+		{
+			bConsumed = true;
+		}
+		else if(!InNode->ShouldCreateStructEvalHandlers() || !EnumHasAnyFlags(InFlags, EPinProcessingFlags::All))
 		{
 			bConsumed = true;
 		}
@@ -295,39 +326,45 @@ void UAnimBlueprintExtension_Base::ProcessNodePins(UAnimGraphNode_Base* InNode,
 			{
 				if (SourcePin->LinkedTo.Num() == 0)
 				{
-					// Literal that can be pushed into the CDO instead of re-evaluated every frame
-					bConsumed = true;
+					if(EnumHasAnyFlags(InFlags, EPinProcessingFlags::Constants))
+					{
+						// Literal that can be pushed into the CDO instead of re-evaluated every frame
+						bConsumed = true;
+					}
 				}
 				else
 				{
-					// Dynamic value that needs to be wired up and evaluated each frame
-					const FString& EvaluationHandlerStr = SourcePinProperty->GetMetaData(AnimGraphDefaultSchema->NAME_OnEvaluate);
-					FName EvaluationHandlerName(*EvaluationHandlerStr);
-					if (EvaluationHandlerName != NAME_None)
+					if (EnumHasAnyFlags(InFlags, EPinProcessingFlags::BlueprintHandlers))
 					{
-						// warn that NAME_OnEvaluate is deprecated:
-						InCompilationContext.GetMessageLog().Warning(*LOCTEXT("OnEvaluateDeprecated", "OnEvaluate meta data is deprecated, found on @@").ToString(), SourcePinProperty);
-					}
+						// Dynamic value that needs to be wired up and evaluated each frame
+						const FString& EvaluationHandlerStr = SourcePinProperty->GetMetaData(AnimGraphDefaultSchema->NAME_OnEvaluate);
+						FName EvaluationHandlerName(*EvaluationHandlerStr);
+						if (EvaluationHandlerName != NAME_None)
+						{
+							// warn that NAME_OnEvaluate is deprecated:
+							InCompilationContext.GetMessageLog().Warning(*LOCTEXT("OnEvaluateDeprecated", "OnEvaluate meta data is deprecated, found on @@").ToString(), SourcePinProperty);
+						}
 					
-					ensure(EvalHandler.NodeVariableProperty == nullptr || EvalHandler.NodeVariableProperty == NodeProperty);
-					EvalHandler.AnimGraphNode = InNode;
-					EvalHandler.NodeVariableProperty = NodeProperty;
-					EvalHandler.RegisterPin(SourcePin, SourcePinProperty, SourceArrayIndex);
-					// if it's not instance property, ensure we mark it
-					EvalHandler.bServicesNodeProperties = EvalHandler.bServicesNodeProperties | !bInstancePropertyExists;
+						ensure(EvalHandler.NodeVariableProperty == nullptr || EvalHandler.NodeVariableProperty == NodeProperty);
+						EvalHandler.AnimGraphNode = InNode;
+						EvalHandler.NodeVariableProperty = NodeProperty;
+						EvalHandler.RegisterPin(SourcePin, SourcePinProperty, SourceArrayIndex, EnumHasAnyFlags(InFlags, EPinProcessingFlags::PropertyAccessFastPath));
+						// if it's not instance property, ensure we mark it
+						EvalHandler.bServicesNodeProperties = EvalHandler.bServicesNodeProperties || !bInstancePropertyExists;
+
+						if (CustomPropertyNode)
+						{
+							EvalHandler.bServicesInstanceProperties = EvalHandler.bServicesInstanceProperties || bInstancePropertyExists;
 
-					if (CustomPropertyNode)
-					{
-						EvalHandler.bServicesInstanceProperties = EvalHandler.bServicesInstanceProperties | bInstancePropertyExists;
+							FAnimNodeSinglePropertyHandler* SinglePropHandler = EvalHandler.ServicedProperties.Find(SourcePinProperty->GetFName());
+							check(SinglePropHandler); // Should have been added in RegisterPin
 
-						FAnimNodeSinglePropertyHandler* SinglePropHandler = EvalHandler.ServicedProperties.Find(SourcePinProperty->GetFName());
-						check(SinglePropHandler); // Should have been added in RegisterPin
+							// Flag that the target property is actually on the instance class and not the node
+							SinglePropHandler->bInstanceIsTarget = bInstancePropertyExists;
+						}
 
-						// Flag that the target property is actually on the instance class and not the node
-						SinglePropHandler->bInstanceIsTarget = bInstancePropertyExists;
+						bConsumed = true;
 					}
-
-					bConsumed = true;
 				}
 
 				UEdGraphPin* TrueSourcePin = InCompilationContext.GetMessageLog().FindSourcePin(SourcePin);
@@ -338,40 +375,46 @@ void UAnimBlueprintExtension_Base::ProcessNodePins(UAnimGraphNode_Base* InNode,
 			}
 		}
 
-		if (!bConsumed && (SourcePin->Direction == EGPD_Input))
+		if (!bConsumed && (SourcePin->Direction == EGPD_Input) && !AnimGraphDefaultSchema->IsPosePin(SourcePin->PinType))
 		{
 			//@TODO: ANIMREFACTOR: It's probably OK to have certain pins ignored eventually, but this is very helpful during development
 			InCompilationContext.GetMessageLog().Note(TEXT("@@ was visible but ignored"), SourcePin);
 		}
 	}
 
-	// Add any property bindings
-	for(const TPair<FName, FAnimGraphNodePropertyBinding>& PropertyBinding : InNode->PropertyBindings)
+	if (EnumHasAnyFlags(InFlags, EPinProcessingFlags::PropertyAccessBindings))
 	{
-		if(PropertyBinding.Value.bIsBound)
+		// Add any property bindings
+		if (UAnimGraphNodeBinding_Base* Binding = Cast<UAnimGraphNodeBinding_Base>(InNode->GetMutableBinding()))
 		{
-			FEvaluationHandlerRecord& EvalHandler = PerNodeStructEvalHandlers.FindOrAdd(InNode);
-			EvalHandler.AnimGraphNode = InNode;
+			for(const TPair<FName, FAnimGraphNodePropertyBinding>& PropertyBinding : Binding->PropertyBindings)
+			{
+				if(PropertyBinding.Value.bIsBound)
+				{
+					FEvaluationHandlerRecord& EvalHandler = PerNodeStructEvalHandlers.FindOrAdd(InNode);
+					EvalHandler.AnimGraphNode = InNode;
 
-			// for array properties we need to account for the extra FName number 
-			FName ComparisonName = PropertyBinding.Key;
-			ComparisonName.SetNumber(0);
+					// for array properties we need to account for the extra FName number 
+					FName ComparisonName = PropertyBinding.Key;
+					ComparisonName.SetNumber(0);
 
-			if (FProperty* Property = FindFProperty<FProperty>(NodeProperty->Struct, ComparisonName))
-			{
-				EvalHandler.NodeVariableProperty = NodeProperty;
-				EvalHandler.bServicesNodeProperties = true;
-				EvalHandler.RegisterPropertyBinding(Property, PropertyBinding.Value);
-			}
-			else if(FProperty* ClassProperty = FindFProperty<FProperty>(InCompilationContext.GetBlueprint()->SkeletonGeneratedClass, PropertyBinding.Value.PropertyName))
-			{
-				EvalHandler.NodeVariableProperty = NodeProperty;
-				EvalHandler.bServicesInstanceProperties = true;
-				EvalHandler.RegisterPropertyBinding(ClassProperty, PropertyBinding.Value);
-			}
-			else
-			{
-				InCompilationContext.GetMessageLog().Warning(*FString::Printf(TEXT("ICE: @@ Failed to find a property '%s'"), *ComparisonName.ToString()), InNode);
+					if (FProperty* Property = FindFProperty<FProperty>(NodeProperty->Struct, ComparisonName))
+					{
+						EvalHandler.NodeVariableProperty = NodeProperty;
+						EvalHandler.bServicesNodeProperties = true;
+						EvalHandler.RegisterPropertyBinding(Property, PropertyBinding.Value);
+					}
+					else if(FProperty* ClassProperty = FindFProperty<FProperty>(InCompilationContext.GetBlueprint()->SkeletonGeneratedClass, PropertyBinding.Value.PropertyName))
+					{
+						EvalHandler.NodeVariableProperty = NodeProperty;
+						EvalHandler.bServicesInstanceProperties = true;
+						EvalHandler.RegisterPropertyBinding(ClassProperty, PropertyBinding.Value);
+					}
+					else
+					{
+						InCompilationContext.GetMessageLog().Warning(*FString::Printf(TEXT("ICE: @@ Failed to find a property '%s'"), *ComparisonName.ToString()), InNode);
+					}
+				}
 			}
 		}
 	}
@@ -726,13 +769,13 @@ void UAnimBlueprintExtension_Base::CreateEvaluationHandler(IAnimBlueprintCompila
 			ExecChain = K2Schema->FindExecutionPin(*AssignmentNode, EGPD_Output);
 		}
 
-		if(InstanceAssignmentNode != nullptr)
+		if (InstanceAssignmentNode != nullptr)
 		{
 			for (int32 PinIndex = 0; PinIndex < InstanceAssignmentNode->ShowPinForProperties.Num(); ++PinIndex)
 			{
 				FOptionalPinFromProperty& TestProperty = InstanceAssignmentNode->ShowPinForProperties[PinIndex];
-				if(TestProperty.bShowPin)
-				{				
+				if (TestProperty.bShowPin)
+				{
 					UEdGraphPin* ExecVariablesIn = K2Schema->FindExecutionPin(*InstanceAssignmentNode, EGPD_Input);
 					ExecChain->MakeLinkTo(ExecVariablesIn);
 					ExecChain = K2Schema->FindExecutionPin(*InstanceAssignmentNode, EGPD_Output);
@@ -746,18 +789,18 @@ void UAnimBlueprintExtension_Base::CreateEvaluationHandler(IAnimBlueprintCompila
 void UAnimBlueprintExtension_Base::RedirectPropertyAccesses(IAnimBlueprintCompilationContext& InCompilationContext, UAnimGraphNode_Base* InNode, FEvaluationHandlerRecord& InRecord)
 {
 	const FStructProperty* MutableDataProperty = InCompilationContext.GetMutableDataProperty();
-	if(InCompilationContext.IsAnimGraphNodeFolded(InNode) && MutableDataProperty != nullptr)
+	if (InCompilationContext.IsAnimGraphNodeFolded(InNode) && MutableDataProperty != nullptr)
 	{
-		for(TPair<FName, FAnimNodeSinglePropertyHandler>& NamePropertyPair : InRecord.ServicedProperties)
+		for (TPair<FName, FAnimNodeSinglePropertyHandler>& NamePropertyPair : InRecord.ServicedProperties)
 		{
-			if(const IAnimBlueprintCompilationContext::FFoldedPropertyRecord* FoldedPropertyRecord = InCompilationContext.GetFoldedPropertyRecord(InNode, NamePropertyPair.Key))
+			if (const IAnimBlueprintCompilationContext::FFoldedPropertyRecord* FoldedPropertyRecord = InCompilationContext.GetFoldedPropertyRecord(InNode, NamePropertyPair.Key))
 			{
-				for(FPropertyCopyRecord& CopyRecord : NamePropertyPair.Value.CopyRecords)
+				for (FPropertyCopyRecord& CopyRecord : NamePropertyPair.Value.CopyRecords)
 				{
-					if(CopyRecord.DestPropertyPath.Num() > 1)
+					if (CopyRecord.DestPropertyPath.Num() > 1)
 					{
 						// If this record writes to the node, switch it to the mutable data's property instead
-						if(CopyRecord.DestPropertyPath[0] == InRecord.NodeVariableProperty->GetName())
+						if (CopyRecord.DestPropertyPath[0] == InRecord.NodeVariableProperty->GetName())
 						{
 							CopyRecord.DestPropertyPath[0] = MutableDataProperty->GetName();
 
@@ -765,14 +808,14 @@ void UAnimBlueprintExtension_Base::RedirectPropertyAccesses(IAnimBlueprintCompil
 							FString DestPropertyPathWithoutArray = DestPropertyPathTail;
 							FString ArrayIndex;
 							int32 ArrayDelim = INDEX_NONE;
-							if(DestPropertyPathTail.FindChar(TEXT('['), ArrayDelim))
+							if (DestPropertyPathTail.FindChar(TEXT('['), ArrayDelim))
 							{
 								DestPropertyPathWithoutArray = DestPropertyPathTail.Left(ArrayDelim);
 								ArrayIndex = DestPropertyPathTail.RightChop(ArrayDelim);
 							}
 
 							// Switch the destination property from the node's property to the generated one
-							if(DestPropertyPathWithoutArray == FoldedPropertyRecord->Property->GetName())
+							if (DestPropertyPathWithoutArray == FoldedPropertyRecord->Property->GetName())
 							{
 								CopyRecord.DestPropertyPath[1] = FoldedPropertyRecord->GeneratedProperty->GetName() + ArrayIndex;
 							}
@@ -780,34 +823,50 @@ void UAnimBlueprintExtension_Base::RedirectPropertyAccesses(IAnimBlueprintCompil
 					}
 				}
 			}
-		}	
+		}
 	}
 }
 
-void UAnimBlueprintExtension_Base::FEvaluationHandlerRecord::PatchFunctionNameAndCopyRecordsInto(FExposedValueHandler& Handler) const
+void UAnimBlueprintExtension_Base::FEvaluationHandlerRecord::PatchAnimNodeExposedValueHandler(const UClass* InClass, IAnimBlueprintCompilationBracketContext& InCompilationContext) const
 {
-	Handler.CopyRecords.Empty();
+	UAnimBlueprintGeneratedClass* AnimClass = Cast<UAnimBlueprintGeneratedClass>(const_cast<UClass*>(InClass));
+	FStructProperty* HandlerProperty = CastFieldChecked<FStructProperty>(InCompilationContext.GetAllocatedHandlerPropertiesByNode().FindChecked(AnimGraphNode));
+	check(HandlerProperty->GetOwner<UScriptStruct>() == AnimClass->GetSparseClassDataStruct());
+	void* ConstantNodeData = const_cast<void*>(AnimClass->GetSparseClassData(EGetSparseClassDataMethod::ReturnIfNull));
+	check(ConstantNodeData);
 
-	for (const TPair<FName, FAnimNodeSinglePropertyHandler>& ServicedPropPair : ServicedProperties)
+	if (HandlerProperty->Struct == FAnimNodeExposedValueHandler_PropertyAccess::StaticStruct())
 	{
-		const FName& PropertyName = ServicedPropPair.Key;
-		const FAnimNodeSinglePropertyHandler& PropertyHandler = ServicedPropPair.Value;
+		FAnimNodeExposedValueHandler_PropertyAccess* Handler = HandlerProperty->ContainerPtrToValuePtr<FAnimNodeExposedValueHandler_PropertyAccess>(ConstantNodeData);
+
+		Handler->CopyRecords.Empty();
 
-		for (const FPropertyCopyRecord& PropertyCopyRecord : PropertyHandler.CopyRecords)
+		for (const TPair<FName, FAnimNodeSinglePropertyHandler>& ServicedPropPair : ServicedProperties)
 		{
-			// Only unbatched copies can be processed on a per-node basis
-			// Skip invalid copy indices as these are usually the result of BP errors/warnings
-			if(PropertyCopyRecord.LibraryCompiledHandle.IsValid() && PropertyCopyRecord.LibraryCompiledHandle.GetBatchId() == (int32)EAnimPropertyAccessCallSite::WorkerThread_Unbatched)
+			const FName& PropertyName = ServicedPropPair.Key;
+			const FAnimNodeSinglePropertyHandler& PropertyHandler = ServicedPropPair.Value;
+
+			for (const FPropertyCopyRecord& PropertyCopyRecord : PropertyHandler.CopyRecords)
 			{
-				Handler.CopyRecords.Emplace(PropertyCopyRecord.LibraryCompiledHandle.GetId(), PropertyCopyRecord.Operation);
+				// Only unbatched copies can be processed on a per-node basis
+				// Skip invalid copy indices as these are usually the result of BP errors/warnings
+				if (PropertyCopyRecord.LibraryCompiledHandle.IsValid() && PropertyCopyRecord.LibraryCompiledHandle.GetBatchId() == (int32)EAnimPropertyAccessCallSite::WorkerThread_Unbatched)
+				{
+					Handler->CopyRecords.Emplace(PropertyCopyRecord.LibraryCompiledHandle.GetId(), PropertyCopyRecord.Operation);
+				}
 			}
 		}
 	}
 
-	if (!IsFastPath())
+	if (HandlerProperty->Struct->IsChildOf(FAnimNodeExposedValueHandler_Base::StaticStruct()))
 	{
-		// not all of our pins use copy records so we will need to call our exposed value handler
-		Handler.BoundFunction = HandlerFunctionName;
+		FAnimNodeExposedValueHandler_Base* Handler = HandlerProperty->ContainerPtrToValuePtr<FAnimNodeExposedValueHandler_Base>(ConstantNodeData);
+
+		if (!IsFastPath())
+		{
+			// not all of our pins use copy records so we will need to call our exposed value handler
+			Handler->BoundFunction = HandlerFunctionName;
+		}
 	}
 }
 
@@ -877,7 +936,7 @@ static UEdGraphNode* FollowKnots(UEdGraphPin* FromPin, UEdGraphPin*& ToPin)
 	return nullptr;
 }
 
-void UAnimBlueprintExtension_Base::FEvaluationHandlerRecord::RegisterPin(UEdGraphPin* DestPin, FProperty* AssociatedProperty, int32 AssociatedPropertyArrayIndex)
+void UAnimBlueprintExtension_Base::FEvaluationHandlerRecord::RegisterPin(UEdGraphPin* DestPin, FProperty* AssociatedProperty, int32 AssociatedPropertyArrayIndex, bool bAllowFastPath)
 {
 	FAnimNodeSinglePropertyHandler& Handler = ServicedProperties.FindOrAdd(AssociatedProperty->GetFName());
 
@@ -899,6 +958,7 @@ void UAnimBlueprintExtension_Base::FEvaluationHandlerRecord::RegisterPin(UEdGrap
 	}
 
 	Handler.CopyRecords.Emplace(DestPin, AssociatedProperty, AssociatedPropertyArrayIndex, MoveTemp(DestPropertyPath));
+	Handler.CopyRecords.Last().bIsFastPath = bAllowFastPath;
 }
 
 void UAnimBlueprintExtension_Base::FEvaluationHandlerRecord::RegisterPropertyBinding(FProperty* InProperty, const FAnimGraphNodePropertyBinding& InBinding)
diff --git a/Engine/Source/Editor/AnimGraph/Private/AnimBlueprintExtension_PropertyAccess.cpp b/Engine/Source/Editor/AnimGraph/Private/AnimBlueprintExtension_PropertyAccess.cpp
--- a/Engine/Source/Editor/AnimGraph/Private/AnimBlueprintExtension_PropertyAccess.cpp
+++ b/Engine/Source/Editor/AnimGraph/Private/AnimBlueprintExtension_PropertyAccess.cpp
@@ -52,6 +52,15 @@ FPropertyAccessHandle UAnimBlueprintExtension_PropertyAccess::AddCopy(TArrayView
 	return FPropertyAccessHandle();
 }
 
+FPropertyAccessHandle UAnimBlueprintExtension_PropertyAccess::AddAccess(TArrayView<FString> InPath, UObject* InObject)
+{
+	if (PropertyAccessLibraryCompiler.IsValid())
+	{
+		return PropertyAccessLibraryCompiler->AddAccess(InPath, InObject);
+	}
+	return FPropertyAccessHandle();
+}
+
 void UAnimBlueprintExtension_PropertyAccess::HandleStartCompilingClass(const UClass* InClass, IAnimBlueprintCompilationBracketContext& InCompilationContext, IAnimBlueprintGeneratedClassCompiledData& OutCompiledData)
 {
 	IPropertyAccessEditor& PropertyAccessEditor = IModularFeatures::Get().GetModularFeature<IPropertyAccessEditor>("PropertyAccessEditor");
@@ -184,6 +193,15 @@ FCompiledPropertyAccessHandle UAnimBlueprintExtension_PropertyAccess::GetCompile
 	return FCompiledPropertyAccessHandle();
 }
 
+EPropertyAccessCopyType UAnimBlueprintExtension_PropertyAccess::GetCompiledHandleAccessType(FPropertyAccessHandle InHandle) const
+{
+	if (PropertyAccessLibraryCompiler.IsValid())
+	{
+		return PropertyAccessLibraryCompiler->GetCompiledHandleAccessType(InHandle);
+	}
+	return EPropertyAccessCopyType::None;
+}
+
 FText UAnimBlueprintExtension_PropertyAccess::GetCompiledHandleContext(FCompiledPropertyAccessHandle InHandle)
 {
 	UEnum* EnumClass = StaticEnum<EAnimPropertyAccessCallSite>();
diff --git a/Engine/Source/Editor/AnimGraph/Private/AnimGraphNode_ApplyMeshSpaceAdditive.cpp b/Engine/Source/Editor/AnimGraph/Private/AnimGraphNode_ApplyMeshSpaceAdditive.cpp
--- a/Engine/Source/Editor/AnimGraph/Private/AnimGraphNode_ApplyMeshSpaceAdditive.cpp
+++ b/Engine/Source/Editor/AnimGraph/Private/AnimGraphNode_ApplyMeshSpaceAdditive.cpp
@@ -144,7 +144,7 @@ void UAnimGraphNode_ApplyMeshSpaceAdditive::PostEditChangeProperty(struct FPrope
 			FScopedTransaction Transaction(LOCTEXT("ChangeAlphaInputType", "Change Alpha Input Type"));
 			Modify();
 			PinToBreakLinksTo->BreakAllPinLinks();
-			PropertyBindings.Remove(PinToBreakLinksTo->PinName);
+			RemoveBindings(PinToBreakLinksTo->PinName);
 			ReconstructNode();
 
 			FBlueprintEditorUtils::MarkBlueprintAsStructurallyModified(GetBlueprint());
diff --git a/Engine/Source/Editor/AnimGraph/Private/AnimGraphNode_AssetPlayerBase.cpp b/Engine/Source/Editor/AnimGraph/Private/AnimGraphNode_AssetPlayerBase.cpp
--- a/Engine/Source/Editor/AnimGraph/Private/AnimGraphNode_AssetPlayerBase.cpp
+++ b/Engine/Source/Editor/AnimGraph/Private/AnimGraphNode_AssetPlayerBase.cpp
@@ -387,7 +387,7 @@ FText UAnimGraphNode_AssetPlayerBase::GetNodeTitleHelper(ENodeTitleType::Type In
 		bool bHasBinding = false;
 		if(InAssetPin != nullptr)
 		{
-			if (PropertyBindings.Find(InAssetPin->GetFName()) != nullptr)
+			if (HasBinding(InAssetPin->GetFName()))
 			{
 				bHasBinding = true;
 			}
@@ -508,7 +508,7 @@ void UAnimGraphNode_AssetPlayerBase::ValidateAnimNodeDuringCompilationHelper(USk
 			bool bAlwaysDynamic = false;
 			if(InExposedPin != nullptr)
 			{
-				if(PropertyBindings.Contains(InExposedPin->GetFName()))
+				if(HasBinding(InExposedPin->GetFName()))
 				{
 					bHasBinding = true;
 				}
