AnimNext params in Anim BPs

diff --git a/Engine/Source/Editor/AnimGraph/Private/AnimGraphNode_Base.cpp b/Engine/Source/Editor/AnimGraph/Private/AnimGraphNode_Base.cpp
--- a/Engine/Source/Editor/AnimGraph/Private/AnimGraphNode_Base.cpp
+++ b/Engine/Source/Editor/AnimGraph/Private/AnimGraphNode_Base.cpp
@@ -42,6 +42,8 @@
 #include "Widgets/Text/STextBlock.h"
 #include "Fonts/FontMeasure.h"
 #include "ObjectEditorUtils.h"
+#include "UObject/FortniteMainBranchObjectVersion.h"
+#include "AnimGraphNodeBinding_Base.h"
 
 #define LOCTEXT_NAMESPACE "UAnimGraphNode_Base"
 
@@ -51,6 +53,19 @@
 UAnimGraphNode_Base::UAnimGraphNode_Base(const FObjectInitializer& ObjectInitializer)
 	: Super(ObjectInitializer)
 {
+	if(!HasAnyFlags(RF_ClassDefaultObject))
+	{
+		if(UAnimBlueprint* AnimBlueprint = GetTypedOuter<UAnimBlueprint>())
+		{
+			UClass* BindingClass = AnimBlueprint->GetDefaultBindingClass();
+			if (BindingClass == nullptr)
+			{
+				BindingClass = UAnimGraphNodeBinding_Base::StaticClass();
+			}
+
+			Binding = CastChecked<UAnimGraphNodeBinding>(ObjectInitializer.CreateDefaultSubobject(this, TEXT("Binding"), BindingClass, BindingClass));
+		}
+	}
 }
 
 void UAnimGraphNode_Base::ExpandNode(FKismetCompilerContext& CompilerContext, UEdGraph* SourceGraph)
@@ -58,6 +73,11 @@ void UAnimGraphNode_Base::ExpandNode(FKismetCompilerContext& CompilerContext, UE
 	TUniquePtr<IAnimBlueprintCompilationContext> CompilationContext = IAnimBlueprintCompilationContext::Get(CompilerContext);
 	UAnimBlueprintExtension_Base* Extension = UAnimBlueprintExtension_Base::GetExtension<UAnimBlueprintExtension_Base>(GetAnimBlueprint());
 	Extension->CreateEvaluationHandlerForNode(*CompilationContext.Get(), this);
+
+	if (Binding)
+	{
+		Binding->OnExpandNode(*CompilationContext.Get(), this, SourceGraph);
+	}
 }
 
 void UAnimGraphNode_Base::PreEditChange(FProperty* PropertyThatWillChange)
@@ -113,7 +133,12 @@ void UAnimGraphNode_Base::PostEditChangeProperty(FPropertyChangedEvent& Property
 		GetAnimBlueprint()->RequestRefreshExtensions();
 		FBlueprintEditorUtils::MarkBlueprintAsStructurallyModified(GetAnimBlueprint());
 	}
-	
+	else if (PropertyName == GET_MEMBER_NAME_CHECKED(UAnimGraphNode_Base, Binding))
+	{
+		GetAnimBlueprint()->RequestRefreshExtensions();
+		FBlueprintEditorUtils::MarkBlueprintAsStructurallyModified(GetAnimBlueprint());
+	}
+
 	Super::PostEditChangeProperty(PropertyChangedEvent);
 
 	PropertyChangeEvent.Broadcast(PropertyChangedEvent);
@@ -139,6 +164,7 @@ void UAnimGraphNode_Base::Serialize(FArchive& Ar)
 	Super::Serialize(Ar);
 
 	Ar.UsingCustomVersion(FReleaseObjectVersion::GUID);
+	Ar.UsingCustomVersion(FFortniteMainBranchObjectVersion::GUID);
 
 	if (Ar.IsLoading())
 	{
@@ -146,7 +172,7 @@ void UAnimGraphNode_Base::Serialize(FArchive& Ar)
 		{
 			// Push any bindings to optional pins
 			bool bPushedBinding = false;
-			for(const TPair<FName, FAnimGraphNodePropertyBinding>& BindingPair : PropertyBindings)
+			for(const TPair<FName, FAnimGraphNodePropertyBinding>& BindingPair : PropertyBindings_DEPRECATED)
 			{
 				for(FOptionalPinFromProperty& OptionalPin : ShowPinForProperties)
 				{
@@ -163,6 +189,14 @@ void UAnimGraphNode_Base::Serialize(FArchive& Ar)
 				FOptionalPinManager::EvaluateOldShownPins(ShowPinForProperties, OldShownPins, this);
 			}
 		}
+
+		if (Ar.CustomVer(FFortniteMainBranchObjectVersion::GUID) < FFortniteMainBranchObjectVersion::AnimGraphNodeBindingExtensions)
+		{
+			// Move internal bindings to extension
+			UAnimGraphNodeBinding_Base* NewBinding = NewObject<UAnimGraphNodeBinding_Base>(this);
+			NewBinding->PropertyBindings = PropertyBindings_DEPRECATED;
+			Binding = NewBinding;
+		}
 	}
 }
 
@@ -336,14 +370,9 @@ void UAnimGraphNode_Base::InternalPinCreation(TArray<UEdGraphPin*>* OldPins)
 		CreateOutputPins();
 	}
 
-	if(HasValidBlueprint())
+	if (Binding)
 	{
-		// Update any binding's display text
-		IPropertyAccessEditor& PropertyAccessEditor = IModularFeatures::Get().GetModularFeature<IPropertyAccessEditor>("PropertyAccessEditor");
-		for(TPair<FName, FAnimGraphNodePropertyBinding>& BindingPair : PropertyBindings)
-		{
-			BindingPair.Value.PathAsText = PropertyAccessEditor.MakeTextPath(BindingPair.Value.PropertyPath, GetBlueprint()->SkeletonGeneratedClass);
-		}
+		Binding->OnInternalPinCreation(this);
 	}
 }
 
@@ -354,54 +383,11 @@ void UAnimGraphNode_Base::AllocateDefaultPins()
 	CreateCustomPins(nullptr);
 }
 
-void UAnimGraphNode_Base::RecalculateBindingType(FAnimGraphNodePropertyBinding& InBinding)
-{
-	if(FProperty* BindingProperty = GetPinProperty(InBinding.PropertyName))
-	{
-		// Use the inner for array properties
-		if(BindingProperty->IsA<FArrayProperty>() && InBinding.ArrayIndex != INDEX_NONE)
-		{
-			BindingProperty = CastFieldChecked<FArrayProperty>(BindingProperty)->Inner;
-		}
-		
-		UAnimBlueprint* AnimBlueprint = GetAnimBlueprint();
-
-		IPropertyAccessEditor& PropertyAccessEditor = IModularFeatures::Get().GetModularFeature<IPropertyAccessEditor>("PropertyAccessEditor");
-		const UAnimationGraphSchema* Schema = GetDefault<UAnimationGraphSchema>();
-		
-		FProperty* LeafProperty = nullptr;
-		int32 ArrayIndex = INDEX_NONE;
-		FPropertyAccessResolveResult Result = PropertyAccessEditor.ResolvePropertyAccess(AnimBlueprint->SkeletonGeneratedClass, InBinding.PropertyPath, LeafProperty, ArrayIndex);
-		if(Result.Result == EPropertyAccessResolveResult::Succeeded)
-		{
-			if(LeafProperty)
-			{
-				Schema->ConvertPropertyToPinType(LeafProperty, InBinding.PinType);
-				
-				if(PropertyAccessEditor.GetPropertyCompatibility(LeafProperty, BindingProperty) == EPropertyAccessCompatibility::Promotable)
-				{
-					InBinding.bIsPromotion = true;
-					Schema->ConvertPropertyToPinType(LeafProperty, InBinding.PromotedPinType);
-				}
-				else
-				{
-					InBinding.bIsPromotion = false;
-					InBinding.PromotedPinType = InBinding.PinType;
-				}
-			}
-		}
-	}
-}
-
 void UAnimGraphNode_Base::ReconstructNode()
 {
-	if(HasValidBlueprint())
+	if (Binding)
 	{
-		// Refresh bindings
-		for(auto& BindingPair : PropertyBindings)
-		{
-			RecalculateBindingType(BindingPair.Value);
-		}
+		Binding->OnReconstructNode(this);
 	}
 	
 	Super::ReconstructNode();
@@ -639,8 +625,11 @@ void UAnimGraphNode_Base::ProcessDuringCompilation(IAnimBlueprintCompilationCont
 {
 	UAnimBlueprintExtension_Base* Extension = UAnimBlueprintExtension::GetExtension<UAnimBlueprintExtension_Base>(GetAnimBlueprint());
 
-	// Record pose pins for later patchup and gather pins that have an associated evaluation handler
-	Extension->ProcessNodePins(this, InCompilationContext, OutCompiledData);
+	// Record pose pins for later patchup
+	Extension->ProcessPosePins(this, InCompilationContext, OutCompiledData);
+
+	// Process bindings on this node
+	Binding->ProcessDuringCompilation(InCompilationContext, OutCompiledData);
 
 	// Resolve functions
 	GetFNode()->InitialUpdateFunction.SetFromFunction(InitialUpdateFunction.ResolveMember<UFunction>(GetBlueprintClassFromNode()));
@@ -847,14 +836,14 @@ void UAnimGraphNode_Base::AddPinSearchMetaDataInfo(const UEdGraphPin* InPin, TAr
 {
 	Super::AddPinSearchMetaDataInfo(InPin, OutTaggedMetaData);
 	
-	FName BindingName;
-	FProperty* PinProperty;
-	int32 OptionalPinIndex;
-	if(GetPinBindingInfo(InPin->GetFName(), BindingName, PinProperty, OptionalPinIndex))
+	if(Binding)
 	{
-		if(const FAnimGraphNodePropertyBinding* BindingInfo = PropertyBindings.Find(BindingName))
+		FName BindingName;
+		FProperty* PinProperty;
+		int32 OptionalPinIndex;
+		if (GetPinBindingInfo(InPin->GetFName(), BindingName, PinProperty, OptionalPinIndex))
 		{
-			OutTaggedMetaData.Add(FSearchTagDataPair(FText::FromString(TEXT("Binding")), BindingInfo->PathAsText));
+			Binding->AddPinSearchMetaDataInfo(InPin, BindingName, OutTaggedMetaData);
 		}
 	}
 }
@@ -868,7 +857,7 @@ bool UAnimGraphNode_Base::IsPinExposedAndLinked(const FString& InPinName, const
 bool UAnimGraphNode_Base::IsPinExposedAndBound(const FString& InPinName, const EEdGraphPinDirection InDirection) const
 {
 	UEdGraphPin* Pin = FindPin(InPinName, InDirection);
-	return Pin != nullptr && Pin->LinkedTo.Num() == 0 && PropertyBindings.Find(Pin->GetFName()) != nullptr;
+	return Pin != nullptr && Pin->LinkedTo.Num() == 0 && Binding && Binding->HasBinding(Pin->GetFName(), true);
 }
 
 bool UAnimGraphNode_Base::IsPinUnlinkedUnboundAndUnset(const FString& InPinName, const EEdGraphPinDirection InDirection) const
@@ -876,7 +865,7 @@ bool UAnimGraphNode_Base::IsPinUnlinkedUnboundAndUnset(const FString& InPinName,
 	UEdGraphPin* Pin = FindPin(InPinName, InDirection);
 	if (Pin)
 	{
-		return Pin->LinkedTo.Num() == 0 && PropertyBindings.Find(Pin->GetFName()) == nullptr && Pin->DefaultValue == Pin->AutogeneratedDefaultValue;
+		return Pin->LinkedTo.Num() == 0 && (Binding == nullptr || !Binding->HasBinding(Pin->GetFName(), true)) && Pin->DefaultValue == Pin->AutogeneratedDefaultValue;
 	}
 	return false;
 }
@@ -915,20 +904,22 @@ bool UAnimGraphNode_Base::GetPinBindingInfo(FName InPinName, FName& OutBindingNa
 
 bool UAnimGraphNode_Base::HasBinding(FName InBindingName) const
 {
-	// Comparison without name index to deal with arrays
-	const FName ComparisonName = FName(InBindingName, 0);
-
-	for(const TPair<FName, FAnimGraphNodePropertyBinding>& BindingPair : PropertyBindings)
+	if (Binding)
 	{
-		if(ComparisonName == FName(BindingPair.Key, 0))
-		{
-			return true;
-		}
+		return Binding->HasBinding(InBindingName, false);
 	}
 
 	return false;
 }
 
+void UAnimGraphNode_Base::RemoveBindings(FName InBindingName)
+{
+	if (Binding)
+	{
+		return Binding->RemoveBindings(InBindingName);
+	}
+}
+
 void UAnimGraphNode_Base::SetTag(FName InTag)
 {
 	FName OldTag = Tag;
@@ -959,17 +950,14 @@ void UAnimGraphNode_Base::PinConnectionListChanged(UEdGraphPin* Pin)
 {
 	if(Pin->LinkedTo.Num() > 0)
 	{
-		// If we have links, clear any bindings
-		// Compare FName without number to make sure we catch array properties that are split into multiple pins
-		FName ComparisonName = Pin->GetFName();
-		ComparisonName.SetNumber(0);
-
-		for(auto Iter = PropertyBindings.CreateIterator(); Iter; ++Iter)
+		if (Binding)
 		{
-			if(ComparisonName == FName(Iter.Key(), 0))
-			{
-				Iter.RemoveCurrent();
-			}
+			// If we have links, clear any bindings
+			// Compare FName without number to make sure we catch array properties that are split into multiple pins
+			FName ComparisonName = Pin->GetFName();
+			ComparisonName.SetNumber(0);
+
+			Binding->RemoveBindings(ComparisonName);
 		}
 	}
 }
@@ -1017,717 +1005,27 @@ void UAnimGraphNode_Base::AutowireNewNode(UEdGraphPin* FromPin)
 TSharedRef<SWidget> UAnimGraphNode_Base::MakePropertyBindingWidget(const FAnimPropertyBindingWidgetArgs& InArgs)
 {
 	UAnimGraphNode_Base* FirstAnimGraphNode = InArgs.Nodes[0];
-	const bool bMultiSelect = InArgs.Nodes.Num() > 1;
-	
-	if(FirstAnimGraphNode->HasValidBlueprint() && IModularFeatures::Get().IsModularFeatureAvailable("PropertyAccessEditor"))
-	{
-		UBlueprint* Blueprint = FirstAnimGraphNode->GetAnimBlueprint();
-		
-		int32 PinArrayIndex = InArgs.PinName.GetNumber() - 1;
-		const bool bIsArrayOrArrayElement = InArgs.PinProperty->IsA<FArrayProperty>();
-		const bool bIsArrayElement = bIsArrayOrArrayElement && PinArrayIndex != INDEX_NONE && InArgs.bPropertyIsOnFNode;
-		const bool bIsArray = bIsArrayOrArrayElement && PinArrayIndex == INDEX_NONE && InArgs.bPropertyIsOnFNode;
-
-		FProperty* PropertyToBindTo = bIsArrayElement ? CastField<FArrayProperty>(InArgs.PinProperty)->Inner : InArgs.PinProperty;
-		
-		// Properties could potentially be removed underneath this widget, so keep a TFieldPath reference to them
-		TFieldPath<FProperty> BindingPropertyPath(PropertyToBindTo);
-		TFieldPath<FProperty> PinPropertyPath(InArgs.PinProperty);
-		
-		auto OnCanBindProperty = [BindingPropertyPath](FProperty* InProperty)
-		{
-			// Note: We support type promotion here
-			IPropertyAccessEditor& PropertyAccessEditor = IModularFeatures::Get().GetModularFeature<IPropertyAccessEditor>("PropertyAccessEditor");
-			FProperty* BindingProperty = BindingPropertyPath.Get();
-			return BindingProperty && PropertyAccessEditor.GetPropertyCompatibility(InProperty, BindingProperty) != EPropertyAccessCompatibility::Incompatible;
-		};
-
-		auto OnCanBindFunction = [BindingPropertyPath](UFunction* InFunction)
-		{
-			IPropertyAccessEditor& PropertyAccessEditor = IModularFeatures::Get().GetModularFeature<IPropertyAccessEditor>("PropertyAccessEditor");
-			FProperty* BindingProperty = BindingPropertyPath.Get();
-			
-			// Note: We support type promotion here
-			return InFunction->NumParms == 1 
-				&& BindingProperty != nullptr && PropertyAccessEditor.GetPropertyCompatibility(InFunction->GetReturnProperty(), BindingProperty) != EPropertyAccessCompatibility::Incompatible
-				&& InFunction->HasAnyFunctionFlags(FUNC_BlueprintPure);
-		};
-
-		auto OnAddBinding = [InArgs, Blueprint, BindingPropertyPath, PinPropertyPath, bIsArrayElement, bIsArray](FName InPropertyName, const TArray<FBindingChainElement>& InBindingChain)
-		{
-			const UEdGraphSchema_K2* Schema = GetDefault<UEdGraphSchema_K2>();
-			IPropertyAccessEditor& PropertyAccessEditor = IModularFeatures::Get().GetModularFeature<IPropertyAccessEditor>("PropertyAccessEditor");
-
-			for(UAnimGraphNode_Base* AnimGraphNode : InArgs.Nodes)
-			{
-				AnimGraphNode->Modify();
-
-				// Reset to default so that references are not preserved
-				FProperty* BindingProperty = BindingPropertyPath.Get();
-				if(BindingProperty && BindingProperty->GetOwner<UStruct>() && AnimGraphNode->GetFNodeType()->IsChildOf(BindingProperty->GetOwner<UStruct>()) && BindingProperty->IsA<FObjectPropertyBase>())
-				{
-					void* PropertyAddress = BindingProperty->ContainerPtrToValuePtr<void>(AnimGraphNode->GetFNode());
-					BindingProperty->InitializeValue(PropertyAddress);
-				}
-				
-				// Pins are exposed if we have a binding or not - and after running this we do.
-				InArgs.OnSetPinVisibility.ExecuteIfBound(AnimGraphNode, true, InArgs.OptionalPinIndex);
-
-				// Need to break all pin links now we have a binding
-				if(UEdGraphPin* Pin = AnimGraphNode->FindPin(InArgs.PinName))
-				{
-					Pin->BreakAllPinLinks();
-				}
-
-				if(bIsArray)
-				{
-					// Remove bindings for array elements if this is an array
-					FName ComparisonName(InArgs.BindingName, 0);
-					for(auto Iter = AnimGraphNode->PropertyBindings.CreateIterator(); Iter; ++Iter)
-					{
-						if(ComparisonName == FName(Iter.Key(), 0))
-						{
-							Iter.RemoveCurrent();
-						}
-					}
-				}
-				else if(bIsArrayElement)
-				{
-					// If we are an array element, remove only whole-array bindings
-					FName ComparisonName(InArgs.BindingName, 0);
-					for(auto Iter = AnimGraphNode->PropertyBindings.CreateIterator(); Iter; ++Iter)
-					{
-						if(ComparisonName == Iter.Key())
-						{
-							Iter.RemoveCurrent();
-						}
-					}
-				}
-				
-				const FFieldVariant& LeafField = InBindingChain.Last().Field;
-
-				FProperty* PinProperty = PinPropertyPath.Get();
-				if(PinProperty && BindingProperty)
-				{
-					FAnimGraphNodePropertyBinding Binding;
-					Binding.PropertyName = InArgs.BindingName;
-					if(bIsArrayElement)
-					{
-						// Pull array index from the pin's FName if this is an array property
-						Binding.ArrayIndex = InArgs.PinName.GetNumber() - 1;
-					}
-					PropertyAccessEditor.MakeStringPath(InBindingChain, Binding.PropertyPath);
-					Binding.PathAsText = PropertyAccessEditor.MakeTextPath(Binding.PropertyPath, Blueprint->SkeletonGeneratedClass);
-					Binding.Type = LeafField.IsA<UFunction>() ? EAnimGraphNodePropertyBindingType::Function : EAnimGraphNodePropertyBindingType::Property;
-					Binding.bIsBound = true;
-					AnimGraphNode->RecalculateBindingType(Binding);
-
-					AnimGraphNode->PropertyBindings.Add(InArgs.BindingName, Binding);
-				}
-			}
-
-			FBlueprintEditorUtils::MarkBlueprintAsModified(Blueprint);
-		};
-
-		auto OnRemoveBinding = [InArgs, Blueprint](FName InPropertyName)
-		{
-			for(UAnimGraphNode_Base* AnimGraphNode : InArgs.Nodes)
-			{
-				AnimGraphNode->Modify();
-				AnimGraphNode->PropertyBindings.Remove(InArgs.BindingName);
-			}
-
-			FBlueprintEditorUtils::MarkBlueprintAsModified(Blueprint);
-		};
-
-		auto CanRemoveBinding = [InArgs](FName InPropertyName)
-		{
-			for(UAnimGraphNode_Base* AnimGraphNode : InArgs.Nodes)
-			{
-				if(AnimGraphNode->PropertyBindings.Contains(InArgs.BindingName))
-				{
-					return true;
-				}
-			}
-
-			return false;
-		}; 
-
-		enum class ECurrentValueType : int32
-		{
-			None,
-			Pin,
-			Binding,
-			Dynamic,
-			MultipleValues,
-		};
-
-		auto CurrentBindingText = [InArgs]()
-		{
-			ECurrentValueType CurrentValueType = ECurrentValueType::None;
-
-			const FText MultipleValues = LOCTEXT("MultipleValuesLabel", "Multiple Values");
-			const FText Bind = LOCTEXT("BindLabel", "Bind");
-			const FText ExposedAsPin = LOCTEXT("ExposedAsPinLabel", "Pin");
-			const FText Dynamic = LOCTEXT("DynamicLabel", "Dynamic");
-			FText CurrentValue = Bind;
-
-			auto SetAssignValue = [&CurrentValueType, &CurrentValue, &MultipleValues](const FText& InValue, ECurrentValueType InType)
-			{
-				if(CurrentValueType != ECurrentValueType::MultipleValues)
-				{
-					if(CurrentValueType == ECurrentValueType::None)
-					{
-						CurrentValueType = InType;
-						CurrentValue = InValue;
-					}
-					else if(CurrentValueType == InType)
-					{
-						if(!CurrentValue.EqualTo(InValue))
-						{
-							CurrentValueType = ECurrentValueType::MultipleValues;
-							CurrentValue = MultipleValues;
-						}
-					}
-					else
-					{
-						CurrentValueType = ECurrentValueType::MultipleValues;
-						CurrentValue = MultipleValues;
-					}
-				}
-			};
-
-			const FName ComparisonName = FName(InArgs.PinName, 0);
-
-			for(UAnimGraphNode_Base* AnimGraphNode : InArgs.Nodes)
-			{
-				if(FAnimGraphNodePropertyBinding* BindingPtr = AnimGraphNode->PropertyBindings.Find(InArgs.BindingName))
-				{
-					SetAssignValue(BindingPtr->PathAsText, ECurrentValueType::Binding);
-				}
-				else if(AnimGraphNode->AlwaysDynamicProperties.Find(ComparisonName))
-				{
-					SetAssignValue(Dynamic, ECurrentValueType::Dynamic);
-				}
-				else
-				{
-					TArrayView<FOptionalPinFromProperty> OptionalPins; 
-					InArgs.OnGetOptionalPins.ExecuteIfBound(AnimGraphNode, OptionalPins);
-					if(OptionalPins.IsValidIndex(InArgs.OptionalPinIndex) && OptionalPins[InArgs.OptionalPinIndex].bShowPin)
-					{
-						SetAssignValue(InArgs.bOnGraphNode ? Bind : ExposedAsPin, ECurrentValueType::Pin);
-					}
-					else
-					{
-						SetAssignValue(Bind, ECurrentValueType::None);
-					}
-				}
-			}
-
-			return CurrentValue;
-		};
-
-		auto CurrentBindingToolTipText = [InArgs]()
-		{
-			ECurrentValueType CurrentValueType = ECurrentValueType::None;
-
-			const FText MultipleValues = LOCTEXT("MultipleValuesToolTip", "Bindings Have Multiple Values");
-			const FText ExposedAsPin = LOCTEXT("ExposedAsPinToolTip", "Exposed As a Pin on the Node");
-			const FText BindValue = LOCTEXT("BindValueToolTip", "Bind This Value");
-			const FText DynamicValue = LOCTEXT("DynamicValueToolTip", "Dynamic value that can be set externally");
-			FText CurrentValue;
-			
-			auto SetAssignValue = [&CurrentValueType, &CurrentValue, &MultipleValues](const FText& InValue, ECurrentValueType InType)
-			{
-				if(CurrentValueType != ECurrentValueType::MultipleValues)
-				{
-					if(CurrentValueType == ECurrentValueType::None)
-					{
-						CurrentValueType = InType;
-						CurrentValue = InValue;
-					}
-					else if(CurrentValueType == InType)
-					{
-						if(!CurrentValue.EqualTo(InValue))
-						{
-							CurrentValueType = ECurrentValueType::MultipleValues;
-							CurrentValue = MultipleValues;
-						}
-					}
-					else
-					{
-						CurrentValueType = ECurrentValueType::MultipleValues;
-						CurrentValue = MultipleValues;
-					}
-				}
-			};
-
-			const FName ComparisonName = FName(InArgs.PinName, 0);
-
-			for(UAnimGraphNode_Base* AnimGraphNode : InArgs.Nodes)
-			{
-				if(FAnimGraphNodePropertyBinding* BindingPtr = AnimGraphNode->PropertyBindings.Find(InArgs.BindingName))
-				{
-					if(BindingPtr->PathAsText.IsEmpty())
-					{
-						SetAssignValue(BindValue, ECurrentValueType::Binding);
-					}
-					else
-					{
-						IPropertyAccessEditor& PropertyAccessEditor = IModularFeatures::Get().GetModularFeature<IPropertyAccessEditor>("PropertyAccessEditor");
-						const FText UnderlyingPath = PropertyAccessEditor.MakeTextPath(BindingPtr->PropertyPath);
-						const FText& CompilationContext = BindingPtr->CompiledContext;
-						const FText& CompilationContextDesc = BindingPtr->CompiledContextDesc;
-						if(CompilationContext.IsEmpty() && CompilationContextDesc.IsEmpty())
-						{
-							SetAssignValue(FText::Format(LOCTEXT("BindingToolTipFormat", "Pin is bound to property '{0}'\nNative: {1}"), BindingPtr->PathAsText, UnderlyingPath), ECurrentValueType::Binding);
-						}
-						else
-						{
-							SetAssignValue(FText::Format(LOCTEXT("BindingToolTipFormatWithDesc", "Pin is bound to property '{0}'\nNative: {1}\n{2}\n{2}"), BindingPtr->PathAsText, UnderlyingPath, CompilationContext, CompilationContextDesc), ECurrentValueType::Binding);
-						}
-					}
-				}
-				else if(AnimGraphNode->AlwaysDynamicProperties.Find(ComparisonName))
-				{
-					SetAssignValue(DynamicValue, ECurrentValueType::Dynamic);
-				}
-				else
-				{
-					TArrayView<FOptionalPinFromProperty> OptionalPins; 
-					InArgs.OnGetOptionalPins.ExecuteIfBound(AnimGraphNode, OptionalPins);
-					if(OptionalPins.IsValidIndex(InArgs.OptionalPinIndex) && OptionalPins[InArgs.OptionalPinIndex].bShowPin)
-					{
-						SetAssignValue(InArgs.bOnGraphNode ? BindValue : ExposedAsPin, ECurrentValueType::Pin);
-					}
-					else
-					{
-						SetAssignValue(BindValue, ECurrentValueType::None);
-					}
-				}
-			}
-
-			return CurrentValue;
-		};
-
-		auto CurrentBindingImage = [InArgs, BindingPropertyPath]() -> const FSlateBrush*
-		{
-			static FName PropertyIcon(TEXT("Kismet.VariableList.TypeIcon"));
-			static FName FunctionIcon(TEXT("GraphEditor.Function_16x"));
+	UClass* BindingClass = FirstAnimGraphNode->Binding ? FirstAnimGraphNode->Binding->GetClass() : nullptr;
+	UAnimGraphNodeBinding* BindingCDO = BindingClass ? CastChecked<UAnimGraphNodeBinding>(BindingClass->GetDefaultObject()) : nullptr;
 
-			EAnimGraphNodePropertyBindingType BindingType = EAnimGraphNodePropertyBindingType::None;
-			for(UObject* OuterObject : InArgs.Nodes)
-			{
-				if(UAnimGraphNode_Base* AnimGraphNode = Cast<UAnimGraphNode_Base>(OuterObject))
-				{
-					TArrayView<FOptionalPinFromProperty> OptionalPins; 
-					InArgs.OnGetOptionalPins.ExecuteIfBound(AnimGraphNode, OptionalPins);
-					if(OptionalPins.IsValidIndex(InArgs.OptionalPinIndex) && OptionalPins[InArgs.OptionalPinIndex].bShowPin)
-					{
-						BindingType = EAnimGraphNodePropertyBindingType::None;
-						break;
-					}
-					else if(FAnimGraphNodePropertyBinding* BindingPtr = AnimGraphNode->PropertyBindings.Find(InArgs.BindingName))
-					{
-						if(BindingType == EAnimGraphNodePropertyBindingType::None)
-						{
-							BindingType = BindingPtr->Type;
-						}
-						else if(BindingType != BindingPtr->Type)
-						{
-							BindingType = EAnimGraphNodePropertyBindingType::None;
-							break;
-						}
-					}
-					else if(BindingType != EAnimGraphNodePropertyBindingType::None)
-					{
-						BindingType = EAnimGraphNodePropertyBindingType::None;
-						break;
-					}
-				}
-			}
-
-			if (BindingType == EAnimGraphNodePropertyBindingType::Function)
-			{
-				return FAppStyle::GetBrush(FunctionIcon);
-			}
-			else
-			{
-				const UAnimationGraphSchema* AnimationGraphSchema = GetDefault<UAnimationGraphSchema>();
-				FProperty* BindingProperty = BindingPropertyPath.Get();
-				FEdGraphPinType PinType;
-				if(BindingProperty != nullptr && AnimationGraphSchema->ConvertPropertyToPinType(BindingProperty, PinType))
-				{
-					return FBlueprintEditorUtils::GetIconFromPin(PinType, false);
-				}
-				else
-				{
-					return FAppStyle::GetBrush(PropertyIcon);
-				}
-			}
-		};
-
-		auto CurrentBindingColor = [InArgs, BindingPropertyPath]() -> FLinearColor
-		{
-			const UEdGraphSchema_K2* Schema = GetDefault<UEdGraphSchema_K2>();
-
-			FLinearColor BindingColor = FLinearColor::Gray;
-
-			FProperty* BindingProperty = BindingPropertyPath.Get();
-			FEdGraphPinType PinType;
-			if(BindingProperty != nullptr && Schema->ConvertPropertyToPinType(BindingProperty, PinType))
-			{
-				BindingColor = Schema->GetPinTypeColor(PinType);
-
-				enum class EPromotionState
-				{
-					NotChecked,
-					NotPromoted,
-					Promoted,
-				} Promotion = EPromotionState::NotChecked;
-
-				for(UAnimGraphNode_Base* AnimGraphNode : InArgs.Nodes)
-				{
-					TArrayView<FOptionalPinFromProperty> OptionalPins; 
-					InArgs.OnGetOptionalPins.ExecuteIfBound(AnimGraphNode, OptionalPins);
-					if(FAnimGraphNodePropertyBinding* BindingPtr = AnimGraphNode->PropertyBindings.Find(InArgs.BindingName))
-					{
-						if(Promotion == EPromotionState::NotChecked)
-						{
-							if(BindingPtr->bIsPromotion)
-							{
-								Promotion = EPromotionState::Promoted;
-								BindingColor = Schema->GetPinTypeColor(BindingPtr->PromotedPinType);
-							}
-							else
-							{
-								Promotion = EPromotionState::NotPromoted;
-							}
-						}
-						else
-						{
-							EPromotionState NewPromotion = BindingPtr->bIsPromotion ? EPromotionState::Promoted : EPromotionState::NotPromoted;
-							if(Promotion != NewPromotion)
-							{
-								BindingColor = FLinearColor::Gray;
-								break;
-							}
-						}
-					}
-					else if(OptionalPins.IsValidIndex(InArgs.OptionalPinIndex) && OptionalPins[InArgs.OptionalPinIndex].bShowPin)
-					{
-						if(Promotion == EPromotionState::NotChecked)
-						{
-							Promotion = EPromotionState::NotPromoted;
-						}
-						else if(Promotion == EPromotionState::Promoted)
-						{
-							BindingColor = FLinearColor::Gray;
-							break;
-						}
-					}
-				}
-			}
-
-			return BindingColor;
-		};
-
-		auto AddMenuExtension = [InArgs, Blueprint, BindingPropertyPath, PinPropertyPath](FMenuBuilder& InMenuBuilder)
-		{
-			InMenuBuilder.BeginSection("Pins", LOCTEXT("Pin", "Pin"));
-			{
-				auto ExposeAsPin = [InArgs, Blueprint]()
-				{
-					bool bHasBinding = false;
-
-					for(UAnimGraphNode_Base* AnimGraphNode : InArgs.Nodes)
-					{
-						if(AnimGraphNode->HasBinding(InArgs.BindingName))
-						{
-							bHasBinding = true;
-							break;
-						}
-					}
-
-					{
-						FScopedTransaction Transaction(LOCTEXT("PinExposure", "Pin Exposure"));
-
-						// Switching from non-pin to pin, remove any bindings
-						for(UAnimGraphNode_Base* AnimGraphNode : InArgs.Nodes)
-						{
-							AnimGraphNode->Modify();
-
-							TArrayView<FOptionalPinFromProperty> OptionalPins; 
-							InArgs.OnGetOptionalPins.ExecuteIfBound(AnimGraphNode, OptionalPins);
-							const bool bVisible = OptionalPins.IsValidIndex(InArgs.OptionalPinIndex) && OptionalPins[InArgs.OptionalPinIndex].bShowPin;
-							InArgs.OnSetPinVisibility.ExecuteIfBound(AnimGraphNode, !bVisible || bHasBinding, InArgs.OptionalPinIndex);
-
-							// Remove all bindings that match the property, array or array elements
-							const FName ComparisonName = FName(InArgs.BindingName, 0);
-							for(auto It = AnimGraphNode->PropertyBindings.CreateIterator(); It; ++It)
-							{
-								if(ComparisonName == FName(It.Key(), 0))
-								{
-									It.RemoveCurrent();
-								}
-							}
-						}
-
-						FBlueprintEditorUtils::MarkBlueprintAsModified(Blueprint);
-					}
-				}; 
-
-				auto GetExposedAsPinCheckState = [InArgs]()
-				{
-					bool bPinShown = false;
-					bool bHasBinding = false;
-
-					for(UAnimGraphNode_Base* AnimGraphNode : InArgs.Nodes)
-					{
-						if(AnimGraphNode->HasBinding(InArgs.BindingName))
-						{
-							bHasBinding = true;
-							break;
-						}
-
-						TArrayView<FOptionalPinFromProperty> OptionalPins; 
-						InArgs.OnGetOptionalPins.ExecuteIfBound(AnimGraphNode, OptionalPins);
-						bPinShown |= OptionalPins.IsValidIndex(InArgs.OptionalPinIndex) && OptionalPins[InArgs.OptionalPinIndex].bShowPin;
-					}
-
-					// Pins are exposed if we have a binding or not, so treat as unchecked only if we have
-					// no binding
-					return bPinShown && !bHasBinding ? ECheckBoxState::Checked : ECheckBoxState::Unchecked;
-				};
-				
-				InMenuBuilder.AddMenuEntry(
-					LOCTEXT("ExposeAsPin", "Expose As Pin"),
-					LOCTEXT("ExposeAsPinTooltip", "Show/hide this property as a pin on the node"),
-					FSlateIcon(FAppStyle::GetAppStyleSetName(), "GraphEditor.PinIcon"),
-					FUIAction(
-						FExecuteAction::CreateLambda(ExposeAsPin),
-						FCanExecuteAction(),
-						FGetActionCheckState::CreateLambda(GetExposedAsPinCheckState)
-					),
-					NAME_None,
-					EUserInterfaceActionType::Check
-				);
-
-				if(InArgs.bPropertyIsOnFNode)
-				{
-					auto MakeDynamic = [InArgs, Blueprint]()
-					{
-						// Comparison without name index to deal with arrays
-						const FName ComparisonName = FName(InArgs.PinName, 0);
-
-						bool bIsAlwaysDynamic = false;
-						for(UAnimGraphNode_Base* AnimGraphNode : InArgs.Nodes)
-						{
-							if(AnimGraphNode->AlwaysDynamicProperties.Contains(ComparisonName))
-							{
-								bIsAlwaysDynamic = true;
-								break;
-							}
-						}
-
-						{
-							FScopedTransaction Transaction(LOCTEXT("AlwaysDynamic", "Always Dynamic"));
-
-							for(UAnimGraphNode_Base* AnimGraphNode : InArgs.Nodes)
-							{
-								AnimGraphNode->Modify();
-
-								if(bIsAlwaysDynamic)
-								{
-									AnimGraphNode->AlwaysDynamicProperties.Remove(ComparisonName);
-								}
-								else
-								{
-									AnimGraphNode->AlwaysDynamicProperties.Add(ComparisonName);
-								}
-							}
-
-							FBlueprintEditorUtils::MarkBlueprintAsModified(Blueprint);
-						}
-					};
-
-					auto GetDynamicCheckState = [InArgs]()
-					{
-						// Comparison without name index to deal with arrays
-						const FName ComparisonName = FName(InArgs.PinName, 0);
-
-						bool bIsAlwaysDynamic = false;
-						for(UAnimGraphNode_Base* AnimGraphNode : InArgs.Nodes)
-						{
-							for(const FName& AlwaysDynamicPropertyName : AnimGraphNode->AlwaysDynamicProperties)
-							{
-								if(ComparisonName == FName(AlwaysDynamicPropertyName, 0))
-								{
-									bIsAlwaysDynamic = true;
-									break;
-								}
-							}
-						}
-						
-						return bIsAlwaysDynamic ? ECheckBoxState::Checked : ECheckBoxState::Unchecked;
-					};
-					
-					InMenuBuilder.AddMenuEntry(
-						LOCTEXT("DynamicValue", "Dynamic Value"),
-						LOCTEXT("DynamicValueTooltip", "Flag this value as dynamic. This way it can be set from functions even when not exposed as a pin."),
-						FSlateIcon(),
-						FUIAction(
-							FExecuteAction::CreateLambda(MakeDynamic),
-							FCanExecuteAction(),
-							FGetActionCheckState::CreateLambda(GetDynamicCheckState)
-						),
-						NAME_None,
-						EUserInterfaceActionType::Check
-					);
-				}
-			}
-			InMenuBuilder.EndSection();
-		};
-
-		FPropertyBindingWidgetArgs Args;
-		Args.Property = PropertyToBindTo;
-		Args.OnCanBindProperty = FOnCanBindProperty::CreateLambda(OnCanBindProperty);
-		Args.OnCanBindFunction = FOnCanBindFunction::CreateLambda(OnCanBindFunction);
-		Args.OnCanBindToClass = FOnCanBindToClass::CreateLambda([](UClass* InClass){ return true; });
-		Args.OnAddBinding = FOnAddBinding::CreateLambda(OnAddBinding);
-		Args.OnRemoveBinding = FOnRemoveBinding::CreateLambda(OnRemoveBinding);
-		Args.OnCanRemoveBinding = FOnCanRemoveBinding::CreateLambda(CanRemoveBinding);
-		Args.CurrentBindingText = MakeAttributeLambda(CurrentBindingText);
-		Args.CurrentBindingToolTipText = MakeAttributeLambda(CurrentBindingToolTipText);
-		Args.CurrentBindingImage = MakeAttributeLambda(CurrentBindingImage);
-		Args.CurrentBindingColor = MakeAttributeLambda(CurrentBindingColor);
-		Args.MenuExtender = MakeShared<FExtender>();
-		Args.MenuExtender->AddMenuExtension("BindingActions", EExtensionHook::Before, nullptr, FMenuExtensionDelegate::CreateLambda(AddMenuExtension));
-		if(InArgs.MenuExtender.IsValid())
-		{
-			Args.MenuExtender = FExtender::Combine( { Args.MenuExtender, InArgs.MenuExtender } );
-		}
-
-		IPropertyAccessBlueprintBinding::FContext BindingContext;
-		BindingContext.Blueprint = FBlueprintEditorUtils::FindBlueprintForNode(FirstAnimGraphNode);
-		BindingContext.Graph = FirstAnimGraphNode->GetGraph();
-		BindingContext.Node = FirstAnimGraphNode;
-		BindingContext.Pin = FirstAnimGraphNode->FindPin(InArgs.PinName);
-
-		auto OnSetPropertyAccessContextId = [InArgs, Blueprint](const FName& InContextId)
-		{
-			for(UAnimGraphNode_Base* AnimGraphNode : InArgs.Nodes)
-			{
-				if(FAnimGraphNodePropertyBinding* Binding = AnimGraphNode->PropertyBindings.Find(InArgs.BindingName))
-				{
-					AnimGraphNode->Modify();
-					Binding->ContextId = InContextId;
-				}
-			}
-
-			FBlueprintEditorUtils::MarkBlueprintAsModified(Blueprint);
-		};
-
-		auto OnCanSetPropertyAccessContextId = [InArgs, FirstAnimGraphNode](const FName& InContextId)
-		{
-			return FirstAnimGraphNode->PropertyBindings.Find(InArgs.BindingName) != nullptr;
-		};
-		
-		auto OnGetPropertyAccessContextId = [InArgs]() -> FName
-		{
-			FName CurrentContext = NAME_None;
-			for(UAnimGraphNode_Base* AnimGraphNode : InArgs.Nodes)
-			{
-				if(FAnimGraphNodePropertyBinding* Binding = AnimGraphNode->PropertyBindings.Find(InArgs.BindingName))
-				{
-					if(CurrentContext != NAME_None && CurrentContext != Binding->ContextId)
-					{
-						return NAME_None;
-					}
-					else
-					{
-						CurrentContext = Binding->ContextId;
-					}
-				}
-			}
-		
-			return CurrentContext;
-		};
-		
-		IPropertyAccessBlueprintBinding::FBindingMenuArgs MenuArgs;
-		MenuArgs.OnSetPropertyAccessContextId = FOnSetPropertyAccessContextId::CreateLambda(OnSetPropertyAccessContextId);
-		MenuArgs.OnCanSetPropertyAccessContextId = FOnCanSetPropertyAccessContextId::CreateLambda(OnCanSetPropertyAccessContextId);
-		MenuArgs.OnGetPropertyAccessContextId = FOnGetPropertyAccessContextId::CreateLambda(OnGetPropertyAccessContextId);
-		
-		// Add the binding menu extenders
-		TArray<TSharedPtr<FExtender>> Extenders( { Args.MenuExtender } );
-		for(IPropertyAccessBlueprintBinding* Binding : IModularFeatures::Get().GetModularFeatureImplementations<IPropertyAccessBlueprintBinding>("PropertyAccessBlueprintBinding"))
+	// Check all nodes have the same kinds of binding - we cant do multi-edit with bindings of differing types
+	if (InArgs.Nodes.Num() > 0)
+	{
+		for (int32 NodeIndex = 1; NodeIndex < InArgs.Nodes.Num(); ++NodeIndex)
 		{
-			TSharedPtr<FExtender> BindingExtender = Binding->MakeBindingMenuExtender(BindingContext, MenuArgs);
-			if(BindingExtender)
+			UAnimGraphNode_Base* OtherAnimGraphNode = InArgs.Nodes[NodeIndex];
+			UClass* OtherBindingClass = OtherAnimGraphNode->Binding ? OtherAnimGraphNode->Binding->GetClass() : nullptr;
+			if (OtherBindingClass == nullptr || BindingClass != OtherBindingClass)
 			{
-				Extenders.Add(BindingExtender);
+				BindingCDO = nullptr;
+				break;
 			}
 		}
-		
-		if(Extenders.Num() > 0)
-		{
-			Args.MenuExtender = FExtender::Combine(Extenders);
-		}
-
-		Args.bAllowNewBindings = false;
-		Args.bAllowArrayElementBindings = !bIsArray;
-		Args.bAllowUObjectFunctions = !bIsArray;
-		Args.bAllowStructFunctions = !bIsArray;
-
-		IPropertyAccessEditor& PropertyAccessEditor = IModularFeatures::Get().GetModularFeature<IPropertyAccessEditor>("PropertyAccessEditor");
-
-		const FTextBlockStyle& TextBlockStyle = FAppStyle::Get().GetWidgetStyle<FTextBlockStyle>("PropertyAccess.CompiledContext.Text");
-		
-		return
-			SNew(SBox)
-			.MaxDesiredWidth(200.0f)
-			[
-				SNew(SOverlay)
-				+SOverlay::Slot()
-				.HAlign(HAlign_Right)
-				.VAlign(VAlign_Bottom)
-				[
-					SNew(SBorder)
-					.Padding(FMargin(1.0f, 3.0f, 1.0f, 1.0f))
-					.Visibility(InArgs.bOnGraphNode ? EVisibility::Visible : EVisibility::Collapsed)
-					.BorderImage(FAppStyle::GetBrush("PropertyAccess.CompiledContext.Border"))
-					.RenderTransform_Lambda([InArgs, FirstAnimGraphNode, &TextBlockStyle]()
-					{
-						const FAnimGraphNodePropertyBinding* BindingPtr = FirstAnimGraphNode->PropertyBindings.Find(InArgs.BindingName);
-						FVector2D TextSize(0.0f, 0.0f);
-						if(BindingPtr)
-						{
-							const TSharedRef<FSlateFontMeasure> FontMeasureService = FSlateApplication::Get().GetRenderer()->GetFontMeasureService();
-							TextSize = FontMeasureService->Measure(BindingPtr->CompiledContext, TextBlockStyle.Font);
-						}
-						return FSlateRenderTransform(FVector2D(0.0f, TextSize.Y - 1.0f));
-					})	
-					[
-						SNew(STextBlock)
-						.TextStyle(&TextBlockStyle)
-						.Visibility_Lambda([InArgs, FirstAnimGraphNode, bMultiSelect]()
-						{
-							const FAnimGraphNodePropertyBinding* BindingPtr = FirstAnimGraphNode->PropertyBindings.Find(InArgs.BindingName);
-							return bMultiSelect || BindingPtr == nullptr || BindingPtr->CompiledContext.IsEmpty() ? EVisibility::Collapsed : EVisibility::Visible;
-						})
-						.Text_Lambda([InArgs, FirstAnimGraphNode]()
-						{
-							const FAnimGraphNodePropertyBinding* BindingPtr = FirstAnimGraphNode->PropertyBindings.Find(InArgs.BindingName);
-							return BindingPtr != nullptr ? BindingPtr->CompiledContext : FText::GetEmpty();
-						})
-					]
-				]
-				+SOverlay::Slot()
-				[
-					PropertyAccessEditor.MakePropertyBindingWidget(Blueprint, Args)
-				]	
-			];
+	}
+	
+	if (BindingCDO)
+	{
+		return BindingCDO->MakePropertyBindingWidget(InArgs);
 	}
 	else
 	{
@@ -1737,145 +1035,35 @@ TSharedRef<SWidget> UAnimGraphNode_Base::MakePropertyBindingWidget(const FAnimPr
 
 void UAnimGraphNode_Base::HandleVariableRenamed(UBlueprint* InBlueprint, UClass* InVariableClass, UEdGraph* InGraph, const FName& InOldVarName, const FName& InNewVarName)
 {
-	IPropertyAccessEditor& PropertyAccessEditor = IModularFeatures::Get().GetModularFeature<IPropertyAccessEditor>("PropertyAccessEditor");
-
-	UClass* SkeletonVariableClass = FBlueprintEditorUtils::GetSkeletonClass(InVariableClass);
-	
-	// See if any of bindings reference the variable
-	for(auto& BindingPair : PropertyBindings)
+	if (Binding)
 	{
-		TArray<int32> RenameIndices;
-		IPropertyAccessEditor::FResolvePropertyAccessArgs ResolveArgs;
-		ResolveArgs.PropertyFunction = [InOldVarName, SkeletonVariableClass, &RenameIndices](int32 InSegmentIndex, FProperty* InProperty, int32 InStaticArrayIndex)
-		{
-			UClass* OwnerClass = InProperty->GetOwnerClass();
-			if(OwnerClass && InProperty->GetFName() == InOldVarName && OwnerClass->IsChildOf(SkeletonVariableClass))
-			{
-				RenameIndices.Add(InSegmentIndex);
-			}
-		};
-		
-		PropertyAccessEditor.ResolvePropertyAccess(GetAnimBlueprint()->SkeletonGeneratedClass, BindingPair.Value.PropertyPath, ResolveArgs);
-
-		// Rename any references we found
-		for(const int32& RenameIndex : RenameIndices)
-		{
-			BindingPair.Value.PropertyPath[RenameIndex] = InNewVarName.ToString();
-			BindingPair.Value.PathAsText = PropertyAccessEditor.MakeTextPath(BindingPair.Value.PropertyPath, GetAnimBlueprint()->SkeletonGeneratedClass);
-		}
+		Binding->HandleVariableRenamed(InBlueprint, InVariableClass, InGraph, InOldVarName, InNewVarName);
 	}
 }
 
 void UAnimGraphNode_Base::ReplaceReferences(UBlueprint* InBlueprint, UBlueprint* InReplacementBlueprint, const FMemberReference& InSource, const FMemberReference& InReplacement)
 {
-	IPropertyAccessEditor& PropertyAccessEditor = IModularFeatures::Get().GetModularFeature<IPropertyAccessEditor>("PropertyAccessEditor");
-	
-	UClass* SkeletonClass = InBlueprint->SkeletonGeneratedClass;
-
-	FMemberReference Source = InSource;
-	FProperty* SourceProperty = Source.ResolveMember<FProperty>(InBlueprint);
-	FMemberReference Replacement = InReplacement;
-	FProperty* ReplacementProperty = Replacement.ResolveMember<FProperty>(InReplacementBlueprint);
-	
-	// See if any of bindings reference the variable
-	for(auto& BindingPair : PropertyBindings)
+	if (Binding)
 	{
-		TArray<int32> ReplaceIndices;
-		IPropertyAccessEditor::FResolvePropertyAccessArgs ResolveArgs;
-		ResolveArgs.PropertyFunction = [SourceProperty, &ReplaceIndices](int32 InSegmentIndex, FProperty* InProperty, int32 InStaticArrayIndex)
-		{
-			if(InProperty == SourceProperty)
-			{
-				ReplaceIndices.Add(InSegmentIndex);
-			}
-		};
-		
-		PropertyAccessEditor.ResolvePropertyAccess(GetAnimBlueprint()->SkeletonGeneratedClass, BindingPair.Value.PropertyPath, ResolveArgs);
-
-		// Replace any references we found
-		for(const int32& RenameIndex : ReplaceIndices)
-		{
-			BindingPair.Value.PropertyPath[RenameIndex] = ReplacementProperty->GetName();
-			BindingPair.Value.PathAsText = PropertyAccessEditor.MakeTextPath(BindingPair.Value.PropertyPath, GetAnimBlueprint()->SkeletonGeneratedClass);
-		}
+		Binding->ReplaceReferences(InBlueprint, InReplacementBlueprint, InSource, InReplacement);
 	}
 }
 
 bool UAnimGraphNode_Base::ReferencesVariable(const FName& InVarName, const UStruct* InScope) const
 {
-	IPropertyAccessEditor& PropertyAccessEditor = IModularFeatures::Get().GetModularFeature<IPropertyAccessEditor>("PropertyAccessEditor");
-
-	const UClass* SkeletonVariableClass = FBlueprintEditorUtils::GetSkeletonClass(Cast<UClass>(InScope));
-	
-	// See if any of bindings reference the variable
-	for(const auto& BindingPair : PropertyBindings)
+	if (Binding)
 	{
-		bool bReferencesVariable = false;
-		
-		IPropertyAccessEditor::FResolvePropertyAccessArgs ResolveArgs;
-		ResolveArgs.PropertyFunction = [InVarName, SkeletonVariableClass, &bReferencesVariable](int32 InSegmentIndex, FProperty* InProperty, int32 InStaticArrayIndex)
-		{
-			if(SkeletonVariableClass)
-			{
-				const UClass* OwnerSkeletonVariableClass = FBlueprintEditorUtils::GetSkeletonClass(Cast<UClass>(InProperty->GetOwnerStruct()));
-
-				if(OwnerSkeletonVariableClass && InProperty->GetFName() == InVarName && OwnerSkeletonVariableClass->IsChildOf(SkeletonVariableClass))
-				{
-					bReferencesVariable = true;
-				}
-			}
-			else if(InProperty->GetFName() == InVarName)
-			{
-				bReferencesVariable = true;
-			}
-		};
-		
-		PropertyAccessEditor.ResolvePropertyAccess(GetAnimBlueprint()->SkeletonGeneratedClass, BindingPair.Value.PropertyPath, ResolveArgs);
-
-		if(bReferencesVariable)
-		{
-			return true;
-		}
+		return Binding->ReferencesVariable(InVarName, InScope);
 	}
 
 	return false;
 }
 
-bool UAnimGraphNode_Base::ReferencesFunction(const FName& InFunctionName,const UStruct* InScope) const
+bool UAnimGraphNode_Base::ReferencesFunction(const FName& InFunctionName, const UStruct* InScope) const
 {
-	IPropertyAccessEditor& PropertyAccessEditor = IModularFeatures::Get().GetModularFeature<IPropertyAccessEditor>("PropertyAccessEditor");
-
-	const UClass* SkeletonFunctionClass = FBlueprintEditorUtils::GetSkeletonClass(Cast<UClass>(InScope));
-
-	// See if any of bindings reference the function
-	for (const auto& BindingPair : PropertyBindings)
+	if (Binding && Binding->ReferencesFunction(InFunctionName, InScope))
 	{
-		bool bReferencesFunction = false;
-
-		IPropertyAccessEditor::FResolvePropertyAccessArgs ResolveArgs;
-		ResolveArgs.FunctionFunction = [InFunctionName, SkeletonFunctionClass, &bReferencesFunction](int32 InSegmentIndex, UFunction* InFunction, FProperty* InProperty)
-		{
-			if (SkeletonFunctionClass)
-			{
-				const UClass* OwnerSkeletonFunctionClass = FBlueprintEditorUtils::GetSkeletonClass(InFunction->GetOuterUClass());
-
-				if (OwnerSkeletonFunctionClass && InFunction->GetFName() == InFunctionName && OwnerSkeletonFunctionClass->IsChildOf(SkeletonFunctionClass))
-				{
-					bReferencesFunction = true;
-				}
-			}
-			else if (InFunction->GetFName() == InFunctionName)
-			{
-				bReferencesFunction = true;
-			}
-		};
-
-		PropertyAccessEditor.ResolvePropertyAccess(GetAnimBlueprint()->SkeletonGeneratedClass, BindingPair.Value.PropertyPath, ResolveArgs);
-
-		if (bReferencesFunction)
-		{
-			return true;
-		}
+		return true;
 	}
 
 	// Check private member anim node function binding 
diff --git a/Engine/Source/Editor/AnimGraph/Private/AnimGraphNode_CustomProperty.cpp b/Engine/Source/Editor/AnimGraph/Private/AnimGraphNode_CustomProperty.cpp
--- a/Engine/Source/Editor/AnimGraph/Private/AnimGraphNode_CustomProperty.cpp
+++ b/Engine/Source/Editor/AnimGraph/Private/AnimGraphNode_CustomProperty.cpp
@@ -13,6 +13,7 @@
 #include "IAnimBlueprintCompilationContext.h"
 #include "IAnimBlueprintCopyTermDefaultsContext.h"
 #include "UObject/UE5ReleaseStreamObjectVersion.h"
+#include "AnimGraphNodeBinding.h"
 
 #define LOCTEXT_NAMESPACE "CustomPropNode"
 
@@ -310,12 +311,9 @@ bool UAnimGraphNode_CustomProperty::GetPinBindingInfo(FName InPinName, FName& Ou
 
 bool UAnimGraphNode_CustomProperty::HasBinding(FName InBindingName) const
 {
-	for(const TPair<FName, FAnimGraphNodePropertyBinding>& BindingPair : PropertyBindings)
+	if (GetBinding())
 	{
-		if(InBindingName == BindingPair.Key)
-		{
-			return true;
-		}
+		return GetBinding()->HasBinding(InBindingName, true);
 	}
 
 	return false;
@@ -385,46 +383,37 @@ void UAnimGraphNode_CustomProperty::PinConnectionListChanged(UEdGraphPin* Pin)
 	if(Pin->LinkedTo.Num() > 0)
 	{
 		// If we have links, clear any bindings
-		PropertyBindings.Remove(*GetPinTargetVariableName(Pin));
+		RemoveBindings(*GetPinTargetVariableName(Pin));
 	}
 }
 
 void UAnimGraphNode_CustomProperty::PostDuplicate(bool bDuplicateForPIE)
 {
 	Super::PostDuplicate(bDuplicateForPIE);
-	
-	const FString NewGUID = NodeGuid.ToString();
-	
-	TMap<FName, FAnimGraphNodePropertyBinding> NewBindings;
-	
-	// Update any bindings to point at the newly-generated property name
-	for(const TPair<FName, FAnimGraphNodePropertyBinding>& BindingPair : PropertyBindings)
+
+	if (GetBinding())
 	{
-		FString BindingNameString = BindingPair.Key.ToString();
-		FAnimGraphNodePropertyBinding Binding = BindingPair.Value;
+		const FString NewGUID = NodeGuid.ToString();
 
-		// Check pins
-		for (UEdGraphPin* Pin : Pins)
+		GetMutableBinding()->UpdateBindingNames([this, NewGUID](const FString& InOldBindingName)
 		{
-			FString PinTargetVariableNameBase = GetPinTargetVariableNameBase(Pin->GetFName());
-			if(BindingNameString.StartsWith(PinTargetVariableNameBase))
+			// Check pins
+			for (UEdGraphPin* Pin : Pins)
 			{
-				// Replace GUID with this node's if it differs
-				FString OldGUID = BindingNameString.RightChop(PinTargetVariableNameBase.Len());
-				if(OldGUID != NewGUID)
+				FString PinTargetVariableNameBase = GetPinTargetVariableNameBase(Pin->GetFName());
+				if (InOldBindingName.StartsWith(PinTargetVariableNameBase))
 				{
-					FName PropertyName = *(BindingNameString.Replace(*OldGUID, *NewGUID));
-					Binding.PropertyName = PropertyName;
-					NewBindings.Add(PropertyName, Binding);
+					// Replace GUID with this node's if it differs
+					FString OldGUID = InOldBindingName.RightChop(PinTargetVariableNameBase.Len());
+					if (OldGUID != NewGUID)
+					{
+						return InOldBindingName.Replace(*OldGUID, *NewGUID);
+					}
 				}
-				break;
 			}
-		}
-	}
 
-	if(NewBindings.Num() > 0)
-	{
-		PropertyBindings = NewBindings;
+			return InOldBindingName;
+		});
 	}
 }
 
diff --git a/Engine/Source/Editor/AnimGraph/Private/AnimStateConduitNode.cpp b/Engine/Source/Editor/AnimGraph/Private/AnimStateConduitNode.cpp
--- a/Engine/Source/Editor/AnimGraph/Private/AnimStateConduitNode.cpp
+++ b/Engine/Source/Editor/AnimGraph/Private/AnimStateConduitNode.cpp
@@ -114,7 +114,7 @@ void UAnimStateConduitNode::ValidateNodeDuringCompilation(class FCompilerResults
 	UAnimGraphNode_TransitionResult* ResultNode = TransGraph->GetResultNode();
 	check(ResultNode);
 
-	if (ResultNode->PropertyBindings.Num() > 0 && ResultNode->PropertyBindings.CreateIterator()->Value.bIsBound)
+	if (ResultNode->HasBinding(GET_MEMBER_NAME_CHECKED(FAnimNode_TransitionResult, bCanEnterTransition)))
 	{
 		// Rule is bound so nothing more to check
 	}
diff --git a/Engine/Source/Editor/AnimGraph/Private/AnimStateTransitionNode.cpp b/Engine/Source/Editor/AnimGraph/Private/AnimStateTransitionNode.cpp
--- a/Engine/Source/Editor/AnimGraph/Private/AnimStateTransitionNode.cpp
+++ b/Engine/Source/Editor/AnimGraph/Private/AnimStateTransitionNode.cpp
@@ -725,7 +725,7 @@ void UAnimStateTransitionNode::ValidateNodeDuringCompilation(class FCompilerResu
 				}
 			}
 		}
-		else if (ResultNode->PropertyBindings.Num() > 0 && ResultNode->PropertyBindings.CreateIterator()->Value.bIsBound)
+		else if (ResultNode->HasBinding(GET_MEMBER_NAME_CHECKED(FAnimNode_TransitionResult, bCanEnterTransition)))
 		{
 			// Rule is bound so nothing more to check
 		}
diff --git a/Engine/Source/Editor/AnimGraph/Private/AnimationGraphSchema.cpp b/Engine/Source/Editor/AnimGraph/Private/AnimationGraphSchema.cpp
--- a/Engine/Source/Editor/AnimGraph/Private/AnimationGraphSchema.cpp
+++ b/Engine/Source/Editor/AnimGraph/Private/AnimationGraphSchema.cpp
@@ -352,7 +352,7 @@ bool UAnimationGraphSchema::TryCreateConnection(UEdGraphPin* A, UEdGraphPin* B)
 			// Compare FName without number to make sure we catch array properties that are split into multiple pins
 			FName ComparisonName = InputPin->GetFName();
 			ComparisonName.SetNumber(0);
-			AnimGraphNode->PropertyBindings.Remove(ComparisonName);
+			AnimGraphNode->RemoveBindings(ComparisonName);
 		}
 
 		return true;
@@ -908,7 +908,7 @@ TSharedPtr<SWidget> UAnimationGraphSchema::MakeBindingWidgetForPin(const TArray<
 		{
 			if(bInOnGraphNode)
 			{
-				if (const FAnimGraphNodePropertyBinding* BindingPtr = FirstNode->PropertyBindings.Find(BindingName))
+				if (FirstNode->HasBinding(BindingName))
 				{
 					return EVisibility::Visible;
 				}
diff --git a/Engine/Source/Editor/AnimGraph/Public/AnimBlueprintExtension_PropertyAccess.h b/Engine/Source/Editor/AnimGraph/Public/AnimBlueprintExtension_PropertyAccess.h
--- a/Engine/Source/Editor/AnimGraph/Public/AnimBlueprintExtension_PropertyAccess.h
+++ b/Engine/Source/Editor/AnimGraph/Public/AnimBlueprintExtension_PropertyAccess.h
@@ -62,7 +62,7 @@ public:
 	
 public:
 	// Add a copy to the property access library we are compiling
-	// @return an integer handle to the pending copy. This can be resolved to a true copy index by calling MapCopyIndex
+	// @return an integer handle to the pending copy. This can be resolved to a true copy index by calling GetCompiledHandle
 	FPropertyAccessHandle AddCopy(TArrayView<FString> InSourcePath, TArrayView<FString> InDestPath, const FName& InContextId, UObject* InObject = nullptr);
 	
 	UE_DEPRECATED(5.0, "Please use AddCopy with a context ID")
@@ -72,6 +72,10 @@ public:
 		return Handle.GetId();
 	}
 
+	// Add an access to the property access library we are compiling
+	// @return an integer handle to the pending access. This can be resolved to a true access index by calling GetCompiledHandle
+	FPropertyAccessHandle AddAccess(TArrayView<FString> InPath, UObject* InObject = nullptr);
+
 	// Delegate called when the library is compiled (whether successfully or not)
 	FSimpleMulticastDelegate& OnPreLibraryCompiled() { return OnPreLibraryCompiledDelegate; }
 
@@ -81,12 +85,8 @@ public:
 	// Maps the initial copy handle to a true handle, post compilation
 	FCompiledPropertyAccessHandle GetCompiledHandle(FPropertyAccessHandle InHandle) const;
 
-	UE_DEPRECATED(5.0, "Please use GetCompiledHandle")
-	int32 MapCopyIndex(int32 InIndex) const
-	{
-		FCompiledPropertyAccessHandle CompiledHandle = GetCompiledHandle(FPropertyAccessHandle(InIndex));
-		return CompiledHandle.GetId();
-	}
+	// Get the access type for the specified handle
+	EPropertyAccessCopyType GetCompiledHandleAccessType(FPropertyAccessHandle InHandle) const;
 
 	// Expands a property access path to a pure chain of BP nodes
 	void ExpandPropertyAccess(FKismetCompilerContext& InCompilerContext, TArrayView<FString> InSourcePath, UEdGraph* InParentGraph, UEdGraphPin* InTargetPin) const;
diff --git a/Engine/Source/Editor/AnimGraph/Public/AnimGraphNodeAlphaOptions.h b/Engine/Source/Editor/AnimGraph/Public/AnimGraphNodeAlphaOptions.h
--- a/Engine/Source/Editor/AnimGraph/Public/AnimGraphNodeAlphaOptions.h
+++ b/Engine/Source/Editor/AnimGraph/Public/AnimGraphNodeAlphaOptions.h
@@ -105,7 +105,7 @@ struct FAnimGraphNodeAlphaOptions
 					if (InAnimNode.AlphaInputType != EAnimAlphaInputType::Float)
 					{
 						Pin->BreakAllPinLinks();
-						InNode->PropertyBindings.Remove(Pin->PinName);
+						InNode->RemoveBindings(Pin->PinName);
 					}
 				}
 				else if (Pin->PinName == GET_MEMBER_NAME_STRING_CHECKED(AnimNodeType, bAlphaBoolEnabled))
@@ -113,7 +113,7 @@ struct FAnimGraphNodeAlphaOptions
 					if (InAnimNode.AlphaInputType != EAnimAlphaInputType::Bool)
 					{
 						Pin->BreakAllPinLinks();
-						InNode->PropertyBindings.Remove(Pin->PinName);
+						InNode->RemoveBindings(Pin->PinName);
 					}
 				}
 				else if (Pin->PinName == GET_MEMBER_NAME_STRING_CHECKED(AnimNodeType, AlphaCurveName))
@@ -121,7 +121,7 @@ struct FAnimGraphNodeAlphaOptions
 					if (InAnimNode.AlphaInputType != EAnimAlphaInputType::Curve)
 					{
 						Pin->BreakAllPinLinks();
-						InNode->PropertyBindings.Remove(Pin->PinName);
+						InNode->RemoveBindings(Pin->PinName);
 					}
 				}
 			}
diff --git a/Engine/Source/Editor/AnimGraph/Public/AnimGraphNode_Base.h b/Engine/Source/Editor/AnimGraph/Public/AnimGraphNode_Base.h
--- a/Engine/Source/Editor/AnimGraph/Public/AnimGraphNode_Base.h
+++ b/Engine/Source/Editor/AnimGraph/Public/AnimGraphNode_Base.h
@@ -31,6 +31,7 @@ class IAnimBlueprintCopyTermDefaultsContext;
 class IAnimBlueprintNodeCopyTermDefaultsContext;
 class IAnimBlueprintNodeOverrideAssetsContext;
 class UAnimBlueprintExtension;
+class UAnimGraphNodeBinding;
 
 struct FPoseLinkMappingRecord
 {
@@ -199,8 +200,8 @@ class ANIMGRAPH_API UAnimGraphNode_Base : public UK2Node
 	TArray<FOptionalPinFromProperty> ShowPinForProperties;
 
 	/** Map from property name->binding info */
- 	UPROPERTY(EditAnywhere, Category=PinOptions)
- 	TMap<FName, FAnimGraphNodePropertyBinding> PropertyBindings;
+ 	UPROPERTY()
+ 	TMap<FName, FAnimGraphNodePropertyBinding> PropertyBindings_DEPRECATED;
 
 	/** Properties marked as always dynamic, so they can be set externally */
 	UPROPERTY()
@@ -222,6 +223,10 @@ class ANIMGRAPH_API UAnimGraphNode_Base : public UK2Node
 	FMemberReference UpdateFunction;
 
 private:
+	// Bindings for pins that this node exposes
+	UPROPERTY(EditAnywhere, Instanced, Category=Bindings)
+	TObjectPtr<UAnimGraphNodeBinding> Binding;
+
 	// Optional reference tag name. If this is set then this node can be referenced from elsewhere in this animation blueprint using an anim node reference
 	UPROPERTY(EditAnywhere, Category = Tag)
 	FName Tag;
@@ -541,6 +546,27 @@ public:
 	// Check if a specified function reference appears to be valid by inspecting only the validity of the name and guid
 	static bool IsPotentiallyBoundFunction(const FMemberReference& FunctionReference);
 
+	// Check whether the specified property is bound via PropertyBindings
+	virtual bool HasBinding(FName InPropertyName) const;
+
+	// Get the bindings for this node
+	const UAnimGraphNodeBinding* GetBinding() const { return Binding; }
+
+	// Get the mutable bindings for this node
+	UAnimGraphNodeBinding* GetMutableBinding() { return Binding; }
+
+	// Remove any bindings for the specified name
+	void RemoveBindings(FName InBindingName);
+
+	// Gets the animation FNode type represented by this ed graph node
+	UScriptStruct* GetFNodeType() const;
+
+	// Gets the animation FNode property represented by this ed graph node
+	FStructProperty* GetFNodeProperty() const;
+
+	// Get the runtime anim node that we template
+	FAnimNode_Base* GetFNode();
+
 protected:
 	friend class FAnimBlueprintCompilerContext;
 	friend class FAnimGraphNodeDetails;
@@ -552,15 +578,6 @@ protected:
 	// Set the tag for this node but without regenerating any BP data for tagging
 	void SetTagInternal(FName InTag) { Tag = InTag; }
 	
-	// Gets the animation FNode type represented by this ed graph node
-	UScriptStruct* GetFNodeType() const;
-
-	// Gets the animation FNode property represented by this ed graph node
-	FStructProperty* GetFNodeProperty() const;
-
-	// Get the runtime anim node that we template
-	FAnimNode_Base* GetFNode();
-	
 	// Get the extension types that this node type holds on the anim blueprint. Some extension types are always requested by the system
 	virtual void GetRequiredExtensions(TArray<TSubclassOf<UAnimBlueprintExtension>>& OutExtensions) const {}
 	
@@ -613,9 +630,6 @@ protected:
 	// @return false if the pin cannot be bound
 	virtual bool GetPinBindingInfo(FName InPinName, FName& OutBindingName, FProperty*& OutPinProperty, int32& OutOptionalPinIndex) const;
 
-	// Check whether the specified property is bound via PropertyBindings
-	virtual bool HasBinding(FName InPropertyName) const;
-	
 	FOnNodePropertyChangedEvent PropertyChangeEvent;
 
 	FOnNodeTitleChangedEvent NodeTitleChangedEvent;
diff --git a/Engine/Source/Editor/AnimGraph/Public/IAnimBlueprintCompilationBracketContext.h b/Engine/Source/Editor/AnimGraph/Public/IAnimBlueprintCompilationBracketContext.h
--- a/Engine/Source/Editor/AnimGraph/Public/IAnimBlueprintCompilationBracketContext.h
+++ b/Engine/Source/Editor/AnimGraph/Public/IAnimBlueprintCompilationBracketContext.h
@@ -19,10 +19,16 @@ public:
 	// Index of the nodes (must match up with the runtime discovery process of nodes, which runs thru the property chain)
 	const TMap<UAnimGraphNode_Base*, int32>& GetAllocatedAnimNodeIndices() const { return GetAllocatedAnimNodeIndicesImpl(); }
 
+	// Map of anim node indices to node handler properties in sparse class data struct
+	const TMap<UAnimGraphNode_Base*, FProperty*>& GetAllocatedHandlerPropertiesByNode() const { return GetAllocatedHandlerPropertiesByNodeImpl(); }
+
 protected:
 	// Get the message log for the current compilation
 	virtual FCompilerResultsLog& GetMessageLogImpl() const = 0;
 
 	// Map of anim node properties to original anim graph node
 	virtual const TMap<UAnimGraphNode_Base*, int32>& GetAllocatedAnimNodeIndicesImpl() const = 0;
+
+	// Map of anim node indices to node handler properties in sparse class data struct
+	virtual const TMap<UAnimGraphNode_Base*, FProperty*>& GetAllocatedHandlerPropertiesByNodeImpl() const = 0;
 };
diff --git a/Engine/Source/Editor/AnimGraph/Public/IAnimBlueprintCompilationContext.h b/Engine/Source/Editor/AnimGraph/Public/IAnimBlueprintCompilationContext.h
--- a/Engine/Source/Editor/AnimGraph/Public/IAnimBlueprintCompilationContext.h
+++ b/Engine/Source/Editor/AnimGraph/Public/IAnimBlueprintCompilationContext.h
@@ -131,6 +131,9 @@ public:
 	// Map of anim node indices to node properties
 	const TMap<UAnimGraphNode_Base*, FProperty*>& GetAllocatedPropertiesByNode() const { return GetAllocatedPropertiesByNodeImpl(); }
 
+	// Map of anim node indices to node handler properties in sparse class data struct
+	const TMap<UAnimGraphNode_Base*, FProperty*>& GetAllocatedHandlerPropertiesByNode() const { return GetAllocatedHandlerPropertiesByNodeImpl(); }
+
 	// Add the specified compiled-in attribute uniquely to the specified node
 	void AddAttributesToNode(UAnimGraphNode_Base* InNode, TArrayView<const FName> InAttributes) const { AddAttributesToNodeImpl(InNode, InAttributes); }
 
@@ -192,6 +195,9 @@ protected:
 	// Map of anim node indices to node properties
 	virtual const TMap<UAnimGraphNode_Base*, FProperty*>& GetAllocatedPropertiesByNodeImpl() const = 0;
 
+	// Map of anim node indices to node handler properties in sparse class data struct
+	virtual const TMap<UAnimGraphNode_Base*, FProperty*>& GetAllocatedHandlerPropertiesByNodeImpl() const = 0;
+
 	// Add the specified compiled-in attribute uniquely to the specified node
 	virtual void AddAttributesToNodeImpl(UAnimGraphNode_Base* InNode, TArrayView<const FName> InAttributes) const = 0;
 
diff --git a/Engine/Source/Editor/Persona/Private/AnimGraphNodeDetails.cpp b/Engine/Source/Editor/Persona/Private/AnimGraphNodeDetails.cpp
--- a/Engine/Source/Editor/Persona/Private/AnimGraphNodeDetails.cpp
+++ b/Engine/Source/Editor/Persona/Private/AnimGraphNodeDetails.cpp
@@ -66,10 +66,8 @@ void FAnimGraphNodeDetails::CustomizeDetails(class IDetailLayoutBuilder& DetailB
 
 	// Hide the pin options property; it's represented inline per-property instead
 	IDetailCategoryBuilder& PinOptionsCategory = DetailBuilder.EditCategory("PinOptions");
-	TSharedRef<IPropertyHandle> AvailablePins = DetailBuilder.GetProperty("ShowPinForProperties");
+	TSharedRef<IPropertyHandle> AvailablePins = DetailBuilder.GetProperty(GET_MEMBER_NAME_CHECKED(UAnimGraphNode_Base, ShowPinForProperties));
 	DetailBuilder.HideProperty(AvailablePins);
-	TSharedRef<IPropertyHandle> PropertyBindings = DetailBuilder.GetProperty("PropertyBindings");
-	DetailBuilder.HideProperty(PropertyBindings);
 
 	// get first animgraph nodes
 	UAnimGraphNode_Base* AnimGraphNode = Cast<UAnimGraphNode_Base>(SelectedObjectsList[0].Get());
@@ -78,6 +76,16 @@ void FAnimGraphNodeDetails::CustomizeDetails(class IDetailLayoutBuilder& DetailB
 		return;
 	}
 
+	// Ensure that switching the binding type will nuke the details panel
+	IDetailCategoryBuilder& BindingCategory = DetailBuilder.EditCategory("Bindings");
+	BindingCategory.SetSortOrder(MAX_int32);
+	
+	TSharedRef<IPropertyHandle> BindingProperty = DetailBuilder.GetProperty(GET_MEMBER_NAME_CHECKED(UAnimGraphNode_Base, Binding));
+	BindingProperty->SetOnPropertyValueChanged(FSimpleDelegate::CreateLambda([&DetailBuilder]()
+	{
+		DetailBuilder.ForceRefreshDetails();
+	}));
+
 	AnimGraphNode->OnPinVisibilityChanged().AddSP(this, &FAnimGraphNodeDetails::OnPinVisibilityChanged);
 
 	// make sure type matches with all the nodes. 
diff --git a/Engine/Source/Editor/UnrealEd/Private/Kismet2/KismetDebugUtilities.cpp b/Engine/Source/Editor/UnrealEd/Private/Kismet2/KismetDebugUtilities.cpp
--- a/Engine/Source/Editor/UnrealEd/Private/Kismet2/KismetDebugUtilities.cpp
+++ b/Engine/Source/Editor/UnrealEd/Private/Kismet2/KismetDebugUtilities.cpp
@@ -1467,11 +1467,7 @@ bool FKismetDebugUtilities::CanWatchPin(const UBlueprint* Blueprint, const UEdGr
 
 			if (UAnimGraphNode_Base* AnimGraphNode = Cast<UAnimGraphNode_Base>(Pin->GetOwningNode()))
 			{
-				// Compare FName without number to make sure we catch array properties that are split into multiple pins
-				FName ComparisonName = Pin->GetFName();
-				ComparisonName.SetNumber(0);
-
-				if (FAnimGraphNodePropertyBinding* BindingPtr = AnimGraphNode->PropertyBindings.Find(ComparisonName))
+				if(AnimGraphNode->HasBinding(Pin->GetFName()))
 				{
 					bHasBinding = true;
 				}
diff --git a/Engine/Source/Editor/UnrealEd/Public/IPropertyAccessCompiler.h b/Engine/Source/Editor/UnrealEd/Public/IPropertyAccessCompiler.h
--- a/Engine/Source/Editor/UnrealEd/Public/IPropertyAccessCompiler.h
+++ b/Engine/Source/Editor/UnrealEd/Public/IPropertyAccessCompiler.h
@@ -11,6 +11,7 @@
 
 class UObject;
 struct FPropertyAccessLibrary;
+enum class EPropertyAccessCopyType : uint8;
 
 // DEPRECATED - The various batching of property copy
 UENUM()
@@ -23,6 +24,13 @@ enum class EPropertyAccessBatchType : uint8
 	Batched,
 };
 
+enum class EPropertyAccessHandleType : int32
+{
+	Copy,
+
+	Access,
+};
+
 // Handle used to track property accesses that are being compiled
 struct FPropertyAccessHandle
 {
@@ -33,10 +41,14 @@ public:
 	// Check this handle for validity
 	bool IsValid() const { return Id != INDEX_NONE; }
 
+	// Get the type
+	EPropertyAccessHandleType GetType() const { return Type; }
+
 	FPropertyAccessHandle() = default;
 	
-	FPropertyAccessHandle(int32 InId)
+	FPropertyAccessHandle(int32 InId, EPropertyAccessHandleType InType)
 		: Id(InId)
+		, Type(InType)
 	{
 	}
 
@@ -48,11 +60,14 @@ public:
 private:
 	// Index used to track accesses by external systems
 	int32 Id = INDEX_NONE;
+
+	// The type of this access
+	EPropertyAccessHandleType Type = EPropertyAccessHandleType::Copy;
 };
 
 inline uint32 GetTypeHash(const FPropertyAccessHandle& Value)
 {
-	return GetTypeHash(Value.GetId());
+	return HashCombine(GetTypeHash(Value.GetId()), GetTypeHash(Value.GetType()));
 }
 
 // Handle used to describe property accesses that have been compiled
@@ -68,14 +83,15 @@ public:
 	int32 GetBatchId() const { return BatchId; }
 
 	// Check this handle for validity
-	bool IsValid() const { return Id != INDEX_NONE && BatchId != INDEX_NONE; }
+	bool IsValid() const { return Id != INDEX_NONE; }
 
 private:
 	friend class FPropertyAccessLibraryCompiler;
 
-	FCompiledPropertyAccessHandle(int32 InId, int32 InBatchId)
+	FCompiledPropertyAccessHandle(int32 InId, int32 InBatchId, EPropertyAccessHandleType InType)
 		: Id(InId)
 		, BatchId(InBatchId)
+		, Type(InType)
 	{
 	}
 
@@ -84,6 +100,9 @@ private:
 
 	// Index of the batch
 	int32 BatchId = INDEX_NONE;
+
+	// The type of this access
+	EPropertyAccessHandleType Type = EPropertyAccessHandleType::Copy;
 };
 
 // A helper used to compile a property access library
@@ -99,7 +118,7 @@ public:
 	virtual void BeginCompilation(const UClass* InClass) {}
 
 	// Add a copy to the property access library we are compiling
-	// @return an integer handle to the pending copy. This can be resolved to a true copy index by calling MapCopyIndex
+	// @return a handle to the pending copy. This can be resolved to a true copy index by calling GetCompiledHandle
 	virtual FPropertyAccessHandle AddCopy(TArrayView<FString> InSourcePath, TArrayView<FString> InDestPath, const FName& InContextId, UObject* InAssociatedObject = nullptr) = 0;
 	
 	UE_DEPRECATED(5.0, "Please use AddCopy with a context ID that returns a handle")
@@ -109,6 +128,10 @@ public:
 		return Handle.GetId();
 	}
 
+	// Add an access to the property access library we are compiling
+	// @return a handle to the pending access. This can be resolved to a true access index by calling GetCompiledHandle
+	virtual FPropertyAccessHandle AddAccess(TArrayView<FString> InPath, UObject* InAssociatedObject = nullptr) = 0;
+
 	// Post-process the library to finish compilation. @return true if compilation succeeded.
 	virtual bool FinishCompilation() = 0;
 
@@ -118,10 +141,6 @@ public:
 	// Maps the initial copy handle to a true handle, post compilation
 	virtual FCompiledPropertyAccessHandle GetCompiledHandle(FPropertyAccessHandle InHandle) const = 0;
 
-	UE_DEPRECATED(5.0, "Please use GetCompiledHandle")
-	virtual int32 MapCopyIndex(int32 InIndex) const
-	{
-		FCompiledPropertyAccessHandle CompiledHandle = GetCompiledHandle(FPropertyAccessHandle(InIndex));
-		return CompiledHandle.GetId();
-	}
+	// Get the access type for the specified handle
+	virtual EPropertyAccessCopyType GetCompiledHandleAccessType(FPropertyAccessHandle InHandle) const = 0;
 };
\ No newline at end of file
diff --git a/Engine/Source/Runtime/Core/Public/UObject/FortniteMainBranchObjectVersion.h b/Engine/Source/Runtime/Core/Public/UObject/FortniteMainBranchObjectVersion.h
--- a/Engine/Source/Runtime/Core/Public/UObject/FortniteMainBranchObjectVersion.h
+++ b/Engine/Source/Runtime/Core/Public/UObject/FortniteMainBranchObjectVersion.h
@@ -348,6 +348,9 @@ struct FFortniteMainBranchObjectVersion
 		// Fixup all flags/outering on static meshes on water bodies by rebuilding them completely
 		WaterBodyStaticMeshFixup,
 
+		// Binding extensions for anim graph nodes
+		AnimGraphNodeBindingExtensions,
+
 		// -----<new versions can be added above this line>-------------------------------------------------
 		VersionPlusOne,
 		LatestVersion = VersionPlusOne - 1
diff --git a/Engine/Source/Runtime/Engine/Classes/Animation/AnimBlueprint.h b/Engine/Source/Runtime/Engine/Classes/Animation/AnimBlueprint.h
--- a/Engine/Source/Runtime/Engine/Classes/Animation/AnimBlueprint.h
+++ b/Engine/Source/Runtime/Engine/Classes/Animation/AnimBlueprint.h
@@ -16,6 +16,8 @@ class USkeleton;
 class UPoseWatch;
 class UPoseWatchFolder;
 struct FAnimBlueprintDebugData;
+class UAnimGraphNodeBinding;
+class UClass;
 
 USTRUCT()
 struct FAnimGroupInfo
@@ -237,6 +239,9 @@ public:
 	// @param	bInIsInterface		Whether the anim blueprint to check is an interface
 	ENGINE_API bool IsCompatibleByAssetString(const FString& InSkeletonAsset, bool bInIsTemplate, bool bInIsInterface) const;
 	
+	// Get the default binding type that any new nodes will use when created
+	ENGINE_API UClass* GetDefaultBindingClass() const { return DefaultBindingClass; }
+
 public:
 	// Array of overrides to asset containing nodes in the parent that have been overridden
 	UPROPERTY()
@@ -264,6 +269,10 @@ private:
 	UPROPERTY(duplicatetransient, AssetRegistrySearchable)
 	TSoftObjectPtr<class UAnimBlueprint> PreviewAnimationBlueprint;
 
+	/** The default binding type that any new nodes will use when created */
+	UPROPERTY(EditAnywhere, Category=Bindings, meta=(AllowedClasses="/Script/AnimGraph.AnimGraphNodeBinding", ShowDisplayNames=true, NoClear))
+	TObjectPtr<UClass> DefaultBindingClass;
+
 	/** The method by which a preview animation blueprint is applied, either as an overlay layer, or as a linked instance */
 	UPROPERTY()
 	EPreviewAnimationBlueprintApplicationMethod PreviewAnimationBlueprintApplicationMethod;
diff --git a/Engine/Source/Runtime/Engine/Classes/Animation/AnimBlueprintGeneratedClass.h b/Engine/Source/Runtime/Engine/Classes/Animation/AnimBlueprintGeneratedClass.h
--- a/Engine/Source/Runtime/Engine/Classes/Animation/AnimBlueprintGeneratedClass.h
+++ b/Engine/Source/Runtime/Engine/Classes/Animation/AnimBlueprintGeneratedClass.h
@@ -368,6 +368,7 @@ class UAnimBlueprintGeneratedClass : public UBlueprintGeneratedClass, public IAn
 	friend class FAnimBlueprintCompilerContext;
 	friend class FAnimBlueprintGeneratedClassCompiledData;
 	friend class FKismetDebugUtilities;
+	friend class UAnimBlueprintExtension_Base;
 
 	// List of state machines present in this blueprint class
 	UPROPERTY()
diff --git a/Engine/Source/Runtime/Engine/Private/AnimBlueprint.cpp b/Engine/Source/Runtime/Engine/Private/AnimBlueprint.cpp
--- a/Engine/Source/Runtime/Engine/Private/AnimBlueprint.cpp
+++ b/Engine/Source/Runtime/Engine/Private/AnimBlueprint.cpp
@@ -31,6 +31,8 @@ UAnimBlueprint::UAnimBlueprint(const FObjectInitializer& ObjectInitializer)
 		// Ensure that we are able to compile this anim BP by loading the compiler's module
 		FModuleManager::Get().LoadModuleChecked("AnimGraph");
 	}
+
+	DefaultBindingClass = FindObject<UClass>(nullptr, TEXT("/Script/AnimGraph.AnimGraphNodeBinding_Base"));
 #endif
 }
 
diff --git a/Engine/Source/Runtime/Engine/Private/Animation/AnimSubsystem_Base.cpp b/Engine/Source/Runtime/Engine/Private/Animation/AnimSubsystem_Base.cpp
--- a/Engine/Source/Runtime/Engine/Private/Animation/AnimSubsystem_Base.cpp
+++ b/Engine/Source/Runtime/Engine/Private/Animation/AnimSubsystem_Base.cpp
@@ -2,11 +2,32 @@
 
 #include "Animation/AnimSubsystem_Base.h"
 #include "Animation/AnimSubsystem.h"
+#include "Animation/AnimClassInterface.h"
 
 #include UE_INLINE_GENERATED_CPP_BY_NAME(AnimSubsystem_Base)
 
 void FAnimSubsystem_Base::OnPostLoadDefaults(FAnimSubsystemPostLoadDefaultsContext& InContext)
 {
-	FExposedValueHandler::ClassInitialization(ExposedValueHandlers, InContext.DefaultAnimInstance->GetClass());
+	PatchValueHandlers(InContext.DefaultAnimInstance->GetClass());
 }
 
+void FAnimSubsystem_Base::PatchValueHandlers(UClass* InClass)
+{
+	ExposedValueHandlers.Reset();
+
+	void* SparseClassData = const_cast<void*>(InClass->GetSparseClassData(EGetSparseClassDataMethod::ReturnIfNull));
+	check(SparseClassData);
+	for (TFieldIterator<FProperty> ItParam(InClass->GetSparseClassDataStruct()); ItParam; ++ItParam)
+	{
+		if (FStructProperty* StructProperty = CastField<FStructProperty>(*ItParam))
+		{
+			if (StructProperty->Struct->IsChildOf(FAnimNodeExposedValueHandler::StaticStruct()))
+			{
+				FExposedValueHandler& NewHandler = ExposedValueHandlers.AddDefaulted_GetRef();
+				NewHandler.HandlerStruct = StructProperty->Struct;
+				NewHandler.Handler = StructProperty->ContainerPtrToValuePtr<FAnimNodeExposedValueHandler>(SparseClassData);
+				NewHandler.Handler->Initialize(InClass);
+			}
+		}
+	}
+}
\ No newline at end of file
diff --git a/Engine/Source/Runtime/Engine/Private/Animation/ExposedValueHandler.cpp b/Engine/Source/Runtime/Engine/Private/Animation/ExposedValueHandler.cpp
--- a/Engine/Source/Runtime/Engine/Private/Animation/ExposedValueHandler.cpp
+++ b/Engine/Source/Runtime/Engine/Private/Animation/ExposedValueHandler.cpp
@@ -7,20 +7,46 @@
 
 #include UE_INLINE_GENERATED_CPP_BY_NAME(ExposedValueHandler)
 
+// Deprecation support
+PRAGMA_DISABLE_DEPRECATION_WARNINGS
+TArray<FExposedValueCopyRecord> FExposedValueHandler::CopyRecords;
+TObjectPtr<UFunction> FExposedValueHandler::Function;
+const FPropertyAccessLibrary* FExposedValueHandler::PropertyAccessLibrary;
+FName FExposedValueHandler::BoundFunction;
+PRAGMA_ENABLE_DEPRECATION_WARNINGS
+
 void FExposedValueHandler::ClassInitialization(TArray<FExposedValueHandler>& InHandlers, UClass* InClass)
 {
 	if(const FAnimSubsystem_PropertyAccess* Subsystem = IAnimClassInterface::GetFromClass(InClass)->FindSubsystem<FAnimSubsystem_PropertyAccess>())
 	{
-		const FPropertyAccessLibrary& PropertyAccessLibrary = Subsystem->GetLibrary();
+		const FPropertyAccessLibrary& Library = Subsystem->GetLibrary();
 
 		for(FExposedValueHandler& Handler : InHandlers)
 		{
-			Handler.Initialize(InClass, PropertyAccessLibrary);
+			PRAGMA_DISABLE_DEPRECATION_WARNINGS
+			Handler.Initialize(InClass, Library);
+			PRAGMA_ENABLE_DEPRECATION_WARNINGS
 		}
 	}
 }
 
 void FExposedValueHandler::Initialize(UClass* InClass, const FPropertyAccessLibrary& InPropertyAccessLibrary)
+{
+	if(Handler)
+	{
+		Handler->Initialize(InClass);
+	}
+}
+
+void FExposedValueHandler::Execute(const FAnimationBaseContext& Context) const
+{
+	if (Handler)
+	{
+		Handler->Execute(Context);
+	}
+}
+
+void FAnimNodeExposedValueHandler_Base::Initialize(const UClass* InClass)
 {
 	if (BoundFunction != NAME_None)
 	{
@@ -39,29 +65,41 @@ void FExposedValueHandler::Initialize(UClass* InClass, const FPropertyAccessLibr
 	{
 		Function = nullptr;
 	}
-
-	// Cache property access library
-	PropertyAccessLibrary = &InPropertyAccessLibrary;
 }
 
-void FExposedValueHandler::Execute(const FAnimationBaseContext& Context) const
+void FAnimNodeExposedValueHandler_Base::Execute(const FAnimationBaseContext& InContext) const
 {
 	if (Function != nullptr)
 	{
-		Context.AnimInstanceProxy->GetAnimInstanceObject()->ProcessEvent(Function, nullptr);
+		InContext.AnimInstanceProxy->GetAnimInstanceObject()->ProcessEvent(Function, nullptr);
 	}
+}
+
+void FAnimNodeExposedValueHandler_PropertyAccess::Initialize(const UClass* InClass)
+{
+	Super::Initialize(InClass);
+
+	if (const FAnimSubsystem_PropertyAccess* Subsystem = IAnimClassInterface::GetFromClass(InClass)->FindSubsystem<FAnimSubsystem_PropertyAccess>())
+	{
+		PropertyAccessLibrary = &Subsystem->GetLibrary();
+	}
+}
+
+void FAnimNodeExposedValueHandler_PropertyAccess::Execute(const FAnimationBaseContext& InContext) const
+{
+	Super::Execute(InContext);
 
-	if(CopyRecords.Num() > 0)
+	if (CopyRecords.Num() > 0)
 	{
-		if(PropertyAccessLibrary != nullptr)
+		if (PropertyAccessLibrary != nullptr)
 		{
-			UObject* AnimInstanceObject = Context.AnimInstanceProxy->GetAnimInstanceObject();
+			UObject* AnimInstanceObject = InContext.AnimInstanceProxy->GetAnimInstanceObject();
 			PropertyAccess::FCopyBatchId CopyBatch((int32)EAnimPropertyAccessCallSite::WorkerThread_Unbatched);
-			for(const FExposedValueCopyRecord& CopyRecord : CopyRecords)
+			for (const FExposedValueCopyRecord& CopyRecord : CopyRecords)
 			{
 				PropertyAccess::ProcessCopy(AnimInstanceObject, *PropertyAccessLibrary, CopyBatch, CopyRecord.CopyIndex, [&CopyRecord](const FProperty* InProperty, void* InAddress)
 				{
-					if(CopyRecord.PostCopyOperation == EPostCopyOperation::LogicalNegateBool)
+					if (CopyRecord.PostCopyOperation == EPostCopyOperation::LogicalNegateBool)
 					{
 						bool bValue = static_cast<const FBoolProperty*>(InProperty)->GetPropertyValue(InAddress);
 						static_cast<const FBoolProperty*>(InProperty)->SetPropertyValue(InAddress, !bValue);
diff --git a/Engine/Source/Runtime/Engine/Private/PropertyAccess.cpp b/Engine/Source/Runtime/Engine/Private/PropertyAccess.cpp
--- a/Engine/Source/Runtime/Engine/Private/PropertyAccess.cpp
+++ b/Engine/Source/Runtime/Engine/Private/PropertyAccess.cpp
@@ -579,6 +579,22 @@ struct FPropertyAccessSystem
 			InAddressFunction(Address);
 		}
 	}
+	
+	static void GetAccessAddress(void* InContainer, const FPropertyAccessLibrary& InLibrary, int32 InAccessIndex, TFunctionRef<void(const FProperty*, void*)> InFunction)
+	{
+		if(InLibrary.DestAccesses.IsValidIndex(InAccessIndex))
+		{
+			const FPropertyAccessIndirectionChain& DestAccess = InLibrary.DestAccesses[InAccessIndex];
+			if (FProperty* DestProperty = DestAccess.Property.Get())
+			{
+				GetAccessAddress(InContainer, InLibrary, DestAccess, [&InFunction, DestProperty](void* InAddress)
+				{
+					InFunction(DestProperty, InAddress);
+				});
+			}
+		}
+	}
+
 	// Process a single copy
 	static void ProcessCopy(UStruct* InStruct, void* InContainer, const FPropertyAccessLibrary& InLibrary, int32 InCopyIndex, int32 InBatchId, TFunctionRef<void(const FProperty*, void*)> InPostCopyOperation)
 	{
@@ -727,6 +743,11 @@ namespace PropertyAccess
 		::FPropertyAccessSystem::ProcessCopy(InObject->GetClass(), InObject, InLibrary, InCopyIndex, InBatchId.Id, InPostCopyOperation);
 	}
 
+	void GetAccessAddress(UObject* InObject, const FPropertyAccessLibrary& InLibrary, int32 InAccessIndex, TFunctionRef<void(const FProperty*, void*)> InFunction)
+	{
+		::FPropertyAccessSystem::GetAccessAddress(InObject, InLibrary, InAccessIndex, InFunction);
+	}
+
 	void BindEvents(UObject* InObject, const FPropertyAccessLibrary& InLibrary)
 	{
 	}
diff --git a/Engine/Source/Runtime/Engine/Public/Animation/AnimSubsystem_Base.h b/Engine/Source/Runtime/Engine/Public/Animation/AnimSubsystem_Base.h
--- a/Engine/Source/Runtime/Engine/Public/Animation/AnimSubsystem_Base.h
+++ b/Engine/Source/Runtime/Engine/Public/Animation/AnimSubsystem_Base.h
@@ -17,10 +17,10 @@ struct FAnimSubsystem_Base : public FAnimSubsystem
 	// FAnimSubsystem interface
 	ENGINE_API virtual void OnPostLoadDefaults(FAnimSubsystemPostLoadDefaultsContext& InContext) override;
 
-	// Get the exposed value handlers held on this subsystem
 	const TArray<FExposedValueHandler>& GetExposedValueHandlers() const { return ExposedValueHandlers; }
 
+	ENGINE_API void PatchValueHandlers(UClass* InClass);
+
 private:
-	UPROPERTY()
 	TArray<FExposedValueHandler> ExposedValueHandlers;
 };
diff --git a/Engine/Source/Runtime/Engine/Public/Animation/ExposedValueHandler.h b/Engine/Source/Runtime/Engine/Public/Animation/ExposedValueHandler.h
--- a/Engine/Source/Runtime/Engine/Public/Animation/ExposedValueHandler.h
+++ b/Engine/Source/Runtime/Engine/Public/Animation/ExposedValueHandler.h
@@ -10,6 +10,7 @@
 struct FPropertyAccessLibrary;
 class UClass;
 struct FAnimationBaseContext;
+class FCheckFastPathLatentCommand;
 
 UENUM()
 enum class EPostCopyOperation : uint8
@@ -45,34 +46,98 @@ struct FExposedValueHandler
 {
 	GENERATED_USTRUCT_BODY()
 
-	FExposedValueHandler()
-		: Function(nullptr)
-		, PropertyAccessLibrary(nullptr)
-		, BoundFunction(NAME_None)
-	{
-	}
+	friend struct FAnimSubsystem_Base;
+
+	FExposedValueHandler() = default;
 
 	// Direct data access to property in anim instance
-	UPROPERTY()
-	TArray<FExposedValueCopyRecord> CopyRecords;
+	UE_DEPRECATED(5.3, "This has been moved to FAnimNodeExposedValueHandler")
+	static ENGINE_API TArray<FExposedValueCopyRecord> CopyRecords;
+
+	// function pointer if BoundFunction != NAME_None
+	UE_DEPRECATED(5.3, "This has been moved to FAnimNodeExposedValueHandler")
+	static ENGINE_API TObjectPtr<UFunction> Function;
+
+	UE_DEPRECATED(5.3, "This has been moved to FAnimNodeExposedValueHandler")
+	static ENGINE_API const FPropertyAccessLibrary* PropertyAccessLibrary;
+
+	UE_DEPRECATED(5.3, "This has been moved to FAnimNodeExposedValueHandler")
+	static ENGINE_API FName BoundFunction;
+
+	UE_DEPRECATED(5.3, "This function was meant for internal use, but was public. Initialization should now use the internal FAnimNodeExposedValueHandler")
+	static ENGINE_API void ClassInitialization(TArray<FExposedValueHandler>& Handlers, UClass* InClass);
+
+	UE_DEPRECATED(5.3, "This function was meant for internal use, but was public. Initialization should now use the internal FAnimNodeExposedValueHandler")
+	ENGINE_API void Initialize(UClass* InClass, const FPropertyAccessLibrary& InPropertyAccessLibrary);
+
+	// Execute the handler
+	ENGINE_API void Execute(const FAnimationBaseContext& Context) const;
+
+	// Access the handler struct
+	const UScriptStruct* GetHandlerStruct() const { return HandlerStruct; }
+
+	// Access the handler
+	const FAnimNodeExposedValueHandler* GetHandler() const { return Handler; }
+
+private:
+	// The type of the handler
+	UScriptStruct* HandlerStruct = nullptr;
+
+	// Ptr to actual handler (property of the anim instance)
+	FAnimNodeExposedValueHandler* Handler = nullptr;
+};
+
+USTRUCT()
+struct FAnimNodeExposedValueHandler
+{
+	GENERATED_BODY()
+
+	friend struct FAnimSubsystem_Base;
+
+	virtual ~FAnimNodeExposedValueHandler() {}
+
+	// Set up the handler when the class is loaded
+	ENGINE_API virtual void Initialize(const UClass* InClass) {}
+
+	// Execute the handler to set specified values in a node
+	ENGINE_API virtual void Execute(const FAnimationBaseContext& InContext) const {}
+};
+
+USTRUCT()
+struct FAnimNodeExposedValueHandler_Base : public FAnimNodeExposedValueHandler
+{
+	GENERATED_BODY()
+
+	friend class UAnimBlueprintExtension_Base;
+
+	// FAnimNodeExposedValueHandler interface
+	ENGINE_API virtual void Initialize(const UClass* InClass) override;
+	ENGINE_API virtual void Execute(const FAnimationBaseContext& InContext) const override;
 
 	// function pointer if BoundFunction != NAME_None
 	UPROPERTY()
 	TObjectPtr<UFunction> Function;
 
-	// Cached property access library ptr
-	const FPropertyAccessLibrary* PropertyAccessLibrary;
-
 	// The function to call to update associated properties (can be NAME_None)
 	UPROPERTY()
 	FName BoundFunction;
+};
 
-	// Helper function to bind an array of handlers.
-	static ENGINE_API void ClassInitialization(TArray<FExposedValueHandler>& Handlers, UClass* InClass);
+USTRUCT()
+struct FAnimNodeExposedValueHandler_PropertyAccess : public FAnimNodeExposedValueHandler_Base
+{
+	GENERATED_BODY()
 
-	// Bind copy records and cache UFunction if necessary
-	ENGINE_API void Initialize(UClass* InClass, const FPropertyAccessLibrary& InPropertyAccessLibrary);
+	friend class UAnimBlueprintExtension_Base;
 
-	// Execute the function and copy records
-	ENGINE_API void Execute(const FAnimationBaseContext& Context) const;
+	// FAnimNodeExposedValueHandler interface
+	virtual void Initialize(const UClass* InClass) override;
+	virtual void Execute(const FAnimationBaseContext& InContext) const override;
+
+	// Direct data access to property in anim instance
+	UPROPERTY()
+	TArray<FExposedValueCopyRecord> CopyRecords;
+
+	// Cached property access library ptr
+	const FPropertyAccessLibrary* PropertyAccessLibrary;
 };
diff --git a/Engine/Source/Runtime/Engine/Public/PropertyAccess.h b/Engine/Source/Runtime/Engine/Public/PropertyAccess.h
--- a/Engine/Source/Runtime/Engine/Public/PropertyAccess.h
+++ b/Engine/Source/Runtime/Engine/Public/PropertyAccess.h
@@ -77,6 +77,9 @@ namespace PropertyAccess
 
 	UE_DEPRECATED(5.0, "Property Access Events are no longer supported")
 	ENGINE_API extern int32 GetEventId(const UClass* InClass, TArrayView<const FName> InPath);
+
+	/** Gets the property and address of the specified access, minimally-resolving all compiled-in indirections */
+	ENGINE_API extern void GetAccessAddress(UObject* InObject, const FPropertyAccessLibrary& InLibrary, int32 InAccessIndex, TFunctionRef<void(const FProperty*, void*)> InFunction);
 }
 
 // The type of an indirection
