Added "Update If Active" flag to property access, to allow freezing specific properties during blend out.

diff --git a/Engine/Plugins/Experimental/Animation/PoseSearch/Source/Runtime/Private/AnimNode_BlendStack.cpp b/Engine/Plugins/Experimental/Animation/PoseSearch/Source/Runtime/Private/AnimNode_BlendStack.cpp
--- a/Engine/Plugins/Experimental/Animation/PoseSearch/Source/Runtime/Private/AnimNode_BlendStack.cpp
+++ b/Engine/Plugins/Experimental/Animation/PoseSearch/Source/Runtime/Private/AnimNode_BlendStack.cpp
@@ -454,8 +454,8 @@ void FAnimNode_BlendStack_Standalone::UpdateAssetPlayer(const FAnimationUpdateCo
 			break;
 		}
 
-		const FAnimationUpdateContext AnimPlayerContext = Context.FractionalWeightAndRootMotion(AnimPlayerBlendWeight, AnimPlayerBlendWeight);
-		UpdateSample(AnimPlayerContext, AnimPlayerIndex);
+		FAnimationUpdateContext AnimPlayerContext = Context.FractionalWeightAndRootMotion(AnimPlayerBlendWeight, AnimPlayerBlendWeight);
+		UpdateSample((AnimPlayerIndex == 0) ? AnimPlayerContext : AnimPlayerContext.AsInactive(), AnimPlayerIndex);
 		CurrentWeightMultiplier *= (1.f - BlendInPercentage);
 	}
 
diff --git a/Engine/Source/Editor/AnimGraph/Internal/AnimBlueprintExtension_Base.h b/Engine/Source/Editor/AnimGraph/Internal/AnimBlueprintExtension_Base.h
--- a/Engine/Source/Editor/AnimGraph/Internal/AnimBlueprintExtension_Base.h
+++ b/Engine/Source/Editor/AnimGraph/Internal/AnimBlueprintExtension_Base.h
@@ -84,6 +84,7 @@ private:
 			, BindingContextId(NAME_None)
 			, Operation(EPostCopyOperation::None)
 			, bIsFastPath(true)
+			, bOnlyUpdateWhenActive(false)
 		{}
 
 		bool IsFastPath() const
@@ -125,6 +126,8 @@ private:
 
 		/** Fast-path flag */
 		bool bIsFastPath;
+
+		bool bOnlyUpdateWhenActive;
 	};
 
 	// Context used to build fast-path copy records
diff --git a/Engine/Source/Editor/AnimGraph/Private/AnimBlueprintExtension_Base.cpp b/Engine/Source/Editor/AnimGraph/Private/AnimBlueprintExtension_Base.cpp
--- a/Engine/Source/Editor/AnimGraph/Private/AnimBlueprintExtension_Base.cpp
+++ b/Engine/Source/Editor/AnimGraph/Private/AnimBlueprintExtension_Base.cpp
@@ -852,7 +852,7 @@ void UAnimBlueprintExtension_Base::FEvaluationHandlerRecord::PatchAnimNodeExpose
 				// Skip invalid copy indices as these are usually the result of BP errors/warnings
 				if (PropertyCopyRecord.LibraryCompiledHandle.IsValid() && PropertyCopyRecord.LibraryCompiledHandle.GetBatchId() == (int32)EAnimPropertyAccessCallSite::WorkerThread_Unbatched)
 				{
-					Handler->CopyRecords.Emplace(PropertyCopyRecord.LibraryCompiledHandle.GetId(), PropertyCopyRecord.Operation);
+					Handler->CopyRecords.Emplace(PropertyCopyRecord.LibraryCompiledHandle.GetId(), PropertyCopyRecord.Operation, PropertyCopyRecord.bOnlyUpdateWhenActive);
 				}
 			}
 		}
@@ -991,6 +991,7 @@ void UAnimBlueprintExtension_Base::FEvaluationHandlerRecord::RegisterPropertyBin
 	CopyRecord.DestProperty = InProperty;
 	CopyRecord.DestArrayIndex = InBinding.ArrayIndex;
 	CopyRecord.BindingContextId = InBinding.ContextId;
+	CopyRecord.bOnlyUpdateWhenActive = InBinding.bOnlyUpdateWhenActive;
 }
 
 void UAnimBlueprintExtension_Base::FEvaluationHandlerRecord::BuildFastPathCopyRecords(IAnimBlueprintPostExpansionStepContext& InCompilationContext)
diff --git a/Engine/Source/Editor/AnimGraph/Private/AnimGraphNodeBinding_Base.cpp b/Engine/Source/Editor/AnimGraph/Private/AnimGraphNodeBinding_Base.cpp
--- a/Engine/Source/Editor/AnimGraph/Private/AnimGraphNodeBinding_Base.cpp
+++ b/Engine/Source/Editor/AnimGraph/Private/AnimGraphNodeBinding_Base.cpp
@@ -937,6 +937,98 @@ TSharedRef<SWidget> UAnimGraphNodeBinding_Base::MakePropertyBindingWidget(const
 						NAME_None,
 						EUserInterfaceActionType::Check
 					);
+
+					// Only Update When Active flag.
+					{
+						auto ToggleOnlyUpdateWhenActive = [InArgs, Blueprint]()
+						{
+							// by default, the new desired value is true...
+							bool bNewOnlyUpdateWhenActive = true;
+							for(const UAnimGraphNode_Base* AnimGraphNode : InArgs.Nodes)
+							{
+								if (const UAnimGraphNodeBinding_Base* GraphNodeBinding = Cast<const UAnimGraphNodeBinding_Base>(AnimGraphNode->GetBinding()))
+								{
+									if (const FAnimGraphNodePropertyBinding* PropertyBinding = GraphNodeBinding->PropertyBindings.Find(InArgs.BindingName))
+									{
+										//... unless any currently selected nodes already have it enabled. Set all to false to respect multiple values behavior in blueprints.
+										if (PropertyBinding->bOnlyUpdateWhenActive == true)
+										{
+											bNewOnlyUpdateWhenActive = false;
+											break;
+										}
+									}
+								}
+							}
+
+							FScopedTransaction Transaction(LOCTEXT("ToggleOnlyUpdateWhenActive", "Toggle Only Update When Active"));
+							for(UAnimGraphNode_Base* AnimGraphNode : InArgs.Nodes)
+							{
+								AnimGraphNode->Modify();
+								if (UAnimGraphNodeBinding_Base* GraphNodeBinding = Cast<UAnimGraphNodeBinding_Base>(AnimGraphNode->GetMutableBinding()))
+								{
+									if (FAnimGraphNodePropertyBinding* PropertyBinding = GraphNodeBinding->PropertyBindings.Find(InArgs.BindingName))
+									{
+										PropertyBinding->bOnlyUpdateWhenActive = bNewOnlyUpdateWhenActive;
+									}
+								}
+							}
+							FBlueprintEditorUtils::MarkBlueprintAsModified(Blueprint);
+						};
+
+						auto GetOnlyUpdateWhenActiveState = [InArgs]()
+						{
+							ECheckBoxState OutCheckboxState = ECheckBoxState::Unchecked;
+							for(const UAnimGraphNode_Base* AnimGraphNode : InArgs.Nodes)
+							{
+								if (const UAnimGraphNodeBinding_Base* GraphNodeBinding = Cast<const UAnimGraphNodeBinding_Base>(AnimGraphNode->GetBinding()))
+								{
+									if (const FAnimGraphNodePropertyBinding* PropertyBinding = GraphNodeBinding->PropertyBindings.Find(InArgs.BindingName))
+									{
+										if (PropertyBinding->bOnlyUpdateWhenActive == true)
+										{
+											// If any selected node has the option enabled, add a check mark.
+											// @todo: Ideally for multiple values, we'd use a different icon, but that's not what the other settings do.
+											OutCheckboxState = ECheckBoxState::Checked;
+											break;
+										}
+									}
+								}
+							}
+
+							return OutCheckboxState;
+						};
+
+						auto CanToggleUpdateWhenActiveState = [InArgs]()
+						{
+							for(const UAnimGraphNode_Base* AnimGraphNode : InArgs.Nodes)
+							{
+								if (const UAnimGraphNodeBinding_Base* GraphNodeBinding = Cast<const UAnimGraphNodeBinding_Base>(AnimGraphNode->GetBinding()))
+								{
+									// If any selected node doesn't have a binding object, disable this option.
+									// We only support this feature through property access bindings on the node.
+									if (GraphNodeBinding->PropertyBindings.Find(InArgs.BindingName) == nullptr)
+									{
+										return false;
+									}
+								}
+							}
+
+							return true;
+						};
+
+						InMenuBuilder.AddMenuEntry(
+							LOCTEXT("OnlyUpdateWhenActive", "Only Update When Active"),
+							LOCTEXT("OnlyUpdateWhenActiveTooltip", "Only update this property when the node is in an active graph branch (not blending out). To enable this, the property must be bound through property access."),
+							FSlateIcon(),
+							FUIAction(
+								FExecuteAction::CreateLambda(ToggleOnlyUpdateWhenActive),
+								FCanExecuteAction::CreateLambda(CanToggleUpdateWhenActiveState),
+								FGetActionCheckState::CreateLambda(GetOnlyUpdateWhenActiveState)
+							),
+							NAME_None,
+							EUserInterfaceActionType::Check
+						);
+					}
 				}
 			}
 			InMenuBuilder.EndSection();
diff --git a/Engine/Source/Editor/AnimGraph/Public/AnimGraphNode_Base.h b/Engine/Source/Editor/AnimGraph/Public/AnimGraphNode_Base.h
--- a/Engine/Source/Editor/AnimGraph/Public/AnimGraphNode_Base.h
+++ b/Engine/Source/Editor/AnimGraph/Public/AnimGraphNode_Base.h
@@ -183,6 +183,9 @@ struct FAnimGraphNodePropertyBinding
 	/** Whether the pin binding is a promotion (e.g. bool->int) */
 	UPROPERTY()
 	bool bIsPromotion = false;
+
+	UPROPERTY()
+	bool bOnlyUpdateWhenActive = false;
 };
 
 /**
@@ -206,7 +209,7 @@ class ANIMGRAPH_API UAnimGraphNode_Base : public UK2Node
 	/** Properties marked as always dynamic, so they can be set externally */
 	UPROPERTY()
 	TSet<FName> AlwaysDynamicProperties;
-	
+
 	UPROPERTY(Transient)
 	EBlueprintUsage BlueprintUsage;
 
diff --git a/Engine/Source/Runtime/AnimGraphRuntime/Private/AnimNodes/AnimNode_BlendListBase.cpp b/Engine/Source/Runtime/AnimGraphRuntime/Private/AnimNodes/AnimNode_BlendListBase.cpp
--- a/Engine/Source/Runtime/AnimGraphRuntime/Private/AnimNodes/AnimNode_BlendListBase.cpp
+++ b/Engine/Source/Runtime/AnimGraphRuntime/Private/AnimNodes/AnimNode_BlendListBase.cpp
@@ -231,7 +231,8 @@ void FAnimNode_BlendListBase::Update_AnyThread(const FAnimationUpdateContext& Co
 			const float BlendWeight = PerBlendData[i].Weight;
 			if (BlendWeight > ZERO_ANIMWEIGHT_THRESH)
 			{
-				BlendPose[i].Update(Context.FractionalWeight(BlendWeight));
+				FAnimationUpdateContext ChildContext = Context.FractionalWeight(BlendWeight);
+				BlendPose[i].Update((i == ChildIndex) ? ChildContext : ChildContext.AsInactive());
 			}
 		}
 
diff --git a/Engine/Source/Runtime/Engine/Classes/Animation/AnimNodeBase.h b/Engine/Source/Runtime/Engine/Classes/Animation/AnimNodeBase.h
--- a/Engine/Source/Runtime/Engine/Classes/Animation/AnimNodeBase.h
+++ b/Engine/Source/Runtime/Engine/Classes/Animation/AnimNodeBase.h
@@ -278,7 +278,14 @@ public:
 	// Get the previous node Id, set when we recurse into graph traversal functions from pose links
 	int32 GetPreviousNodeId() const { return PreviousNodeId; }
 
+	// Get whether the graph branch of this context is active (i.e. NOT blending out). 
+	bool IsActive() const { return bIsActive; }
+
 protected:
+	
+	// Whether this context belongs to graph branch (i.e. NOT blending out).
+	bool bIsActive = true;
+
 	// The current node ID, set when we recurse into graph traversal functions from pose links
 	int32 CurrentNodeId;
 
@@ -376,6 +383,14 @@ public:
 		return Result;
 	}
 
+	FAnimationUpdateContext AsInactive() const
+	{
+		FAnimationUpdateContext Result(*this);
+		Result.bIsActive = false;
+
+		return Result;
+	}
+
 	FAnimationUpdateContext FractionalWeight(float WeightMultiplier) const
 	{
 		FAnimationUpdateContext Result(*this);
diff --git a/Engine/Source/Runtime/Engine/Private/Animation/AnimNode_StateMachine.cpp b/Engine/Source/Runtime/Engine/Private/Animation/AnimNode_StateMachine.cpp
--- a/Engine/Source/Runtime/Engine/Private/Animation/AnimNode_StateMachine.cpp
+++ b/Engine/Source/Runtime/Engine/Private/Animation/AnimNode_StateMachine.cpp
@@ -830,8 +830,14 @@ void FAnimNode_StateMachine::UpdateTransitionStates(const FAnimationUpdateContex
 		case ETransitionLogicType::TLT_StandardBlend:
 			{
 				// update both states
-				UpdateState(Transition.PreviousState, Context.FractionalWeight(GetStateWeight(Transition.PreviousState)));
-				UpdateState(Transition.NextState, Context.FractionalWeight(GetStateWeight(Transition.NextState)));
+				{
+					FAnimationUpdateContext StateContext = Context.FractionalWeight(GetStateWeight(Transition.PreviousState));
+					UpdateState(Transition.PreviousState, Transition.PreviousState == CurrentState ? StateContext : StateContext.AsInactive());
+				}
+				{
+					FAnimationUpdateContext StateContext = Context.FractionalWeight(GetStateWeight(Transition.NextState));
+					UpdateState(Transition.NextState, (Transition.NextState == CurrentState) ? StateContext : StateContext.AsInactive());
+				}
 			}
 			break;
 
@@ -840,7 +846,7 @@ void FAnimNode_StateMachine::UpdateTransitionStates(const FAnimationUpdateContex
 				UE::Anim::TScopedGraphMessage<UE::Anim::FAnimInertializationSyncScope> InertializationSync(Context);
 
 				// update target state
-				UpdateState(Transition.NextState, Context);
+				UpdateState(Transition.NextState, (Transition.NextState == CurrentState) ? Context : Context.AsInactive());
 			}
 			break;
 
@@ -858,7 +864,7 @@ void FAnimNode_StateMachine::UpdateTransitionStates(const FAnimationUpdateContex
 							const bool bUsePreviousState = (Evaluator->DataSource == EEvaluatorDataSource::EDS_SourcePose);
 							const int32 EffectiveStateIndex = bUsePreviousState ? Transition.PreviousState : Transition.NextState;
 							FAnimationUpdateContext ContextToUse = Context.FractionalWeight(bUsePreviousState ? (1.0f - Transition.Alpha) : Transition.Alpha);
-							UpdateState(EffectiveStateIndex, ContextToUse);
+							UpdateState(EffectiveStateIndex, (EffectiveStateIndex == CurrentState) ? ContextToUse : ContextToUse.AsInactive());
 						}
 					}
 				}
diff --git a/Engine/Source/Runtime/Engine/Private/Animation/ExposedValueHandler.cpp b/Engine/Source/Runtime/Engine/Private/Animation/ExposedValueHandler.cpp
--- a/Engine/Source/Runtime/Engine/Private/Animation/ExposedValueHandler.cpp
+++ b/Engine/Source/Runtime/Engine/Private/Animation/ExposedValueHandler.cpp
@@ -97,14 +97,17 @@ void FAnimNodeExposedValueHandler_PropertyAccess::Execute(const FAnimationBaseCo
 			PropertyAccess::FCopyBatchId CopyBatch((int32)EAnimPropertyAccessCallSite::WorkerThread_Unbatched);
 			for (const FExposedValueCopyRecord& CopyRecord : CopyRecords)
 			{
-				PropertyAccess::ProcessCopy(AnimInstanceObject, *PropertyAccessLibrary, CopyBatch, CopyRecord.CopyIndex, [&CopyRecord](const FProperty* InProperty, void* InAddress)
+				if (!CopyRecord.bOnlyUpdateWhenActive || InContext.IsActive())
 				{
-					if (CopyRecord.PostCopyOperation == EPostCopyOperation::LogicalNegateBool)
+					PropertyAccess::ProcessCopy(AnimInstanceObject, *PropertyAccessLibrary, CopyBatch, CopyRecord.CopyIndex, [&CopyRecord](const FProperty* InProperty, void* InAddress)
 					{
-						bool bValue = static_cast<const FBoolProperty*>(InProperty)->GetPropertyValue(InAddress);
-						static_cast<const FBoolProperty*>(InProperty)->SetPropertyValue(InAddress, !bValue);
-					}
-				});
+						if (CopyRecord.PostCopyOperation == EPostCopyOperation::LogicalNegateBool)
+						{
+							bool bValue = static_cast<const FBoolProperty*>(InProperty)->GetPropertyValue(InAddress);
+							static_cast<const FBoolProperty*>(InProperty)->SetPropertyValue(InAddress, !bValue);
+						}
+					});
+				}
 			}
 		}
 	}
diff --git a/Engine/Source/Runtime/Engine/Public/Animation/ExposedValueHandler.h b/Engine/Source/Runtime/Engine/Public/Animation/ExposedValueHandler.h
--- a/Engine/Source/Runtime/Engine/Public/Animation/ExposedValueHandler.h
+++ b/Engine/Source/Runtime/Engine/Public/Animation/ExposedValueHandler.h
@@ -27,9 +27,10 @@ struct FExposedValueCopyRecord
 
 	FExposedValueCopyRecord() = default;
 
-	FExposedValueCopyRecord(int32 InCopyIndex, EPostCopyOperation InPostCopyOperation)
+	FExposedValueCopyRecord(int32 InCopyIndex, EPostCopyOperation InPostCopyOperation, bool bInOnlyUpdateWhenActive)
 		: CopyIndex(InCopyIndex)
 		, PostCopyOperation(InPostCopyOperation)
+		, bOnlyUpdateWhenActive(bInOnlyUpdateWhenActive)
 	{
 	}
 
@@ -38,6 +39,9 @@ struct FExposedValueCopyRecord
 
 	UPROPERTY()
 	EPostCopyOperation PostCopyOperation = EPostCopyOperation::None;
+
+	UPROPERTY()
+	bool bOnlyUpdateWhenActive = false;
 };
 
 // An exposed value updater
