Fix crash on loading an anim BP with pose hander nodes that require skeleton updates to their animations

This manifested as a problem dragging in MetHumans from Quixel Bridge.
Adding a preload dependency for asset player nodes (which include pose handlers) that recurses and preloads all dependencies too prevents load ordering issues where subobjects in anim assets were sometimes not correctly loaded

diff --git a/Engine/Source/Editor/AnimGraph/Private/AnimGraphNode_AssetPlayerBase.cpp b/Engine/Source/Editor/AnimGraph/Private/AnimGraphNode_AssetPlayerBase.cpp
--- a/Engine/Source/Editor/AnimGraph/Private/AnimGraphNode_AssetPlayerBase.cpp
+++ b/Engine/Source/Editor/AnimGraph/Private/AnimGraphNode_AssetPlayerBase.cpp
@@ -549,7 +549,17 @@ void UAnimGraphNode_AssetPlayerBase::PreloadRequiredAssetsHelper(UAnimationAsset
 		AssetToLoad = Cast<UAnimationAsset>(InExposedPin->DefaultObject);
 	}
 
-	PreloadObject(AssetToLoad);
+	if(AssetToLoad)
+	{
+		// HandleAnimReferenceCollection adds 'this' as well as any recursive dependencies
+		TArray<UAnimationAsset*> AssetAndDependencies;
+		AssetToLoad->HandleAnimReferenceCollection(AssetAndDependencies, true);
+
+		for (UAnimationAsset* Dependency : AssetAndDependencies)
+		{
+			PreloadObject(Dependency);
+		}
+	}
 }
 
 #undef LOCTEXT_NAMESPACE
\ No newline at end of file
