Fix node bindings not being present for nodes that are drag-dropped

diff --git a/Engine/Source/Editor/AnimGraph/Private/AnimGraphNode_Base.cpp b/Engine/Source/Editor/AnimGraph/Private/AnimGraphNode_Base.cpp
--- a/Engine/Source/Editor/AnimGraph/Private/AnimGraphNode_Base.cpp
+++ b/Engine/Source/Editor/AnimGraph/Private/AnimGraphNode_Base.cpp
@@ -53,19 +53,6 @@
 UAnimGraphNode_Base::UAnimGraphNode_Base(const FObjectInitializer& ObjectInitializer)
 	: Super(ObjectInitializer)
 {
-	if(!HasAnyFlags(RF_ClassDefaultObject))
-	{
-		if(UAnimBlueprint* AnimBlueprint = GetTypedOuter<UAnimBlueprint>())
-		{
-			UClass* BindingClass = AnimBlueprint->GetDefaultBindingClass();
-			if (BindingClass == nullptr)
-			{
-				BindingClass = UAnimGraphNodeBinding_Base::StaticClass();
-			}
-
-			Binding = CastChecked<UAnimGraphNodeBinding>(ObjectInitializer.CreateDefaultSubobject(this, TEXT("Binding"), BindingClass, BindingClass));
-		}
-	}
 }
 
 void UAnimGraphNode_Base::ExpandNode(FKismetCompilerContext& CompilerContext, UEdGraph* SourceGraph)
@@ -197,6 +184,31 @@ void UAnimGraphNode_Base::Serialize(FArchive& Ar)
 			NewBinding->PropertyBindings = PropertyBindings_DEPRECATED;
 			Binding = NewBinding;
 		}
+
+		if (Ar.CustomVer(FFortniteMainBranchObjectVersion::GUID) < FFortniteMainBranchObjectVersion::FixMissingAnimGraphNodeBindingExtensions)
+		{
+			if(Binding == nullptr)
+			{
+				Binding = NewObject<UAnimGraphNodeBinding_Base>(this);
+			}
+		}
+	}
+}
+
+void UAnimGraphNode_Base::EnsureBindingsArePresent()
+{
+	if(Binding == nullptr)
+	{
+		if (UAnimBlueprint* AnimBlueprint = GetTypedOuter<UAnimBlueprint>())
+		{
+			UClass* BindingClass = AnimBlueprint->GetDefaultBindingClass();
+			if (BindingClass == nullptr)
+			{
+				BindingClass = UAnimGraphNodeBinding_Base::StaticClass();
+			}
+
+			Binding = NewObject<UAnimGraphNodeBinding>(this, BindingClass);
+		}
 	}
 }
 
@@ -204,6 +216,8 @@ void UAnimGraphNode_Base::PostPlacedNewNode()
 {
 	Super::PostPlacedNewNode();
 
+	EnsureBindingsArePresent();
+
 	// This makes sure that all anim BP extensions are registered that this node needs
 	UAnimBlueprintExtension::RequestExtensionsForNode(this);
 }
@@ -212,6 +226,8 @@ void UAnimGraphNode_Base::PostPasteNode()
 {
 	Super::PostPasteNode();
 
+	EnsureBindingsArePresent();
+
 	// This makes sure that all anim BP extensions are registered that this node needs
 	UAnimBlueprintExtension::RequestExtensionsForNode(this);
 }
@@ -629,7 +645,10 @@ void UAnimGraphNode_Base::ProcessDuringCompilation(IAnimBlueprintCompilationCont
 	Extension->ProcessPosePins(this, InCompilationContext, OutCompiledData);
 
 	// Process bindings on this node
-	Binding->ProcessDuringCompilation(InCompilationContext, OutCompiledData);
+	if(Binding)
+	{
+		Binding->ProcessDuringCompilation(InCompilationContext, OutCompiledData);
+	}
 
 	// Resolve functions
 	GetFNode()->InitialUpdateFunction.SetFromFunction(InitialUpdateFunction.ResolveMember<UFunction>(GetBlueprintClassFromNode()));
diff --git a/Engine/Source/Editor/AnimGraph/Public/AnimGraphNode_Base.h b/Engine/Source/Editor/AnimGraph/Public/AnimGraphNode_Base.h
--- a/Engine/Source/Editor/AnimGraph/Public/AnimGraphNode_Base.h
+++ b/Engine/Source/Editor/AnimGraph/Public/AnimGraphNode_Base.h
@@ -648,6 +648,9 @@ protected:
 	/** Helper function used to validate anim node function references */
 	void ValidateFunctionRef(FName InPropertyName, const FMemberReference& InRef, const FText& InFunctionName, FCompilerResultsLog& MessageLog);
 	
+	// Create the bindings subobject if required
+	void EnsureBindingsArePresent();
+
 protected:
 	// Old shown pins. Needs to be a member variable to track pin visibility changes between Pre and PostEditChange 
 	TArray<FName> OldShownPins;
diff --git a/Engine/Source/Runtime/Core/Public/UObject/FortniteMainBranchObjectVersion.h b/Engine/Source/Runtime/Core/Public/UObject/FortniteMainBranchObjectVersion.h
--- a/Engine/Source/Runtime/Core/Public/UObject/FortniteMainBranchObjectVersion.h
+++ b/Engine/Source/Runtime/Core/Public/UObject/FortniteMainBranchObjectVersion.h
@@ -354,6 +354,9 @@ struct FFortniteMainBranchObjectVersion
 		// Function data stores a map from work to debug operands
 		RigVMSaveDebugMapInGraphFunctionData,
 
+		// Fix missing binding extensions for some anim graph nodes
+		FixMissingAnimGraphNodeBindingExtensions,
+
 		// -----<new versions can be added above this line>-------------------------------------------------
 		VersionPlusOne,
 		LatestVersion = VersionPlusOne - 1
