Fixed BlueprintEditor not getting refreshed after adding an animation layer node into an animation graph

diff --git a/Engine/Source/Editor/AnimGraph/Private/AnimationGraphSchema.cpp b/Engine/Source/Editor/AnimGraph/Private/AnimationGraphSchema.cpp
--- a/Engine/Source/Editor/AnimGraph/Private/AnimationGraphSchema.cpp
+++ b/Engine/Source/Editor/AnimGraph/Private/AnimationGraphSchema.cpp
@@ -119,6 +119,17 @@ FReply FAnimationLayerDragDropAction::DroppedOnPanel(const TSharedRef< class SWi
 			LinkedInputLayerNodeCreator.Finalize();
 			LinkedAnimLayerNode->NodePosX = static_cast<int32>(GraphPosition.X);
 			LinkedAnimLayerNode->NodePosY = static_cast<int32>(GraphPosition.Y);
+
+			UBlueprint* Blueprint = FBlueprintEditorUtils::FindBlueprintForGraphChecked(&Graph);
+			// See if we need to recompile skeleton after adding this node, or just mark dirty
+			if (LinkedAnimLayerNode->NodeCausesStructuralBlueprintChange())
+			{
+				FBlueprintEditorUtils::MarkBlueprintAsStructurallyModified(Blueprint);
+			}
+			else
+			{
+				FBlueprintEditorUtils::MarkBlueprintAsModified(Blueprint);
+			}
 		}
 	}
 	return FReply::Unhandled();
