Fix anim graph node fast path icon not appearing for nodes that do not use blueprint logic

If a node has exposed pins, it will now show up as 'fast path' when all of its pins are constant (as well as when it is actually using the fast path). If a node has no pins exposed it still does not show up as 'fast path', by design to reduce visual clutter.

diff --git a/Engine/Source/Editor/AnimGraph/Internal/AnimBlueprintExtension_Base.h b/Engine/Source/Editor/AnimGraph/Internal/AnimBlueprintExtension_Base.h
--- a/Engine/Source/Editor/AnimGraph/Internal/AnimBlueprintExtension_Base.h
+++ b/Engine/Source/Editor/AnimGraph/Internal/AnimBlueprintExtension_Base.h
@@ -183,6 +183,9 @@ private:
 		// Whether or not our serviced properties are actually on the instance instead of the node
 		bool bServicesInstanceProperties;
 
+		// This eval handler has properties
+		bool bHasProperties;
+
 		// Set of properties serviced by this handler (Map from property name to the record for that property)
 		TMap<FName, FAnimNodeSinglePropertyHandler> ServicedProperties;
 
@@ -200,6 +203,7 @@ private:
 			, EvaluationHandlerIdx(INDEX_NONE)
 			, bServicesNodeProperties(false)
 			, bServicesInstanceProperties(false)
+			, bHasProperties(false)
 			, HandlerFunctionName(NAME_None)
 		{}
 
diff --git a/Engine/Source/Editor/AnimGraph/Private/AnimBlueprintExtension_Base.cpp b/Engine/Source/Editor/AnimGraph/Private/AnimBlueprintExtension_Base.cpp
--- a/Engine/Source/Editor/AnimGraph/Private/AnimBlueprintExtension_Base.cpp
+++ b/Engine/Source/Editor/AnimGraph/Private/AnimBlueprintExtension_Base.cpp
@@ -51,6 +51,17 @@ void UAnimBlueprintExtension_Base::HandleCopyTermDefaultsToDefaultObject(UObject
 	{
 		Subsystem.PatchValueHandlers(DefaultAnimInstance->GetClass());
 
+		// Update blueprint usage of all graph nodes that have properties exposed
+		for (const TPair<UAnimGraphNode_Base*, FEvaluationHandlerRecord>& EvaluationHandlerPair : PerNodeStructEvalHandlers)
+		{
+			if (EvaluationHandlerPair.Value.bHasProperties && EvaluationHandlerPair.Value.ServicedProperties.Num() == 0)
+			{
+				UAnimGraphNode_Base* Node = CastChecked<UAnimGraphNode_Base>(EvaluationHandlerPair.Key);
+				UAnimGraphNode_Base* TrueNode = InCompilationContext.GetMessageLog().FindSourceObjectTypeChecked<UAnimGraphNode_Base>(Node);
+				TrueNode->BlueprintUsage = EBlueprintUsage::DoesNotUseBlueprint;
+			}
+		}
+
 		for(const FEvaluationHandlerRecord& EvaluationHandler : ValidEvaluationHandlerList)
 		{
 			if(EvaluationHandler.AnimGraphNode)
@@ -324,6 +335,8 @@ void UAnimBlueprintExtension_Base::ProcessNonPosePins(UAnimGraphNode_Base* InNod
 			
 			if (SourcePinProperty != NULL)
 			{
+				EvalHandler.bHasProperties = true;
+
 				if (SourcePin->LinkedTo.Num() == 0)
 				{
 					if(EnumHasAnyFlags(InFlags, EPinProcessingFlags::Constants))
