Drive-by cook determinism fix

diff --git a/Engine/Source/Editor/AnimGraph/Internal/AnimBlueprintExtension_Base.h b/Engine/Source/Editor/AnimGraph/Internal/AnimBlueprintExtension_Base.h
--- a/Engine/Source/Editor/AnimGraph/Internal/AnimBlueprintExtension_Base.h
+++ b/Engine/Source/Editor/AnimGraph/Internal/AnimBlueprintExtension_Base.h
@@ -73,6 +73,7 @@ private:
 			, BindingContextId(NAME_None)
 			, Operation(EPostCopyOperation::None)
 			, bIsFastPath(true)
+			, bOnlyUpdateWhenActive(false)
 		{}
 
 		FPropertyCopyRecord(const TArray<FString>& InSourcePropertyPath, const TArray<FString>& InDestPropertyPath)
