Fixed being unable to select external BPs in linked anim layer nodes

Switched custom parsing to use FEditorClassUtils::GetImplementedInterfaceClassPathsFromAsset

diff --git a/Engine/Source/Editor/AnimGraph/Private/AnimGraphNode_LinkedAnimLayer.cpp b/Engine/Source/Editor/AnimGraph/Private/AnimGraphNode_LinkedAnimLayer.cpp
--- a/Engine/Source/Editor/AnimGraph/Private/AnimGraphNode_LinkedAnimLayer.cpp
+++ b/Engine/Source/Editor/AnimGraph/Private/AnimGraphNode_LinkedAnimLayer.cpp
@@ -23,6 +23,7 @@
 #include "BlueprintActionDatabaseRegistrar.h"
 #include "ObjectEditorUtils.h"
 #include "AssetRegistry/AssetRegistryModule.h"
+#include "EditorClassUtils.h"
 
 #define LOCTEXT_NAMESPACE "LinkedAnimLayerNode"
 
@@ -612,34 +613,18 @@ bool UAnimGraphNode_LinkedAnimLayer::OnShouldFilterInstanceBlueprint(const FAsse
 
 			do 
 			{
-				const FString ImplementedInterfaces = CurrentAssetData.GetTagValueRef<FString>(FBlueprintTags::ImplementedInterfaces);
+				TArray<FString> InterfacePaths;
+				FEditorClassUtils::GetImplementedInterfaceClassPathsFromAsset(CurrentAssetData, InterfacePaths);
 
-				if(!ImplementedInterfaces.IsEmpty())
+				for (const FString& InterfacePath : InterfacePaths)
 				{
-					FString FullInterface;
-					FString RemainingString;
-					FString InterfacePath;
-					FString CurrentString = *ImplementedInterfaces;
-					while(CurrentString.Split(TEXT(","), &FullInterface, &RemainingString) && !bMatchesInterface)
-					{
-						if (!CurrentString.StartsWith(TEXT("Graphs=(")))
-						{
-							if (FullInterface.Split(TEXT("\""), &CurrentString, &InterfacePath, ESearchCase::CaseSensitive))
-							{
-								// The interface paths in metadata end with "', so remove those
-								InterfacePath.RemoveFromEnd(TEXT("\"'"));
+					FTopLevelAssetPath AssetPath(InterfacePath);
+					FCoreRedirectObjectName ResolvedInterfaceName = FCoreRedirects::GetRedirectedName(ECoreRedirectFlags::Type_Class, FCoreRedirectObjectName(AssetPath));
 
-								FCoreRedirectObjectName ResolvedInterfaceName = FCoreRedirects::GetRedirectedName(ECoreRedirectFlags::Type_Class, FCoreRedirectObjectName(InterfacePath));
-
-								// Verify against all interfaces we currently implement
-								for(TSubclassOf<UInterface> AnimInterface : AnimInterfaces)
-								{
-									bMatchesInterface |= ResolvedInterfaceName.ObjectName == AnimInterface->GetFName();
-								}
-							}
-						}
-			
-						CurrentString = RemainingString;
+					// Verify against all interfaces we currently implement
+					for (TSubclassOf<UInterface> AnimInterface : AnimInterfaces)
+					{
+						bMatchesInterface |= ResolvedInterfaceName.ObjectName == AnimInterface->GetFName();
 					}
 				}
 
