Added Anim Blueprint Editor preference to disable showing assets in the context menu of anim graph. This results in a significantly faster context menu in large projects.

diff --git a/Engine/Source/Editor/AnimGraph/AnimGraph.Build.cs b/Engine/Source/Editor/AnimGraph/AnimGraph.Build.cs
--- a/Engine/Source/Editor/AnimGraph/AnimGraph.Build.cs
+++ b/Engine/Source/Editor/AnimGraph/AnimGraph.Build.cs
@@ -28,8 +28,8 @@ public class AnimGraph : ModuleRules
 				"UnrealEd",
                 "GraphEditor",
 				"PropertyEditor",
-				
-                "ContentBrowser",
+				"AnimationBlueprintEditor",
+				"ContentBrowser",
 				"KismetWidgets",
 				"ToolMenus",
 				"KismetCompiler",
@@ -46,11 +46,5 @@ public class AnimGraph : ModuleRules
                 "GraphEditor",
             }
 		);
-
-        PrivateIncludePathModuleNames.AddRange(
-            new string[] {
-                "AnimationBlueprintEditor",
-            }
-        );
 	}
 }
diff --git a/Engine/Source/Editor/AnimGraph/Private/AnimGraphNode_AimOffsetLookAt.cpp b/Engine/Source/Editor/AnimGraph/Private/AnimGraphNode_AimOffsetLookAt.cpp
--- a/Engine/Source/Editor/AnimGraph/Private/AnimGraphNode_AimOffsetLookAt.cpp
+++ b/Engine/Source/Editor/AnimGraph/Private/AnimGraphNode_AimOffsetLookAt.cpp
@@ -12,6 +12,7 @@
 #include "UObject/UObjectIterator.h"
 #include "ToolMenus.h"
 #include "AnimGraphCommands.h"
+#include "AnimationBlueprintEditorSettings.h"
 
 /////////////////////////////////////////////////////
 // UAnimGraphNode_RotationOffsetBlendSpace
@@ -113,7 +114,7 @@ void UAnimGraphNode_AimOffsetLookAt::GetMenuActions(FBlueprintActionDatabaseRegi
 		}
 		// else, the Blueprint database is specifically looking for actions pertaining to something different (not a BlendSpace asset)
 	}
-	else
+	else if (GetDefault<UAnimationBlueprintEditorSettings>()->bShowAssetsInBlueprintContextMenu)
 	{
 		UClass* NodeClass = GetClass();
 		for (TObjectIterator<UBlendSpace> BlendSpaceIt; BlendSpaceIt; ++BlendSpaceIt)
diff --git a/Engine/Source/Editor/AnimGraph/Private/AnimGraphNode_AssetPlayerBase.cpp b/Engine/Source/Editor/AnimGraph/Private/AnimGraphNode_AssetPlayerBase.cpp
--- a/Engine/Source/Editor/AnimGraph/Private/AnimGraphNode_AssetPlayerBase.cpp
+++ b/Engine/Source/Editor/AnimGraph/Private/AnimGraphNode_AssetPlayerBase.cpp
@@ -17,6 +17,7 @@
 #include "BlueprintActionDatabaseRegistrar.h"
 #include "BlueprintNodeTemplateCache.h"
 #include "Animation/AnimNode_AssetPlayerBase.h"
+#include "AnimationBlueprintEditorSettings.h"
 
 #define LOCTEXT_NAMESPACE "UAnimGraphNode_AssetPlayerBase"
 
@@ -308,14 +309,17 @@ void UAnimGraphNode_AssetPlayerBase::GetMenuActionsHelper(
 		}	
 		Filter.bRecursiveClasses = true;
 		
-		TArray<FAssetData> Assets;
-		AssetRegistryModule.Get().GetAssets(Filter, Assets);
-
-		for (const FAssetData& AssetData : Assets)
+		if (GetDefault<UAnimationBlueprintEditorSettings>()->bShowAssetsInBlueprintContextMenu)
 		{
-			if(AssetData.IsUAsset())
+			TArray<FAssetData> Assets;
+			AssetRegistryModule.Get().GetAssets(Filter, Assets);
+
+			for (const FAssetData& AssetData : Assets)
 			{
-				MakeActionFromAsset(AssetData);
+				if(AssetData.IsUAsset())
+				{
+					MakeActionFromAsset(AssetData);
+				}
 			}
 		}
 
@@ -333,7 +337,7 @@ void UAnimGraphNode_AssetPlayerBase::GetMenuActionsHelper(
 			MakeActionFromAsset(FAssetData());
 		}
 	}
-	else if (bIsObjectOfAssetType)
+	else if (bIsObjectOfAssetType && GetDefault<UAnimationBlueprintEditorSettings>()->bShowAssetsInBlueprintContextMenu)
 	{
 		MakeActionFromAsset(FAssetData(QueryObject));
 	}
diff --git a/Engine/Source/Editor/AnimationBlueprintEditor/Private/AnimationBlueprintEditorSettings.cpp b/Engine/Source/Editor/AnimationBlueprintEditor/Private/AnimationBlueprintEditorSettings.cpp
--- a/Engine/Source/Editor/AnimationBlueprintEditor/Private/AnimationBlueprintEditorSettings.cpp
+++ b/Engine/Source/Editor/AnimationBlueprintEditor/Private/AnimationBlueprintEditorSettings.cpp
@@ -1,9 +1,21 @@
 // Copyright Epic Games, Inc. All Rights Reserved.
 
 #include "AnimationBlueprintEditorSettings.h"
+#include "BlueprintActionDatabase.h"
 
 void UAnimationBlueprintEditorSettings::PostEditChangeProperty(FPropertyChangedEvent& PropertyChangedEvent)
 {
 	Super::PostEditChangeProperty(PropertyChangedEvent);
+
+	FProperty* ChangedProperty = PropertyChangedEvent.Property;
+	if (ChangedProperty)
+	{
+		// Refresh the action database to add/remove asset names from corresponding asset player's blueprint context actions.
+		if (ChangedProperty->GetFName() == GET_MEMBER_NAME_CHECKED(UAnimationBlueprintEditorSettings, bShowAssetsInBlueprintContextMenu))
+		{
+			FBlueprintActionDatabase::Get().RefreshAll();
+		}
+	}
+
 	OnSettingsChange.Broadcast(this, PropertyChangedEvent.ChangeType);
-}
+}
\ No newline at end of file
diff --git a/Engine/Source/Editor/AnimationBlueprintEditor/Public/AnimationBlueprintEditorSettings.h b/Engine/Source/Editor/AnimationBlueprintEditor/Public/AnimationBlueprintEditorSettings.h
--- a/Engine/Source/Editor/AnimationBlueprintEditor/Public/AnimationBlueprintEditorSettings.h
+++ b/Engine/Source/Editor/AnimationBlueprintEditor/Public/AnimationBlueprintEditorSettings.h
@@ -15,7 +15,7 @@
 
 
 // Settings for the Animation Blueprint Editor
-UCLASS(config = EditorPerProjectUserSettings)
+UCLASS(config = EditorPerProjectUserSettings, MinimalAPI)
 class UAnimationBlueprintEditorSettings : public UObject
 {
 	GENERATED_BODY()
@@ -29,6 +29,10 @@ public:
 	UPROPERTY(EditAnywhere, config, Category = "Graphs")
 	bool bShowGraphCornerText = true;
 
+	/** Whether to display assets on the anim blueprint's context menu. Toggling this option is a slow operation. */
+	UPROPERTY(EditAnywhere, config, Category = "Blueprint")
+	bool bShowAssetsInBlueprintContextMenu = true;
+
 public:
 	DECLARE_MULTICAST_DELEGATE_TwoParams(FOnUpdateSettingsMulticaster, const UAnimationBlueprintEditorSettings*, EPropertyChangeType::Type);
 	FOnUpdateSettingsMulticaster OnSettingsChange;
