Fix for circular dependency in Anim Graph
#rnx

diff --git a/Engine/Source/Editor/AnimGraph/AnimGraph.Build.cs b/Engine/Source/Editor/AnimGraph/AnimGraph.Build.cs
--- a/Engine/Source/Editor/AnimGraph/AnimGraph.Build.cs
+++ b/Engine/Source/Editor/AnimGraph/AnimGraph.Build.cs
@@ -44,6 +44,7 @@ public class AnimGraph : ModuleRules
 			new string[] {
                 "UnrealEd",
                 "GraphEditor",
+				"AnimationBlueprintEditor",
             }
 		);
 	}
