[Backout] - CL27251725
Original CL Desc
-----------------------------------------------------------------
Fix for circular dependency in Anim Graph
#rnx

diff --git a/Engine/Source/Editor/AnimGraph/AnimGraph.Build.cs b/Engine/Source/Editor/AnimGraph/AnimGraph.Build.cs
--- a/Engine/Source/Editor/AnimGraph/AnimGraph.Build.cs
+++ b/Engine/Source/Editor/AnimGraph/AnimGraph.Build.cs
@@ -44,7 +44,6 @@ public class AnimGraph : ModuleRules
 			new string[] {
                 "UnrealEd",
                 "GraphEditor",
-				"AnimationBlueprintEditor",
             }
 		);
 	}
