Handle preloading of animation assets that discover they have more dependencies after being loaded

diff --git a/Engine/Source/Editor/AnimGraph/Private/AnimGraphNode_AssetPlayerBase.cpp b/Engine/Source/Editor/AnimGraph/Private/AnimGraphNode_AssetPlayerBase.cpp
--- a/Engine/Source/Editor/AnimGraph/Private/AnimGraphNode_AssetPlayerBase.cpp
+++ b/Engine/Source/Editor/AnimGraph/Private/AnimGraphNode_AssetPlayerBase.cpp
@@ -552,12 +552,34 @@ void UAnimGraphNode_AssetPlayerBase::PreloadRequiredAssetsHelper(UAnimationAsset
 	if(AssetToLoad)
 	{
 		// HandleAnimReferenceCollection adds 'this' as well as any recursive dependencies
+		TSet<UAnimationAsset*> PreloadedAssets;
+		TSet<UAnimationAsset*> DiscoveredAssets;
 		TArray<UAnimationAsset*> AssetAndDependencies;
 		AssetToLoad->HandleAnimReferenceCollection(AssetAndDependencies, true);
+		DiscoveredAssets.Append(AssetAndDependencies);
 
-		for (UAnimationAsset* Dependency : AssetAndDependencies)
+		while(AssetAndDependencies.Num() != 0)
 		{
-			PreloadObject(Dependency);
+			UAnimationAsset* Asset = AssetAndDependencies.Pop(false);
+			if (!PreloadedAssets.Contains(Asset))
+			{
+				PreloadedAssets.Add(Asset);
+				PreloadObject(Asset);
+
+				// catch any extra dependencies found after preloading the object
+				// this will have some redundancy, but that can't be avoided without 
+				// reworking HandleAnimReferenceCollection
+				TArray<UAnimationAsset*> Dependencies;
+				Asset->HandleAnimReferenceCollection(Dependencies, true);
+				for (UAnimationAsset* AssetDependency : Dependencies)
+				{
+					if (!DiscoveredAssets.Contains(AssetDependency))
+					{
+						DiscoveredAssets.Add(AssetDependency);
+						AssetAndDependencies.Add(AssetDependency);
+					}
+				}
+			}
 		}
 	}
 }
