Add missing RestoreSplitPin calls

diff --git a/Engine/Source/Editor/AnimGraph/Private/AnimGraphNode_CallFunction.cpp b/Engine/Source/Editor/AnimGraph/Private/AnimGraphNode_CallFunction.cpp
--- a/Engine/Source/Editor/AnimGraph/Private/AnimGraphNode_CallFunction.cpp
+++ b/Engine/Source/Editor/AnimGraph/Private/AnimGraphNode_CallFunction.cpp
@@ -96,6 +96,8 @@ void UAnimGraphNode_CallFunction::ReallocatePinsDuringReconstruction(TArray<UEdG
 	Super::ReallocatePinsDuringReconstruction(InOldPins);
 	
 	AllocateFunctionPins();
+
+	RestoreSplitPins(InOldPins);
 }
 
 void UAnimGraphNode_CallFunction::AllocateDefaultPins()
diff --git a/Engine/Source/Editor/AnimGraph/Private/AnimGraphNode_LinkedInputPose.cpp b/Engine/Source/Editor/AnimGraph/Private/AnimGraphNode_LinkedInputPose.cpp
--- a/Engine/Source/Editor/AnimGraph/Private/AnimGraphNode_LinkedInputPose.cpp
+++ b/Engine/Source/Editor/AnimGraph/Private/AnimGraphNode_LinkedInputPose.cpp
@@ -289,6 +289,8 @@ void UAnimGraphNode_LinkedInputPose::ReallocatePinsDuringReconstruction(TArray<U
 	Super::ReallocatePinsDuringReconstruction(OldPins);
 
 	AllocatePinsInternal();
+
+	RestoreSplitPins(OldPins);
 }
 
 void UAnimGraphNode_LinkedInputPose::CreateUserDefinedPins()
