Only process a sequence once when adding multiple keys
# jira UE-193618

diff --git a/Engine/Source/Editor/AnimGraph/Private/AnimPreviewInstance.cpp b/Engine/Source/Editor/AnimGraph/Private/AnimPreviewInstance.cpp
--- a/Engine/Source/Editor/AnimGraph/Private/AnimPreviewInstance.cpp
+++ b/Engine/Source/Editor/AnimGraph/Private/AnimPreviewInstance.cpp
@@ -377,6 +377,9 @@ void FAnimPreviewInstanceProxy::SetKeyImplementation(const FCompactPose& PreCont
 		CurrentSequence->Modify(true);
 		GetAnimInstanceObject()->Modify();
 
+		IAnimationDataController& Controller = CurrentSequence->GetController();
+		Controller.OpenBracket(LOCTEXT("AnimPreviewInstanceProxy_SetKey", "Set Key"));
+
 		TArray<FName> BonesToModify;
 		// need to get component transform first. Depending on when this gets called, the transform is not up-to-date. 
 		// first look at the bonecontrollers, and convert each bone controller to transform curve key
@@ -417,6 +420,8 @@ void FAnimPreviewInstanceProxy::SetKeyImplementation(const FCompactPose& PreCont
 			}
 		}
 
+		Controller.CloseBracket();
+
 		ResetModifiedBone(false);
 
 		OnSetKeyCompleteDelegate.Broadcast();
