Changed display name of AnimNode_TwoWayBlend from "Blend" to "Two Way Blend"

diff --git a/Engine/Source/Editor/AnimGraph/Private/AnimGraphNode_TwoWayBlend.cpp b/Engine/Source/Editor/AnimGraph/Private/AnimGraphNode_TwoWayBlend.cpp
--- a/Engine/Source/Editor/AnimGraph/Private/AnimGraphNode_TwoWayBlend.cpp
+++ b/Engine/Source/Editor/AnimGraph/Private/AnimGraphNode_TwoWayBlend.cpp
@@ -32,7 +32,7 @@ FText UAnimGraphNode_TwoWayBlend::GetTooltipText() const
 
 FText UAnimGraphNode_TwoWayBlend::GetNodeTitle(ENodeTitleType::Type TitleType) const
 {
-	return LOCTEXT("Blend", "Blend");
+	return LOCTEXT("TwoWayBlend", "Two Way Blend");
 }
 
 void UAnimGraphNode_TwoWayBlend::CustomizePinData(UEdGraphPin* Pin, FName SourcePropertyName, int32 ArrayIndex) const
