Ensure that we go to Super::GetPinBindingInfo if UAnimGraphNode_CustomProperty::GetPinBindingInfo cannot resolve the binding

diff --git a/Engine/Source/Editor/AnimGraph/Private/AnimGraphNode_CustomProperty.cpp b/Engine/Source/Editor/AnimGraph/Private/AnimGraphNode_CustomProperty.cpp
--- a/Engine/Source/Editor/AnimGraph/Private/AnimGraphNode_CustomProperty.cpp
+++ b/Engine/Source/Editor/AnimGraph/Private/AnimGraphNode_CustomProperty.cpp
@@ -302,8 +302,11 @@ bool UAnimGraphNode_CustomProperty::GetPinBindingInfo(FName InPinName, FName& Ou
 			return OutPinProperty->GetFName() == InOptionalPin.PropertyName;
 		});
 
-		OutBindingName = *GetPinTargetVariableName(InPinName); 
-		return OutOptionalPinIndex != INDEX_NONE;
+		if (OutOptionalPinIndex != INDEX_NONE)
+        {
+        	OutBindingName = *GetPinTargetVariableName(InPinName); 
+			return true;
+        }
 	}
 
 	return Super::GetPinBindingInfo(InPinName, OutBindingName, OutPinProperty, OutOptionalPinIndex);
