Animation layers now will start on a shared group by default, making them share instance at runtime

diff --git a/Engine/Source/Editor/AnimGraph/Private/AnimGraphDetails.cpp b/Engine/Source/Editor/AnimGraph/Private/AnimGraphDetails.cpp
--- a/Engine/Source/Editor/AnimGraph/Private/AnimGraphDetails.cpp
+++ b/Engine/Source/Editor/AnimGraph/Private/AnimGraphDetails.cpp
@@ -329,9 +329,13 @@ FReply FAnimGraphDetails::OnRemoveInputPoseClicked(UAnimGraphNode_LinkedInputPos
 FText FAnimGraphDetails::OnGetGroupText() const
 {
 	UAnimGraphNode_Root* Root = FBlueprintEditorUtils::GetAnimGraphRoot(Graph);
-	if(Root->Node.GetGroup() == NAME_None)
+	if (Root->Node.GetGroup() == FAnimNode_Root::DefaultSharedGroup)
 	{
-		return LOCTEXT("DefaultGroup", "Default");
+		return LOCTEXT("DefaultGroupSharedGroup", "Default Shared Group");
+	}
+	else if(Root->Node.GetGroup() == NAME_None)
+	{
+		return LOCTEXT("DefaultGroupUngrouped", "Ungrouped");
 	}
 		
 	return FText::FromName(Root->Node.GetGroup());
@@ -343,7 +347,7 @@ void FAnimGraphDetails::OnGroupTextCommitted(const FText& NewText, ETextCommit::
 	{
 		// Remove excess whitespace and prevent categories with just spaces
 		FText GroupName = FText::TrimPrecedingAndTrailing(NewText);
-		if(GroupName.ToString().Equals(TEXT("Default")))
+		if(GroupName.ToString().Equals(TEXT("Ungrouped")))
 		{
 			GroupName = FText::GetEmpty();
 		}
@@ -361,7 +365,7 @@ void FAnimGraphDetails::OnGroupSelectionChanged(TSharedPtr<FText> ProposedSelect
 	if(ProposedSelection.IsValid())
 	{
 		FText GroupName = *ProposedSelection.Get();
-		if(GroupName.ToString().Equals(TEXT("Default")))
+		if(GroupName.ToString().Equals(TEXT("Ungrouped")))
 		{
 			GroupName = FText::GetEmpty();
 		}
@@ -402,7 +406,7 @@ void FAnimGraphDetails::RefreshGroupSource()
 	UAnimGraphNode_Root* Root = FBlueprintEditorUtils::GetAnimGraphRoot(Graph);
 	if(Root->Node.GetGroup() != NAME_None)
 	{
-		GroupSource.Add(MakeShared<FText>(LOCTEXT("DefaultGroup", "Default")));
+		GroupSource.Add(MakeShared<FText>(LOCTEXT("DefaultGroupUngrouped", "Ungrouped")));
 	}
 
 	// Pull groups from implemented functions
diff --git a/Engine/Source/Editor/AnimGraph/Private/AnimGraphNode_Root.cpp b/Engine/Source/Editor/AnimGraph/Private/AnimGraphNode_Root.cpp
--- a/Engine/Source/Editor/AnimGraph/Private/AnimGraphNode_Root.cpp
+++ b/Engine/Source/Editor/AnimGraph/Private/AnimGraphNode_Root.cpp
@@ -5,6 +5,7 @@
 #include "AnimBlueprintCompiler.h"
 #include "IAnimBlueprintCompilationContext.h"
 #include "IAnimBlueprintCopyTermDefaultsContext.h"
+#include "UObject/FortniteMainBranchObjectVersion.h"
 
 /////////////////////////////////////////////////////
 // FPoseLinkMappingRecord
@@ -44,6 +45,22 @@ UAnimGraphNode_Root::UAnimGraphNode_Root(const FObjectInitializer& ObjectInitial
 {
 }
 
+void UAnimGraphNode_Root::Serialize(FArchive& Ar)
+{
+	Super::Serialize(Ar);
+
+	Ar.UsingCustomVersion(FFortniteMainBranchObjectVersion::GUID);
+
+	const int32 CustomAnimVersion = Ar.CustomVer(FFortniteMainBranchObjectVersion::GUID);
+
+	if (Ar.IsLoading() && CustomAnimVersion < FFortniteMainBranchObjectVersion::AnimNodeRootDefaultGroupChange)
+	{
+#if WITH_EDITORONLY_DATA
+		Node.LayerGroup = Node.Group_DEPRECATED;
+#endif
+	}
+}
+
 FLinearColor UAnimGraphNode_Root::GetNodeTitleColor() const
 {
 	return GetDefault<UGraphEditorSettings>()->ResultNodeTitleColor;
diff --git a/Engine/Source/Editor/AnimGraph/Public/AnimGraphNode_Root.h b/Engine/Source/Editor/AnimGraph/Public/AnimGraphNode_Root.h
--- a/Engine/Source/Editor/AnimGraph/Public/AnimGraphNode_Root.h
+++ b/Engine/Source/Editor/AnimGraph/Public/AnimGraphNode_Root.h
@@ -18,6 +18,10 @@ class UAnimGraphNode_Root : public UAnimGraphNode_Base
 	UPROPERTY(EditAnywhere, Category=Settings)
 	FAnimNode_Root Node;
 
+	//~ Begin UObject interface
+	ANIMGRAPH_API virtual void Serialize(FArchive& Ar) override;
+	//~ End of UObject interface
+
 	//~ Begin UEdGraphNode Interface.
 	ANIMGRAPH_API virtual FLinearColor GetNodeTitleColor() const override;
 	ANIMGRAPH_API virtual FText GetTooltipText() const override;
diff --git a/Engine/Source/Runtime/Core/Public/UObject/FortniteMainBranchObjectVersion.h b/Engine/Source/Runtime/Core/Public/UObject/FortniteMainBranchObjectVersion.h
--- a/Engine/Source/Runtime/Core/Public/UObject/FortniteMainBranchObjectVersion.h
+++ b/Engine/Source/Runtime/Core/Public/UObject/FortniteMainBranchObjectVersion.h
@@ -396,6 +396,9 @@ struct FFortniteMainBranchObjectVersion
 		// Fix tangents for non-uniform build scales, and add a flag to optionally match the previous (incorrect) tangents
 		FixedTangentTransformForNonuniformBuildScale,
 
+		// AnimNode Layers will now start in a Shared Group, instead of being each one on a different group at runtime
+		AnimNodeRootDefaultGroupChange,
+
 		// -----<new versions can be added above this line>-------------------------------------------------
 		VersionPlusOne,
 		LatestVersion = VersionPlusOne - 1
diff --git a/Engine/Source/Runtime/Engine/Classes/Animation/AnimNode_Root.h b/Engine/Source/Runtime/Engine/Classes/Animation/AnimNode_Root.h
--- a/Engine/Source/Runtime/Engine/Classes/Animation/AnimNode_Root.h
+++ b/Engine/Source/Runtime/Engine/Classes/Animation/AnimNode_Root.h
@@ -24,7 +24,10 @@ protected:
 
 	/** The group of this root node, used to group this output with others when used in a layer. */
 	UPROPERTY(meta=(FoldProperty))
-	FName Group;
+	FName LayerGroup = DefaultSharedGroup;
+
+	UPROPERTY(meta= (DeprecatedProperty, DeprecationMessage = "Please, use LayerGroup"))
+	FName Group_DEPRECATED;
 #endif
 
 public:	
@@ -43,7 +46,8 @@ public:
 	void SetName(FName InName) { Name = InName; }
 
 	/** Set the group of this root node, used to group this output with others when used in a layer. */
-	void SetGroup(FName InGroup) { Group = InGroup; }
+	void SetGroup(FName InGroup) { LayerGroup = InGroup; }
+
 #endif
 
 	/** Get the name of this root node, used to identify the output of this graph */
@@ -51,4 +55,11 @@ public:
 
 	/** Get the group of this root node, used to group this output with others when used in a layer. */
 	ENGINE_API FName GetGroup() const;
+
+#if WITH_EDITORONLY_DATA
+	/** Default shared group used by layers. */
+	ENGINE_API static FName DefaultSharedGroup;
+#endif
+
+	friend class UAnimGraphNode_Root;
 };
diff --git a/Engine/Source/Runtime/Engine/Private/Animation/AnimNode_Root.cpp b/Engine/Source/Runtime/Engine/Private/Animation/AnimNode_Root.cpp
--- a/Engine/Source/Runtime/Engine/Private/Animation/AnimNode_Root.cpp
+++ b/Engine/Source/Runtime/Engine/Private/Animation/AnimNode_Root.cpp
@@ -9,6 +9,10 @@
 /////////////////////////////////////////////////////
 // FAnimNode_Root
 
+#if WITH_EDITORONLY_DATA
+FName FAnimNode_Root::DefaultSharedGroup("DefaultSharedGroup"); // All layers sharing the instance by default
+#endif
+
 FAnimNode_Root::FAnimNode_Root()
 {
 }
@@ -49,5 +53,5 @@ FName FAnimNode_Root::GetName() const
 
 FName FAnimNode_Root::GetGroup() const
 {
-	return GET_ANIM_NODE_DATA(FName, Group);
+	return GET_ANIM_NODE_DATA(FName, LayerGroup);
 }
