Fixed new anim state machine transition nodes not appearing when the graph is far from the origin

Nodes were initially placed at the origin and would not be considered visible (for auto-rearranging between their states) until the origin was in view

diff --git a/Engine/Source/Editor/AnimGraph/Private/AnimationStateMachineSchema.cpp b/Engine/Source/Editor/AnimGraph/Private/AnimationStateMachineSchema.cpp
--- a/Engine/Source/Editor/AnimGraph/Private/AnimationStateMachineSchema.cpp
+++ b/Engine/Source/Editor/AnimGraph/Private/AnimationStateMachineSchema.cpp
@@ -245,7 +245,8 @@ bool UAnimationStateMachineSchema::CreateAutomaticConversionNodeAndConnections(U
 		&& (NodeA->GetInputPin() != NULL) && (NodeA->GetOutputPin() != NULL)
 		&& (NodeB->GetInputPin() != NULL) && (NodeB->GetOutputPin() != NULL))
 	{
-		UAnimStateTransitionNode* TransitionNode = FEdGraphSchemaAction_NewStateNode::SpawnNodeFromTemplate<UAnimStateTransitionNode>(NodeA->GetGraph(), NewObject<UAnimStateTransitionNode>(), FVector2D(0.0f, 0.0f), false);
+		FVector2D Location = (FVector2D(NodeA->NodePosX, NodeA->NodePosY) + FVector2D(NodeB->NodePosX, NodeB->NodePosY)) * 0.5f;
+		UAnimStateTransitionNode* TransitionNode = FEdGraphSchemaAction_NewStateNode::SpawnNodeFromTemplate<UAnimStateTransitionNode>(NodeA->GetGraph(), NewObject<UAnimStateTransitionNode>(), Location, false);
 
 		if (PinA->Direction == EGPD_Output)
 		{
