Fix unfixable warnings when variables referenced in in anim graph node pin bindings are removed

Also fix copy-pasted code in AnimNext
Also fix crash when compiling a template anim BP without a preview instance (found while attempting repro)

diff --git a/Engine/Plugins/Experimental/AnimNext/Source/AnimNextUncookedOnly/Private/Graph/AnimGraph/AnimBlueprintExtension_AnimNextParameters.cpp b/Engine/Plugins/Experimental/AnimNext/Source/AnimNextUncookedOnly/Private/Graph/AnimGraph/AnimBlueprintExtension_AnimNextParameters.cpp
--- a/Engine/Plugins/Experimental/AnimNext/Source/AnimNextUncookedOnly/Private/Graph/AnimGraph/AnimBlueprintExtension_AnimNextParameters.cpp
+++ b/Engine/Plugins/Experimental/AnimNext/Source/AnimNextUncookedOnly/Private/Graph/AnimGraph/AnimBlueprintExtension_AnimNextParameters.cpp
@@ -22,10 +22,10 @@ void UAnimBlueprintExtension_AnimNextParameters::ProcessNodePins(UAnimGraphNode_
 			NewRecord.AnimGraphNode = InNode;
 			NewRecord.NodeVariableProperty = NodeProperty;
 
-			for (TPair<FName, FAnimNextAnimGraphNodeParameterBinding>& PropertyBindingPair : Binding->PropertyBindings)
+			for(auto Iter = Binding->PropertyBindings.CreateIterator(); Iter; ++Iter)
 			{
-				FName BindingName = PropertyBindingPair.Key;
-				FAnimNextAnimGraphNodeParameterBinding& ParameterBinding = PropertyBindingPair.Value;
+				FName BindingName = Iter.Key();
+				FAnimNextAnimGraphNodeParameterBinding& ParameterBinding = Iter.Value();
 
 				// for array properties we need to account for the extra FName number 
 				FName ComparisonName = BindingName;
@@ -45,7 +45,8 @@ void UAnimBlueprintExtension_AnimNextParameters::ProcessNodePins(UAnimGraphNode_
 				}
 				else
 				{
-					InCompilationContext.GetMessageLog().Warning(*FString::Printf(TEXT("ICE: @@ Failed to find a property '%s'"), *ComparisonName.ToString()), InNode);
+					// Binding is no longer valid, remove it
+					Iter.RemoveCurrent();
 				}
 			}
 		}
diff --git a/Engine/Source/Editor/AnimGraph/Private/AnimBlueprintExtension_Base.cpp b/Engine/Source/Editor/AnimGraph/Private/AnimBlueprintExtension_Base.cpp
--- a/Engine/Source/Editor/AnimGraph/Private/AnimBlueprintExtension_Base.cpp
+++ b/Engine/Source/Editor/AnimGraph/Private/AnimBlueprintExtension_Base.cpp
@@ -400,32 +400,33 @@ void UAnimBlueprintExtension_Base::ProcessNonPosePins(UAnimGraphNode_Base* InNod
 		// Add any property bindings
 		if (UAnimGraphNodeBinding_Base* Binding = Cast<UAnimGraphNodeBinding_Base>(InNode->GetMutableBinding()))
 		{
-			for(const TPair<FName, FAnimGraphNodePropertyBinding>& PropertyBinding : Binding->PropertyBindings)
+			for(auto Iter = Binding->PropertyBindings.CreateIterator(); Iter; ++Iter)
 			{
-				if(PropertyBinding.Value.bIsBound)
+				if(Iter.Value().bIsBound)
 				{
 					FEvaluationHandlerRecord& EvalHandler = PerNodeStructEvalHandlers.FindOrAdd(InNode);
 					EvalHandler.AnimGraphNode = InNode;
 
 					// for array properties we need to account for the extra FName number 
-					FName ComparisonName = PropertyBinding.Key;
+					FName ComparisonName = Iter.Key();
 					ComparisonName.SetNumber(0);
 
 					if (FProperty* Property = FindFProperty<FProperty>(NodeProperty->Struct, ComparisonName))
 					{
 						EvalHandler.NodeVariableProperty = NodeProperty;
 						EvalHandler.bServicesNodeProperties = true;
-						EvalHandler.RegisterPropertyBinding(Property, PropertyBinding.Value);
+						EvalHandler.RegisterPropertyBinding(Property, Iter.Value());
 					}
-					else if(FProperty* ClassProperty = FindFProperty<FProperty>(InCompilationContext.GetBlueprint()->SkeletonGeneratedClass, PropertyBinding.Value.PropertyName))
+					else if(FProperty* ClassProperty = FindFProperty<FProperty>(InCompilationContext.GetBlueprint()->SkeletonGeneratedClass, Iter.Value().PropertyName))
 					{
 						EvalHandler.NodeVariableProperty = NodeProperty;
 						EvalHandler.bServicesInstanceProperties = true;
-						EvalHandler.RegisterPropertyBinding(ClassProperty, PropertyBinding.Value);
+						EvalHandler.RegisterPropertyBinding(ClassProperty, Iter.Value());
 					}
 					else
 					{
-						InCompilationContext.GetMessageLog().Warning(*FString::Printf(TEXT("ICE: @@ Failed to find a property '%s'"), *ComparisonName.ToString()), InNode);
+						// Binding is no longer valid, remove it
+						Iter.RemoveCurrent();
 					}
 				}
 			}
diff --git a/Engine/Source/Editor/AnimationBlueprintEditor/Private/AnimationBlueprintEditor.cpp b/Engine/Source/Editor/AnimationBlueprintEditor/Private/AnimationBlueprintEditor.cpp
--- a/Engine/Source/Editor/AnimationBlueprintEditor/Private/AnimationBlueprintEditor.cpp
+++ b/Engine/Source/Editor/AnimationBlueprintEditor/Private/AnimationBlueprintEditor.cpp
@@ -2301,7 +2301,10 @@ void FAnimationBlueprintEditor::HandleSetObjectBeingDebugged(UObject* InObject)
 			if (InObject->GetWorld()->IsPreviewWorld())
 			{
 				GetPreviewScene()->ShowDefaultMode();
-				GetPreviewScene()->GetPreviewMeshComponent()->PreviewInstance->SetDebugSkeletalMeshComponent(nullptr);
+				if(GetPreviewScene()->GetPreviewMeshComponent()->PreviewInstance)
+				{
+					GetPreviewScene()->GetPreviewMeshComponent()->PreviewInstance->SetDebugSkeletalMeshComponent(nullptr);
+				}
 				GetPreviewScene()->GetPreviewMeshComponent()->bTrackAttachedInstanceLOD = false;
 			}
 			else
@@ -2316,7 +2319,10 @@ void FAnimationBlueprintEditor::HandleSetObjectBeingDebugged(UObject* InObject)
 	{
 		// Clear the copy-pose component and set us back to 'normal'
 		GetPreviewScene()->ShowDefaultMode();
-		GetPreviewScene()->GetPreviewMeshComponent()->PreviewInstance->SetDebugSkeletalMeshComponent(nullptr);
+		if(GetPreviewScene()->GetPreviewMeshComponent()->PreviewInstance)
+		{
+			GetPreviewScene()->GetPreviewMeshComponent()->PreviewInstance->SetDebugSkeletalMeshComponent(nullptr);
+		}
 		GetPreviewScene()->GetPreviewMeshComponent()->bTrackAttachedInstanceLOD = false;
 	}
 }
