Fix anim BP node orphan pins not getting removed on disconnection

diff --git a/Engine/Source/Editor/AnimGraph/Private/AnimGraphNode_Base.cpp b/Engine/Source/Editor/AnimGraph/Private/AnimGraphNode_Base.cpp
--- a/Engine/Source/Editor/AnimGraph/Private/AnimGraphNode_Base.cpp
+++ b/Engine/Source/Editor/AnimGraph/Private/AnimGraphNode_Base.cpp
@@ -967,6 +967,8 @@ FProperty* UAnimGraphNode_Base::GetPinProperty(FName InPinName) const
 
 void UAnimGraphNode_Base::PinConnectionListChanged(UEdGraphPin* Pin)
 {
+	Super::PinConnectionListChanged(Pin);
+	
 	if(Pin->LinkedTo.Num() > 0)
 	{
 		if (Binding)
