Fix pasted 'anim getter' nodes introducing hidden references to their source anim BP

diff --git a/Engine/Source/Editor/AnimGraph/Private/K2Node_AnimGetter.cpp b/Engine/Source/Editor/AnimGraph/Private/K2Node_AnimGetter.cpp
--- a/Engine/Source/Editor/AnimGraph/Private/K2Node_AnimGetter.cpp
+++ b/Engine/Source/Editor/AnimGraph/Private/K2Node_AnimGetter.cpp
@@ -56,6 +56,7 @@ void UK2Node_AnimGetter::PostPasteNode()
 {
 	Super::PostPasteNode();
 
+	SourceAnimBlueprint = Cast<UAnimBlueprint>(GetBlueprint());
 	RestoreStateMachineState();
 	RestoreStateMachineNode();
 	UpdateCachedTitle();
@@ -346,31 +347,46 @@ bool UK2Node_AnimGetter::IsActionFilteredOut(FBlueprintActionFilter const& Filte
 
 void UK2Node_AnimGetter::RestoreStateMachineState()
 {
-	if (SourceStateNode)
+	UAnimationTransitionGraph* TransitionGraph = Cast<UAnimationTransitionGraph>(GetOuter());
+	if (TransitionGraph == nullptr)
 	{
-		if (UAnimationTransitionGraph* TransitionGraph = Cast<UAnimationTransitionGraph>(GetOuter()))
-		{
-			if (UAnimStateTransitionNode* StateTransitionNode = Cast<UAnimStateTransitionNode>(TransitionGraph->GetOuter()))
-			{
-				SourceStateNode = StateTransitionNode->GetPreviousState();
-			}
-		}
+		SourceStateNode = nullptr;
+		return;
+	}
+
+	UAnimStateTransitionNode* StateTransitionNode = Cast<UAnimStateTransitionNode>(TransitionGraph->GetOuter());
+	if (StateTransitionNode == nullptr)
+	{
+		SourceStateNode = nullptr;
+		return;
 	}
+
+	SourceStateNode = StateTransitionNode->GetPreviousState();
 }
 
 void UK2Node_AnimGetter::RestoreStateMachineNode()
 {
-	if(SourceStateNode)
+	if(SourceStateNode == nullptr)
 	{
-		UAnimationStateMachineGraph* Graph = Cast<UAnimationStateMachineGraph>(SourceStateNode->GetOuter());
-		if (Graph)
-		{
-			if (UAnimGraphNode_StateMachine* MachineNode = Cast<UAnimGraphNode_StateMachine>(Graph->GetOuter()))
-			{
-				SourceNode = MachineNode;
-			}
-		}
+		SourceNode = nullptr;
+		return;
+	}
+
+	UAnimationStateMachineGraph* Graph = Cast<UAnimationStateMachineGraph>(SourceStateNode->GetOuter());
+	if (Graph == nullptr)
+	{
+		SourceNode = nullptr;
+		return;
 	}
+
+	UAnimGraphNode_StateMachine* MachineNode = Cast<UAnimGraphNode_StateMachine>(Graph->GetOuter());
+	if (MachineNode == nullptr)
+	{
+		SourceNode = nullptr;
+		return;
+	}
+
+	SourceNode = MachineNode;
 }
 
 bool UK2Node_AnimGetter::GetterRequiresParameter(const UFunction* Getter, FString ParamName)
