Fix function renames not applying to property access nodes and compilation crashes post-rename

Implement function BP rename handling on a per-node basis, similar to existing variable rename handing
Move to using FBlueprintEditorUtils::GetMostUpToDateClass for property access in-editor. This ensures that any local renames are picked up prior to recompilation.

#changelist validated
#virtualized

diff --git a/Engine/Plugins/Developer/PropertyAccessNode/Source/PropertyAccessNode/Private/K2Node_PropertyAccess.cpp b/Engine/Plugins/Developer/PropertyAccessNode/Source/PropertyAccessNode/Private/K2Node_PropertyAccess.cpp
--- a/Engine/Plugins/Developer/PropertyAccessNode/Source/PropertyAccessNode/Private/K2Node_PropertyAccess.cpp
+++ b/Engine/Plugins/Developer/PropertyAccessNode/Source/PropertyAccessNode/Private/K2Node_PropertyAccess.cpp
@@ -316,7 +316,7 @@ void UK2Node_PropertyAccess::HandleVariableRenamed(UBlueprint* InBlueprint, UCla
 {
 	IPropertyAccessEditor& PropertyAccessEditor = IModularFeatures::Get().GetModularFeature<IPropertyAccessEditor>("PropertyAccessEditor");
 
-	UClass* SkeletonVariableClass = FBlueprintEditorUtils::GetSkeletonClass(InVariableClass);
+	UClass* SkeletonVariableClass = FBlueprintEditorUtils::GetMostUpToDateClass(InVariableClass);
 
 	// See if the path references the variable
 	TArray<int32> RenameIndices;
@@ -340,6 +340,34 @@ void UK2Node_PropertyAccess::HandleVariableRenamed(UBlueprint* InBlueprint, UCla
 	}
 }
 
+void UK2Node_PropertyAccess::HandleFunctionRenamed(UBlueprint* InBlueprint, UClass* InFunctionClass, UEdGraph* InGraph, const FName& InOldFuncName, const FName& InNewFuncName)
+{
+	IPropertyAccessEditor& PropertyAccessEditor = IModularFeatures::Get().GetModularFeature<IPropertyAccessEditor>("PropertyAccessEditor");
+
+	UClass* SkeletonFunctionClass = FBlueprintEditorUtils::GetMostUpToDateClass(InFunctionClass);
+
+	// See if the path references the variable
+	TArray<int32> RenameIndices;
+	IPropertyAccessEditor::FResolvePropertyAccessArgs ResolveArgs;
+	ResolveArgs.FunctionFunction = [InOldFuncName, SkeletonFunctionClass, &RenameIndices](int32 InSegmentIndex, UFunction* InFunction, FProperty* InReturnProperty)
+	{
+		const UClass* OwnerSkeletonFunctionClass = FBlueprintEditorUtils::GetMostUpToDateClass(InFunction->GetOuterUClass());
+		if(OwnerSkeletonFunctionClass && InFunction->GetFName() == InOldFuncName && OwnerSkeletonFunctionClass->IsChildOf(SkeletonFunctionClass))
+		{
+			RenameIndices.Add(InSegmentIndex);
+		}
+	};
+	
+	PropertyAccessEditor.ResolvePropertyAccess(GetBlueprint()->SkeletonGeneratedClass, Path, ResolveArgs);
+
+	// Rename any references we found
+	for(const int32& RenameIndex : RenameIndices)
+	{
+		Path[RenameIndex] = InNewFuncName.ToString();
+		TextPath = PropertyAccessEditor.MakeTextPath(Path, GetBlueprint()->SkeletonGeneratedClass);
+	}
+}
+
 void UK2Node_PropertyAccess::ReplaceReferences(UBlueprint* InBlueprint, UBlueprint* InReplacementBlueprint, const FMemberReference& InSource, const FMemberReference& InReplacement)
 {
 	IPropertyAccessEditor& PropertyAccessEditor = IModularFeatures::Get().GetModularFeature<IPropertyAccessEditor>("PropertyAccessEditor");
@@ -376,7 +404,7 @@ bool UK2Node_PropertyAccess::ReferencesVariable(const FName& InVarName, const US
 {
 	IPropertyAccessEditor& PropertyAccessEditor = IModularFeatures::Get().GetModularFeature<IPropertyAccessEditor>("PropertyAccessEditor");
 
-	const UClass* SkeletonVariableClass = FBlueprintEditorUtils::GetSkeletonClass(Cast<UClass>(InScope));
+	const UClass* SkeletonVariableClass = FBlueprintEditorUtils::GetMostUpToDateClass(Cast<UClass>(InScope));
 
 	// See if the path references the variable
 	bool bReferencesVariable = false;
@@ -386,7 +414,7 @@ bool UK2Node_PropertyAccess::ReferencesVariable(const FName& InVarName, const US
 	{
 		if(SkeletonVariableClass)
 		{
-			const UClass* OwnerSkeletonVariableClass = FBlueprintEditorUtils::GetSkeletonClass(Cast<UClass>(InProperty->GetOwnerStruct()));
+			const UClass* OwnerSkeletonVariableClass = FBlueprintEditorUtils::GetMostUpToDateClass(Cast<UClass>(InProperty->GetOwnerStruct()));
 
 			if(OwnerSkeletonVariableClass && InProperty->GetFName() == InVarName && OwnerSkeletonVariableClass->IsChildOf(SkeletonVariableClass))
 			{
@@ -408,7 +436,7 @@ bool UK2Node_PropertyAccess::ReferencesFunction(const FName& InFunctionName, con
 {
 	IPropertyAccessEditor& PropertyAccessEditor = IModularFeatures::Get().GetModularFeature<IPropertyAccessEditor>("PropertyAccessEditor");
 
-	const UClass* SkeletonFunctionClass = FBlueprintEditorUtils::GetSkeletonClass(Cast<UClass>(InScope));
+	const UClass* SkeletonFunctionClass = FBlueprintEditorUtils::GetMostUpToDateClass(Cast<UClass>(InScope));
 
 	// See if the path references the function
 	bool bReferencesFunction = false;
@@ -418,7 +446,7 @@ bool UK2Node_PropertyAccess::ReferencesFunction(const FName& InFunctionName, con
 	{
 		if (SkeletonFunctionClass)
 		{
-			const UClass* OwnerSkeletonFunctionClass = FBlueprintEditorUtils::GetSkeletonClass(InFunction->GetOuterUClass());
+			const UClass* OwnerSkeletonFunctionClass = FBlueprintEditorUtils::GetMostUpToDateClass(InFunction->GetOuterUClass());
 
 			if (OwnerSkeletonFunctionClass && InFunction->GetFName() == InFunctionName && OwnerSkeletonFunctionClass->IsChildOf(SkeletonFunctionClass))
 			{
diff --git a/Engine/Plugins/Developer/PropertyAccessNode/Source/PropertyAccessNode/Private/K2Node_PropertyAccess.h b/Engine/Plugins/Developer/PropertyAccessNode/Source/PropertyAccessNode/Private/K2Node_PropertyAccess.h
--- a/Engine/Plugins/Developer/PropertyAccessNode/Source/PropertyAccessNode/Private/K2Node_PropertyAccess.h
+++ b/Engine/Plugins/Developer/PropertyAccessNode/Source/PropertyAccessNode/Private/K2Node_PropertyAccess.h
@@ -118,6 +118,7 @@ private:
 	virtual bool DrawNodeAsVariable() const override { return true; }
 	virtual FText GetTooltipText() const override;
 	virtual void HandleVariableRenamed(UBlueprint* InBlueprint, UClass* InVariableClass, UEdGraph* InGraph, const FName& InOldVarName, const FName& InNewVarName) override;
+	virtual void HandleFunctionRenamed(UBlueprint* InBlueprint, UClass* InFunctionClass, UEdGraph* InGraph, const FName& InOldFuncName, const FName& InNewFuncName) override;
 	virtual void ReplaceReferences(UBlueprint* InBlueprint, UBlueprint* InReplacementBlueprint, const FMemberReference& InSource, const FMemberReference& InReplacement) override;
 	virtual bool ReferencesVariable(const FName& InVarName, const UStruct* InScope) const override;
 	virtual bool ReferencesFunction(const FName& InFunctionName,const UStruct* InScope) const override;
diff --git a/Engine/Plugins/Runtime/PropertyAccess/Source/PropertyAccessEditor/Private/PropertyAccessEditor.cpp b/Engine/Plugins/Runtime/PropertyAccess/Source/PropertyAccessEditor/Private/PropertyAccessEditor.cpp
--- a/Engine/Plugins/Runtime/PropertyAccess/Source/PropertyAccessEditor/Private/PropertyAccessEditor.cpp
+++ b/Engine/Plugins/Runtime/PropertyAccess/Source/PropertyAccessEditor/Private/PropertyAccessEditor.cpp
@@ -245,6 +245,13 @@ struct FPropertyAccessEditorSystem
 				}
 
 				InContext.CurrentStruct = Segment.Struct;
+#if WITH_EDITOR
+				if (InContext.CurrentStruct && InContext.CurrentStruct->IsA<UClass>())
+				{
+					// Try to use the skeleton class if possible in editor as the regular class may have yet to be compiled
+					InContext.CurrentStruct = FBlueprintEditorUtils::GetMostUpToDateClass(CastChecked<UClass>(InContext.CurrentStruct));
+				}
+#endif
 			}
 
 			if(InContext.Segments.Num() > 0)
diff --git a/Engine/Source/Editor/AnimGraph/Internal/AnimGraphNodeBinding.h b/Engine/Source/Editor/AnimGraph/Internal/AnimGraphNodeBinding.h
--- a/Engine/Source/Editor/AnimGraph/Internal/AnimGraphNodeBinding.h
+++ b/Engine/Source/Editor/AnimGraph/Internal/AnimGraphNodeBinding.h
@@ -45,6 +45,8 @@ class UAnimGraphNodeBinding : public UObject
 
 	virtual void HandleVariableRenamed(UBlueprint* InBlueprint, UClass* InVariableClass, UEdGraph* InGraph, const FName& InOldVarName, const FName& InNewVarName) {}
 
+	virtual void HandleFunctionRenamed(UBlueprint* InBlueprint, UClass* InFunctionClass, UEdGraph* InGraph, const FName& InOldFuncName, const FName& InNewFuncName) {}
+
 	virtual void ReplaceReferences(UBlueprint* InBlueprint, UBlueprint* InReplacementBlueprint, const FMemberReference& InSource, const FMemberReference& InReplacement) {}
 
 	virtual bool ReferencesVariable(const FName& InVarName, const UStruct* InScope) const { return false; }
diff --git a/Engine/Source/Editor/AnimGraph/Private/AnimGraphNodeBinding_Base.cpp b/Engine/Source/Editor/AnimGraph/Private/AnimGraphNodeBinding_Base.cpp
--- a/Engine/Source/Editor/AnimGraph/Private/AnimGraphNodeBinding_Base.cpp
+++ b/Engine/Source/Editor/AnimGraph/Private/AnimGraphNodeBinding_Base.cpp
@@ -137,7 +137,7 @@ void UAnimGraphNodeBinding_Base::HandleVariableRenamed(UBlueprint* InBlueprint,
 {
 	IPropertyAccessEditor& PropertyAccessEditor = IModularFeatures::Get().GetModularFeature<IPropertyAccessEditor>("PropertyAccessEditor");
 
-	UClass* SkeletonVariableClass = FBlueprintEditorUtils::GetSkeletonClass(InVariableClass);
+	UClass* SkeletonVariableClass = FBlueprintEditorUtils::GetMostUpToDateClass(InVariableClass);
 	UBlueprint* OuterBlueprint = GetTypedOuter<UBlueprint>();
 
 	if(OuterBlueprint)
@@ -168,15 +168,52 @@ void UAnimGraphNodeBinding_Base::HandleVariableRenamed(UBlueprint* InBlueprint,
 	}
 }
 
+void UAnimGraphNodeBinding_Base::HandleFunctionRenamed(UBlueprint* InBlueprint, UClass* InFunctionClass, UEdGraph* InGraph, const FName& InOldFuncName, const FName& InNewFuncName)
+{
+	IPropertyAccessEditor& PropertyAccessEditor = IModularFeatures::Get().GetModularFeature<IPropertyAccessEditor>("PropertyAccessEditor");
+
+	UBlueprint* OuterBlueprint = GetTypedOuter<UBlueprint>();
+
+	if(OuterBlueprint)
+	{
+		UClass* SkeletonFunctionClass = FBlueprintEditorUtils::GetMostUpToDateClass(InFunctionClass);
+
+		// See if any of bindings reference the variable
+		for (auto& BindingPair : PropertyBindings)
+		{
+			TArray<int32> RenameIndices;
+			IPropertyAccessEditor::FResolvePropertyAccessArgs ResolveArgs;
+			ResolveArgs.FunctionFunction = [InOldFuncName, SkeletonFunctionClass, &RenameIndices](int32 InSegmentIndex, UFunction* InFunction, FProperty* InReturnProperty)
+			{
+				const UClass* OwnerSkeletonFunctionClass = FBlueprintEditorUtils::GetMostUpToDateClass(InFunction->GetOuterUClass());
+				if (OwnerSkeletonFunctionClass && InFunction->GetFName() == InOldFuncName && OwnerSkeletonFunctionClass->IsChildOf(SkeletonFunctionClass))
+				{
+					RenameIndices.Add(InSegmentIndex);
+				}
+			};
+
+			PropertyAccessEditor.ResolvePropertyAccess(OuterBlueprint->SkeletonGeneratedClass, BindingPair.Value.PropertyPath, ResolveArgs);
+
+			// Rename any references we found
+			for (const int32& RenameIndex : RenameIndices)
+			{
+				BindingPair.Value.PropertyPath[RenameIndex] = InNewFuncName.ToString();
+				BindingPair.Value.PathAsText = PropertyAccessEditor.MakeTextPath(BindingPair.Value.PropertyPath, OuterBlueprint->SkeletonGeneratedClass);
+			}
+		}
+	}
+}
+
 void UAnimGraphNodeBinding_Base::ReplaceReferences(UBlueprint* InBlueprint, UBlueprint* InReplacementBlueprint, const FMemberReference& InSource, const FMemberReference& InReplacement)
 {
 	IPropertyAccessEditor& PropertyAccessEditor = IModularFeatures::Get().GetModularFeature<IPropertyAccessEditor>("PropertyAccessEditor");
 
-	UClass* SkeletonClass = InBlueprint->SkeletonGeneratedClass;
 	UBlueprint* OuterBlueprint = GetTypedOuter<UBlueprint>();
 
 	if (OuterBlueprint)
 	{
+		UClass* SkeletonClass = InBlueprint->SkeletonGeneratedClass;
+
 		FMemberReference Source = InSource;
 		FProperty* SourceProperty = Source.ResolveMember<FProperty>(InBlueprint);
 		FMemberReference Replacement = InReplacement;
@@ -211,11 +248,12 @@ bool UAnimGraphNodeBinding_Base::ReferencesVariable(const FName& InVarName, cons
 {
 	IPropertyAccessEditor& PropertyAccessEditor = IModularFeatures::Get().GetModularFeature<IPropertyAccessEditor>("PropertyAccessEditor");
 
-	const UClass* SkeletonVariableClass = FBlueprintEditorUtils::GetSkeletonClass(Cast<UClass>(InScope));
 	UBlueprint* OuterBlueprint = GetTypedOuter<UBlueprint>();
 
 	if(OuterBlueprint)
 	{
+		const UClass* SkeletonVariableClass = FBlueprintEditorUtils::GetMostUpToDateClass(Cast<UClass>(InScope));
+
 		// See if any of bindings reference the variable
 		for (const auto& BindingPair : PropertyBindings)
 		{
@@ -226,7 +264,7 @@ bool UAnimGraphNodeBinding_Base::ReferencesVariable(const FName& InVarName, cons
 			{
 				if (SkeletonVariableClass)
 				{
-					const UClass* OwnerSkeletonVariableClass = FBlueprintEditorUtils::GetSkeletonClass(Cast<UClass>(InProperty->GetOwnerStruct()));
+					const UClass* OwnerSkeletonVariableClass = FBlueprintEditorUtils::GetMostUpToDateClass(Cast<UClass>(InProperty->GetOwnerStruct()));
 
 					if (OwnerSkeletonVariableClass && InProperty->GetFName() == InVarName && OwnerSkeletonVariableClass->IsChildOf(SkeletonVariableClass))
 					{
@@ -255,11 +293,12 @@ bool UAnimGraphNodeBinding_Base::ReferencesFunction(const FName& InFunctionName,
 {
 	IPropertyAccessEditor& PropertyAccessEditor = IModularFeatures::Get().GetModularFeature<IPropertyAccessEditor>("PropertyAccessEditor");
 
-	const UClass* SkeletonFunctionClass = FBlueprintEditorUtils::GetSkeletonClass(Cast<UClass>(InScope));
 	UBlueprint* OuterBlueprint = GetTypedOuter<UBlueprint>();
 
 	if(OuterBlueprint)
 	{
+		const UClass* SkeletonFunctionClass = FBlueprintEditorUtils::GetMostUpToDateClass(Cast<UClass>(InScope));
+
 		// See if any of bindings reference the function
 		for (const auto& BindingPair : PropertyBindings)
 		{
@@ -270,7 +309,7 @@ bool UAnimGraphNodeBinding_Base::ReferencesFunction(const FName& InFunctionName,
 			{
 				if (SkeletonFunctionClass)
 				{
-					const UClass* OwnerSkeletonFunctionClass = FBlueprintEditorUtils::GetSkeletonClass(InFunction->GetOuterUClass());
+					const UClass* OwnerSkeletonFunctionClass = FBlueprintEditorUtils::GetMostUpToDateClass(InFunction->GetOuterUClass());
 
 					if (OwnerSkeletonFunctionClass && InFunction->GetFName() == InFunctionName && OwnerSkeletonFunctionClass->IsChildOf(SkeletonFunctionClass))
 					{
diff --git a/Engine/Source/Editor/AnimGraph/Private/AnimGraphNodeBinding_Base.h b/Engine/Source/Editor/AnimGraph/Private/AnimGraphNodeBinding_Base.h
--- a/Engine/Source/Editor/AnimGraph/Private/AnimGraphNodeBinding_Base.h
+++ b/Engine/Source/Editor/AnimGraph/Private/AnimGraphNodeBinding_Base.h
@@ -22,6 +22,7 @@ private:
 	virtual void RemoveBindings(FName InBindingName) override;
 	virtual void AddPinSearchMetaDataInfo(const UEdGraphPin* InPin, FName InBindingName, TArray<struct FSearchTagDataPair>& OutTaggedMetaData) const override;
 	virtual void HandleVariableRenamed(UBlueprint* InBlueprint, UClass* InVariableClass, UEdGraph* InGraph, const FName& InOldVarName, const FName& InNewVarName) override;
+	virtual void HandleFunctionRenamed(UBlueprint* InBlueprint, UClass* InFunctionClass, UEdGraph* InGraph, const FName& InOldFuncName, const FName& InNewFuncName) override;
 	virtual void ReplaceReferences(UBlueprint* InBlueprint, UBlueprint* InReplacementBlueprint, const FMemberReference& InSource, const FMemberReference& InReplacement) override;
 	virtual bool ReferencesVariable(const FName& InVarName, const UStruct* InScope) const override;
 	virtual bool ReferencesFunction(const FName& InFunctionName, const UStruct* InScope) const override;
diff --git a/Engine/Source/Editor/AnimGraph/Private/AnimGraphNode_Base.cpp b/Engine/Source/Editor/AnimGraph/Private/AnimGraphNode_Base.cpp
--- a/Engine/Source/Editor/AnimGraph/Private/AnimGraphNode_Base.cpp
+++ b/Engine/Source/Editor/AnimGraph/Private/AnimGraphNode_Base.cpp
@@ -1062,6 +1062,14 @@ void UAnimGraphNode_Base::HandleVariableRenamed(UBlueprint* InBlueprint, UClass*
 	}
 }
 
+void UAnimGraphNode_Base::HandleFunctionRenamed(UBlueprint* InBlueprint, UClass* InFunctionClass, UEdGraph* InGraph, const FName& InOldFuncName, const FName& InNewFuncName)
+{
+	if (Binding)
+	{
+		Binding->HandleFunctionRenamed(InBlueprint, InFunctionClass, InGraph, InOldFuncName, InNewFuncName);
+	}
+}
+
 void UAnimGraphNode_Base::ReplaceReferences(UBlueprint* InBlueprint, UBlueprint* InReplacementBlueprint, const FMemberReference& InSource, const FMemberReference& InReplacement)
 {
 	if (Binding)
diff --git a/Engine/Source/Editor/AnimGraph/Public/AnimGraphNode_Base.h b/Engine/Source/Editor/AnimGraph/Public/AnimGraphNode_Base.h
--- a/Engine/Source/Editor/AnimGraph/Public/AnimGraphNode_Base.h
+++ b/Engine/Source/Editor/AnimGraph/Public/AnimGraphNode_Base.h
@@ -276,6 +276,7 @@ public:
 	virtual bool CanJumpToDefinition() const override;
 	virtual void JumpToDefinition() const override;
 	virtual void HandleVariableRenamed(UBlueprint* InBlueprint, UClass* InVariableClass, UEdGraph* InGraph, const FName& InOldVarName, const FName& InNewVarName) override;
+	virtual void HandleFunctionRenamed(UBlueprint* InBlueprint, UClass* InFunctionClass, UEdGraph* InGraph, const FName& InOldFuncName, const FName& InNewFuncName) override;
 	virtual void ReplaceReferences(UBlueprint* InBlueprint, UBlueprint* InReplacementBlueprint, const FMemberReference& InSource, const FMemberReference& InReplacement) override;
 	virtual bool ReferencesVariable(const FName& InVarName, const UStruct* InScope) const override;
 	virtual bool ReferencesFunction(const FName& InFunctionName, const UStruct* InScope) const override;
diff --git a/Engine/Source/Editor/BlueprintGraph/Classes/K2Node.h b/Engine/Source/Editor/BlueprintGraph/Classes/K2Node.h
--- a/Engine/Source/Editor/BlueprintGraph/Classes/K2Node.h
+++ b/Engine/Source/Editor/BlueprintGraph/Classes/K2Node.h
@@ -430,6 +430,9 @@ public:
 
 	/** Return whether this node references the specified variable, give the supplied scope. Used when variable types are changed. */
 	virtual bool ReferencesVariable(const FName& InVarName, const UStruct* InScope) const { return false; }
+
+	/** Handle when a function is renamed in the Blueprint Palette */
+	virtual void HandleFunctionRenamed(UBlueprint* InBlueprint, UClass* InFunctionClass, UEdGraph* InGraph, const FName& InOldFuncName, const FName& InNewFuncName) {}
 	
 	/** Return whether this node references the specified function, identified by a name and guid pair **/
 	virtual bool ReferencesFunction(const FName& InFunctionName, const UStruct* InScope) const { return false; };
diff --git a/Engine/Source/Editor/UnrealEd/Private/Kismet2/BlueprintEditorUtils.cpp b/Engine/Source/Editor/UnrealEd/Private/Kismet2/BlueprintEditorUtils.cpp
--- a/Engine/Source/Editor/UnrealEd/Private/Kismet2/BlueprintEditorUtils.cpp
+++ b/Engine/Source/Editor/UnrealEd/Private/Kismet2/BlueprintEditorUtils.cpp
@@ -318,6 +318,34 @@ static bool RenameVariableReferencesInGraph(UBlueprint* InBlueprint, UClass* InV
 	return bFoundReference;
 }
 
+
+/**
+ * Looks through the specified graph for any references to the specified 
+ * function, and renames them accordingly.
+ * 
+ * @param  InBlueprint		The blueprint that you want to search through.
+ * @param  InFunctionClass	The class that owns the function that we're renaming
+ * @param  InGraph			Graph to scope the rename to
+ * @param  InOldFuncName	The current name of the function we want to replace
+ * @param  InNewFuncName	The name that we wish to change all references to
+ */
+static bool RenameFunctionReferencesInGraph(UBlueprint* InBlueprint, UClass* InFunctionClass, UEdGraph* InGraph, const FName& InOldFuncName, const FName& InNewFuncName)
+{
+	bool bFoundReference = false;
+
+	for(UEdGraphNode* GraphNode : InGraph->Nodes)
+	{
+		// Allow node to handle function renaming
+		if (UK2Node* const K2Node = Cast<UK2Node>(GraphNode))
+		{
+			bFoundReference |= K2Node->ReferencesFunction(InOldFuncName, nullptr);
+			K2Node->HandleFunctionRenamed(InBlueprint, InFunctionClass, InGraph, InOldFuncName, InNewFuncName);
+		}
+	}
+
+	return bFoundReference;
+}
+
 /**
  * Gathers all variable nodes from all graph's subgraph nodes
  *
@@ -356,6 +384,25 @@ void FBlueprintEditorUtils::RenameVariableReferences(UBlueprint* Blueprint, UCla
 	OnRenameVariableReferencesEvent.Broadcast(Blueprint, VariableClass, OldVarName, NewVarName);
 }
 
+FBlueprintEditorUtils::FOnRenameFunctionReferences FBlueprintEditorUtils::OnRenameFunctionReferencesEvent;
+
+void FBlueprintEditorUtils::RenameFunctionReferences(UBlueprint* Blueprint, UClass* FunctionClass, const FName& OldFuncName, const FName& NewFuncName)
+{
+	TArray<UEdGraph*> AllGraphs;
+	Blueprint->GetAllGraphs(AllGraphs);
+
+	// Update any graph nodes that reference the old function name to instead reference the new name
+	for(UEdGraph* CurrentGraph : AllGraphs)
+	{
+		if (RenameFunctionReferencesInGraph(Blueprint, FunctionClass, CurrentGraph, OldFuncName, NewFuncName))
+		{
+			MarkBlueprintAsModified(Blueprint);
+		}
+	}
+
+	OnRenameFunctionReferencesEvent.Broadcast(Blueprint, FunctionClass, OldFuncName, NewFuncName);
+}
+
 //////////////////////////////////////
 // FBasePinChangeHelper
 
@@ -2652,6 +2699,9 @@ void FBlueprintEditorUtils::RenameGraph(UEdGraph* Graph, const FString& NewNameS
 			}
 		}
 
+		// Replace any other nodes that reference this function
+		ReplaceFunctionReferences(Blueprint, OldGraphName, NewGraphName);
+
 		// We should let the blueprint know we renamed a graph, some stuff may need to be fixed up.
 		Blueprint->NotifyGraphRenamed(Graph, OldGraphName, NewGraphName);
 
@@ -5706,6 +5756,21 @@ void FBlueprintEditorUtils::ReplaceVariableReferences(UBlueprint* Blueprint, con
 	ReplaceVariableReferences(Blueprint, OldVariable->GetFName(), NewVariable->GetFName());
 }
 
+void FBlueprintEditorUtils::ReplaceFunctionReferences(UBlueprint* Blueprint, const FName OldName, const FName NewName)
+{
+	check((OldName != NAME_None) && (NewName != NAME_None));
+
+	FBlueprintEditorUtils::RenameFunctionReferences(Blueprint, Blueprint->GeneratedClass, OldName, NewName);
+
+	TArray<UBlueprint*> Dependents;
+	FindDependentBlueprints(Blueprint, Dependents);
+
+	for (UBlueprint* DependentBp : Dependents)
+	{
+		FBlueprintEditorUtils::RenameFunctionReferences(DependentBp, Blueprint->GeneratedClass, OldName, NewName);
+	}
+}
+
 bool FBlueprintEditorUtils::IsVariableComponent(const FBPVariableDescription& Variable)
 {
 	// Find the variable in the list
diff --git a/Engine/Source/Editor/UnrealEd/Public/Kismet2/BlueprintEditorUtils.h b/Engine/Source/Editor/UnrealEd/Public/Kismet2/BlueprintEditorUtils.h
--- a/Engine/Source/Editor/UnrealEd/Public/Kismet2/BlueprintEditorUtils.h
+++ b/Engine/Source/Editor/UnrealEd/Public/Kismet2/BlueprintEditorUtils.h
@@ -1055,6 +1055,9 @@ public:
 	/** Replaces all variable references in the specified blueprint */
 	static UNREALED_API void ReplaceVariableReferences(UBlueprint* Blueprint, const FProperty* OldVariable, const FProperty* NewVariable);
 
+	/** Replaces all function references in the specified blueprint */
+	static UNREALED_API void ReplaceFunctionReferences(UBlueprint* Blueprint, const FName OldName, const FName NewName);
+
 	/** Check blueprint variable metadata keys/values for validity and make adjustments if needed */
 	static UNREALED_API void FixupVariableDescription(UBlueprint* Blueprint, FBPVariableDescription& VarDesc);
 
@@ -1886,6 +1889,21 @@ public:
 	 */
 	static UNREALED_API void RenameVariableReferences(UBlueprint* Blueprint, UClass* VariableClass, const FName& OldVarName, const FName& NewVarName);
 
+	/** Event fired after RenameFunctionReferences is called */
+	DECLARE_MULTICAST_DELEGATE_FourParams(FOnRenameFunctionReferences, UBlueprint* /*Blueprint*/, UClass* /*FunctionClass*/, const FName& /*OldFuncName*/, const FName& /*NewFuncName*/);
+	static UNREALED_API FOnRenameFunctionReferences OnRenameFunctionReferencesEvent;
+
+	/**
+	 * Looks through the specified blueprint for any references to the specified 
+	 * function, and renames them accordingly.
+	 * 
+	 * @param  Blueprint		The blueprint that you want to search through.
+	 * @param  FunctionClass	The class that owns the function that we're renaming
+	 * @param  OldFuncName		The current name of the function we want to replace
+	 * @param  NewFuncName		The name that we wish to change all references to
+	 */
+	static UNREALED_API void RenameFunctionReferences(UBlueprint* Blueprint, UClass* FunctionClass, const FName& OldFuncName, const FName& NewFuncName);
+
 public:
 	static UNREALED_API FName GetFunctionNameFromClassByGuid(const UClass* InClass, const FGuid FunctionGuid);
 	static UNREALED_API bool GetFunctionGuidFromClassByFieldName(const UClass* InClass, const FName FunctionName, FGuid& FunctionGuid);
