AnimNodeEditMode: avoid opening unnecessary transactions when orbiting the camera
    - and more generally when nothing is transformed in the viewport

NOTE: this is not a fix for this jira, but the problem has been highlighted by it.

#rnx

diff --git a/Engine/Source/Editor/AnimGraph/Private/EditModes/AnimNodeEditMode.cpp b/Engine/Source/Editor/AnimGraph/Private/EditModes/AnimNodeEditMode.cpp
--- a/Engine/Source/Editor/AnimGraph/Private/EditModes/AnimNodeEditMode.cpp
+++ b/Engine/Source/Editor/AnimGraph/Private/EditModes/AnimNodeEditMode.cpp
@@ -291,25 +291,33 @@ FVector FAnimNodeEditMode::GetWidgetLocation() const
 
 bool FAnimNodeEditMode::StartTracking(FEditorViewportClient* InViewportClient, FViewport* InViewport)
 {
-	if (!bInTransaction)
-	{
-		GEditor->BeginTransaction(LOCTEXT("EditSkelControlNodeTransaction", "Edit Skeletal Control Node"));
+	const EAxisList::Type CurrentAxis = InViewportClient ? InViewportClient->GetCurrentWidgetAxis() : EAxisList::None;
+	const UE::Widget::EWidgetMode WidgetMode = InViewportClient ? InViewportClient->GetWidgetMode() : UE::Widget::WM_None;
 
-		for (EditorRuntimeNodePair CurrentNodePair : SelectedAnimNodes)
+	if (WidgetMode != UE::Widget::WM_None && CurrentAxis != EAxisList::None)
+	{
+		if (!bInTransaction)
 		{
-			if (CurrentNodePair.EditorAnimNode != nullptr)
+			GEditor->BeginTransaction(LOCTEXT("EditSkelControlNodeTransaction", "Edit Skeletal Control Node"));
+
+			for (EditorRuntimeNodePair CurrentNodePair : SelectedAnimNodes)
 			{
-				CurrentNodePair.EditorAnimNode->SetFlags(RF_Transactional);
-				CurrentNodePair.EditorAnimNode->Modify();
+				if (CurrentNodePair.EditorAnimNode != nullptr)
+				{
+					CurrentNodePair.EditorAnimNode->SetFlags(RF_Transactional);
+					CurrentNodePair.EditorAnimNode->Modify();
+				}
 			}
-		}
 
-		bInTransaction = true;
-	}
+			bInTransaction = true;
+		}
 
-	bManipulating = true;
+		bManipulating = true;
 
-	return true;
+		return true;
+	}
+	
+	return IAnimNodeEditMode::StartTracking(InViewportClient, InViewport);
 }
 
 bool FAnimNodeEditMode::EndTracking(FEditorViewportClient* InViewportClient, FViewport* InViewport)
@@ -317,15 +325,17 @@ bool FAnimNodeEditMode::EndTracking(FEditorViewportClient* InViewportClient, FVi
 	if (bManipulating)
 	{
 		bManipulating = false;
-	}
 
-	if (bInTransaction)
-	{
-		GEditor->EndTransaction();
-		bInTransaction = false;
+		if (bInTransaction)
+		{
+			GEditor->EndTransaction();
+			bInTransaction = false;
+		}
+
+		return true;
 	}
 
-	return true;
+	return IAnimNodeEditMode::EndTracking(InViewportClient, InViewport);
 }
 
 bool FAnimNodeEditMode::InputKey(FEditorViewportClient* InViewportClient, FViewport* InViewport, FKey InKey, EInputEvent InEvent)
