RandomPlayer anim graph node now calls PreloadRequiredAssets.

diff --git a/Engine/Source/Editor/AnimGraph/Private/AnimGraphNode_RandomPlayer.cpp b/Engine/Source/Editor/AnimGraph/Private/AnimGraphNode_RandomPlayer.cpp
--- a/Engine/Source/Editor/AnimGraph/Private/AnimGraphNode_RandomPlayer.cpp
+++ b/Engine/Source/Editor/AnimGraph/Private/AnimGraphNode_RandomPlayer.cpp
@@ -39,4 +39,14 @@ void UAnimGraphNode_RandomPlayer::GetOutputLinkAttributes(FNodeAttributeArray& O
 	}
 }
 
+void UAnimGraphNode_RandomPlayer::PreloadRequiredAssets()
+{
+	for (const FRandomPlayerSequenceEntry& Entry : Node.Entries)
+	{
+		PreloadObject(Entry.Sequence);
+	}
+
+	Super::PreloadRequiredAssets();
+}
+
 #undef LOCTEXT_NAMESPACE
diff --git a/Engine/Source/Editor/AnimGraph/Public/AnimGraphNode_RandomPlayer.h b/Engine/Source/Editor/AnimGraph/Public/AnimGraphNode_RandomPlayer.h
--- a/Engine/Source/Editor/AnimGraph/Public/AnimGraphNode_RandomPlayer.h
+++ b/Engine/Source/Editor/AnimGraph/Public/AnimGraphNode_RandomPlayer.h
@@ -27,5 +27,6 @@ public:
 
 	// UAnimGraphNode_Base interface
 	virtual void GetOutputLinkAttributes(FNodeAttributeArray& OutAttributes) const override;
+	virtual void PreloadRequiredAssets() override;
 	// End of UAnimGraphNode_Base interface
 };
