Fix for crash when validating auto rule based transitions. Also reinstated some warnings as info

diff --git a/Engine/Source/Editor/AnimGraph/Private/AnimStateTransitionNode.cpp b/Engine/Source/Editor/AnimGraph/Private/AnimStateTransitionNode.cpp
--- a/Engine/Source/Editor/AnimGraph/Private/AnimStateTransitionNode.cpp
+++ b/Engine/Source/Editor/AnimGraph/Private/AnimStateTransitionNode.cpp
@@ -16,6 +16,7 @@
 #include "AnimGraphNode_SequencePlayer.h"
 #include "AnimGraphNode_StateResult.h"
 #include "AnimGraphNode_TransitionResult.h"
+#include "AnimStateAliasNode.h"
 #include "AnimStateConduitNode.h"
 #include "Kismet2/CompilerResultsLog.h"
 #include "EdGraphUtilities.h"
@@ -696,28 +697,43 @@ void UAnimStateTransitionNode::ValidateNodeDuringCompilation(class FCompilerResu
 			// Check for automatic transition rules that are being triggered from looping asset players, as these can often be symptomatic of logic errors
 			if (UAnimStateNodeBase* PreviousState = GetPreviousState())
 			{
-				if (UEdGraph* PreviousStateGraph = PreviousState->GetBoundGraph())
+				// Deal with alias state nodes. Only single source aliases are valid
+				if (UAnimStateAliasNode* AliasState = Cast<UAnimStateAliasNode>(PreviousState))
 				{
-					if (UAnimGraphNode_StateResult* PreviousStateGraphResultNode = Cast<UAnimationStateGraph>(PreviousStateGraph)->GetResultNode())
+					PreviousState = AliasState->GetAliasedState();
+					if (!PreviousState)
 					{
-						for (UEdGraphPin* TestPin : PreviousStateGraphResultNode->Pins)
+						MessageLog.Note(TEXT("Transition @@ is using an automatic transition rule but its source is an Alias node which aliases more than one State"), this);
+						return;
+					}
+				}
+				
+				UAnimationStateGraph* PreviousStateGraph = Cast<UAnimationStateGraph>(PreviousState->GetBoundGraph());
+				if (!PreviousStateGraph)
+				{
+					MessageLog.Note(TEXT("Transition @@ is using an automatic transition rule but source @@ is not a valid State. Check if the transit is connected to a Conduit"), this, PreviousState);
+					return;
+				}
+				
+				if (UAnimGraphNode_StateResult* PreviousStateGraphResultNode = PreviousStateGraph->GetResultNode())
+				{
+					for (UEdGraphPin* TestPin : PreviousStateGraphResultNode->Pins)
+					{
+						// Warn for the trivial but common case of a looping asset player connected directly to the result node
+						if ((TestPin->Direction == EGPD_Input) && (TestPin->LinkedTo.Num() == 1))
 						{
-							// Warn for the trivial but common case of a looping asset player connected directly to the result node
-							if ((TestPin->Direction == EGPD_Input) && (TestPin->LinkedTo.Num() == 1))
+							if (UAnimGraphNode_SequencePlayer* SequencePlayer = Cast<UAnimGraphNode_SequencePlayer>(TestPin->LinkedTo[0]->GetOwningNode()))
 							{
-								if (UAnimGraphNode_SequencePlayer* SequencePlayer = Cast<UAnimGraphNode_SequencePlayer>(TestPin->LinkedTo[0]->GetOwningNode()))
+								if (SequencePlayer->Node.IsLooping())
 								{
-									if (SequencePlayer->Node.IsLooping())
-									{
-										// MessageLog.Warning(TEXT("Transition @@ is using an automatic transition rule but the source @@ is set as looping.  Please clear the 'Loop Animation' flag"), this, SequencePlayer);
-									}
+									MessageLog.Note(TEXT("Transition @@ is using an automatic transition rule but the source @@ is set as looping.  Please clear the 'Loop Animation' flag"), this, SequencePlayer);
 								}
-								else if (UAnimGraphNode_BlendSpacePlayer* BlendSpacePlayer = Cast<UAnimGraphNode_BlendSpacePlayer>(TestPin->LinkedTo[0]->GetOwningNode()))
+							}
+							else if (UAnimGraphNode_BlendSpacePlayer* BlendSpacePlayer = Cast<UAnimGraphNode_BlendSpacePlayer>(TestPin->LinkedTo[0]->GetOwningNode()))
+							{
+								if (BlendSpacePlayer->Node.IsLooping())
 								{
-									if (BlendSpacePlayer->Node.IsLooping())
-									{
-										// MessageLog.Warning(TEXT("Transition @@ is using an automatic transition rule but the source @@ is set as looping.  Please clear the 'Loop' flag"), this, BlendSpacePlayer);
-									}
+									MessageLog.Note(TEXT("Transition @@ is using an automatic transition rule but the source @@ is set as looping.  Please clear the 'Loop' flag"), this, BlendSpacePlayer);
 								}
 							}
 						}
