Revert changes to RestoreStateMachineState and RestoreStateMachineNode which led to nullptr SourceNodes, add compile time error to check for these broken nodes and instruct user to replace them

#changelist validated
#virtualized

diff --git a/Engine/Source/Editor/AnimGraph/Private/K2Node_AnimGetter.cpp b/Engine/Source/Editor/AnimGraph/Private/K2Node_AnimGetter.cpp
--- a/Engine/Source/Editor/AnimGraph/Private/K2Node_AnimGetter.cpp
+++ b/Engine/Source/Editor/AnimGraph/Private/K2Node_AnimGetter.cpp
@@ -239,6 +239,14 @@ void UK2Node_AnimGetter::GetMenuActions(FBlueprintActionDatabaseRegistrar& Actio
 	}
 }
 
+void UK2Node_AnimGetter::ValidateNodeDuringCompilation(class FCompilerResultsLog& MessageLog) const
+{
+	if (SourceNode == nullptr)
+	{
+		MessageLog.Error(TEXT("@@ contains invalid data. Please delete and recreate the node."), this);
+	}
+}
+
 bool UK2Node_AnimGetter::IsActionFilteredOut(FBlueprintActionFilter const& Filter)
 {
 	if(Filter.Context.Graphs.Num() > 0)
@@ -347,46 +355,31 @@ bool UK2Node_AnimGetter::IsActionFilteredOut(FBlueprintActionFilter const& Filte
 
 void UK2Node_AnimGetter::RestoreStateMachineState()
 {
-	UAnimationTransitionGraph* TransitionGraph = Cast<UAnimationTransitionGraph>(GetOuter());
-	if (TransitionGraph == nullptr)
-	{
-		SourceStateNode = nullptr;
-		return;
-	}
-
-	UAnimStateTransitionNode* StateTransitionNode = Cast<UAnimStateTransitionNode>(TransitionGraph->GetOuter());
-	if (StateTransitionNode == nullptr)
+	if (SourceStateNode)
 	{
-		SourceStateNode = nullptr;
-		return;
+		if (UAnimationTransitionGraph* TransitionGraph = Cast<UAnimationTransitionGraph>(GetOuter()))
+		{
+			if (UAnimStateTransitionNode* StateTransitionNode = Cast<UAnimStateTransitionNode>(TransitionGraph->GetOuter()))
+			{
+				SourceStateNode = StateTransitionNode->GetPreviousState();
+			}
+		}
 	}
-
-	SourceStateNode = StateTransitionNode->GetPreviousState();
 }
 
 void UK2Node_AnimGetter::RestoreStateMachineNode()
 {
-	if(SourceStateNode == nullptr)
-	{
-		SourceNode = nullptr;
-		return;
-	}
-
-	UAnimationStateMachineGraph* Graph = Cast<UAnimationStateMachineGraph>(SourceStateNode->GetOuter());
-	if (Graph == nullptr)
+	if (SourceStateNode)
 	{
-		SourceNode = nullptr;
-		return;
-	}
-
-	UAnimGraphNode_StateMachine* MachineNode = Cast<UAnimGraphNode_StateMachine>(Graph->GetOuter());
-	if (MachineNode == nullptr)
-	{
-		SourceNode = nullptr;
-		return;
+		UAnimationStateMachineGraph* Graph = Cast<UAnimationStateMachineGraph>(SourceStateNode->GetOuter());
+		if (Graph)
+		{
+			if (UAnimGraphNode_StateMachine* MachineNode = Cast<UAnimGraphNode_StateMachine>(Graph->GetOuter()))
+			{
+				SourceNode = MachineNode;
+			}
+		}
 	}
-
-	SourceNode = MachineNode;
 }
 
 bool UK2Node_AnimGetter::GetterRequiresParameter(const UFunction* Getter, FString ParamName)
@@ -464,6 +457,10 @@ FText UK2Node_AnimGetter::GenerateTitle(UFunction* Getter, UAnimStateNodeBase* S
 {
 	if (GetterRequiresParameter(Getter, TEXT("AssetPlayerIndex")))
 	{
+		if (!SourceNode)
+		{
+			return FText::Format(LOCTEXT("NodeTitle", "{0} (Invalid node)"), Getter->GetDisplayNameText());
+		}
 		// Should always succeed
 		if (UAnimationAsset* NodeAsset = SourceNode->GetAnimationAsset())
 		{
@@ -478,6 +475,10 @@ FText UK2Node_AnimGetter::GenerateTitle(UFunction* Getter, UAnimStateNodeBase* S
 		}
 		else
 		{
+			if (!SourceNode)
+			{
+				return FText::Format(LOCTEXT("NodeTitle", "{0} (Invalid node)"), Getter->GetDisplayNameText());
+			}
 			// Only requires the state machine
 			return FText::Format(LOCTEXT("NodeTitle", "{0} ({1})"), Getter->GetDisplayNameText(), SourceNode->GetNodeTitle(ENodeTitleType::ListView));
 		}
diff --git a/Engine/Source/Editor/AnimGraph/Public/K2Node_AnimGetter.h b/Engine/Source/Editor/AnimGraph/Public/K2Node_AnimGetter.h
--- a/Engine/Source/Editor/AnimGraph/Public/K2Node_AnimGetter.h
+++ b/Engine/Source/Editor/AnimGraph/Public/K2Node_AnimGetter.h
@@ -82,6 +82,7 @@ public:
 
 	// UK2Node interface
 	virtual void GetMenuActions(FBlueprintActionDatabaseRegistrar& ActionRegistrar) const override;
+	virtual void ValidateNodeDuringCompilation(class FCompilerResultsLog& MessageLog) const override;
 	// end of UK2Node interface
 
 	// The node that is required for the getter
