Fix LOCTEXT key conflict in K2Node_AnimGetter

diff --git a/Engine/Source/Editor/AnimGraph/Private/K2Node_AnimGetter.cpp b/Engine/Source/Editor/AnimGraph/Private/K2Node_AnimGetter.cpp
--- a/Engine/Source/Editor/AnimGraph/Private/K2Node_AnimGetter.cpp
+++ b/Engine/Source/Editor/AnimGraph/Private/K2Node_AnimGetter.cpp
@@ -455,11 +455,13 @@ void UK2Node_AnimGetter::UpdateCachedTitle(FNodeSpawnData& SpawnData)
 
 FText UK2Node_AnimGetter::GenerateTitle(UFunction* Getter, UAnimStateNodeBase* SourceStateNode, UAnimGraphNode_Base* SourceNode)
 {
+	static const FText InvalidNodeTitle = LOCTEXT("NodeTitleInvalid", "{0} (Invalid node)");
+
 	if (GetterRequiresParameter(Getter, TEXT("AssetPlayerIndex")))
 	{
 		if (!SourceNode)
 		{
-			return FText::Format(LOCTEXT("NodeTitle", "{0} (Invalid node)"), Getter->GetDisplayNameText());
+			return FText::Format(InvalidNodeTitle, Getter->GetDisplayNameText());
 		}
 		// Should always succeed
 		if (UAnimationAsset* NodeAsset = SourceNode->GetAnimationAsset())
@@ -477,7 +479,7 @@ FText UK2Node_AnimGetter::GenerateTitle(UFunction* Getter, UAnimStateNodeBase* S
 		{
 			if (!SourceNode)
 			{
-				return FText::Format(LOCTEXT("NodeTitle", "{0} (Invalid node)"), Getter->GetDisplayNameText());
+				return FText::Format(InvalidNodeTitle, Getter->GetDisplayNameText());
 			}
 			// Only requires the state machine
 			return FText::Format(LOCTEXT("NodeTitle", "{0} ({1})"), Getter->GetDisplayNameText(), SourceNode->GetNodeTitle(ENodeTitleType::ListView));
