Fix for anim getters that require no parameters returning an error for a nullptr node on compilation
* Specifically the WasAnimNotifyTriggeredInAnyState and WasAnimNotifyNameTriggeredInAnyState

diff --git a/Engine/Source/Editor/AnimGraph/Private/K2Node_AnimGetter.cpp b/Engine/Source/Editor/AnimGraph/Private/K2Node_AnimGetter.cpp
--- a/Engine/Source/Editor/AnimGraph/Private/K2Node_AnimGetter.cpp
+++ b/Engine/Source/Editor/AnimGraph/Private/K2Node_AnimGetter.cpp
@@ -241,7 +241,14 @@ void UK2Node_AnimGetter::GetMenuActions(FBlueprintActionDatabaseRegistrar& Actio
 
 void UK2Node_AnimGetter::ValidateNodeDuringCompilation(class FCompilerResultsLog& MessageLog) const
 {
-	if (SourceNode == nullptr)
+	const bool bAssetPlayerIndexRequired = GetterRequiresParameter(GetTargetFunction(), TEXT("AssetPlayerIndex"));
+	const bool bMachineIndexRequired = GetterRequiresParameter(GetTargetFunction(), TEXT("MachineIndex"));
+	const bool bTransitionIndexRequired = GetterRequiresParameter(GetTargetFunction(), TEXT("TransitionIndex"));
+	const bool bStateIndexRequired = GetterRequiresParameter(GetTargetFunction(), TEXT("StateIndex"));
+
+	const bool bSourceNodeRequired = bAssetPlayerIndexRequired || bMachineIndexRequired || bTransitionIndexRequired || bStateIndexRequired;
+
+	if (bSourceNodeRequired && SourceNode == nullptr)
 	{
 		MessageLog.Error(TEXT("@@ contains invalid data. Please delete and recreate the node."), this);
 	}
diff --git a/Engine/Source/Editor/AnimGraph/Public/K2Node_AnimGetter.h b/Engine/Source/Editor/AnimGraph/Public/K2Node_AnimGetter.h
--- a/Engine/Source/Editor/AnimGraph/Public/K2Node_AnimGetter.h
+++ b/Engine/Source/Editor/AnimGraph/Public/K2Node_AnimGetter.h
@@ -37,7 +37,7 @@ struct FNodeSpawnData
 	UPROPERTY()
 	FText CachedTitle;
 
-	// The node the spawned getter accesses
+	// The node the spawned getter accesses, if any
 	UPROPERTY()
 	TObjectPtr<UAnimGraphNode_Base> SourceNode;
 
