New TRS Gizmo: implement BeginTransform & EndTransform for PhysicsAssetEditor and AnimNode edit modes

#rnx

diff --git a/Engine/Source/Editor/AnimGraph/Private/EditModes/AnimNodeEditMode.cpp b/Engine/Source/Editor/AnimGraph/Private/EditModes/AnimNodeEditMode.cpp
--- a/Engine/Source/Editor/AnimGraph/Private/EditModes/AnimNodeEditMode.cpp
+++ b/Engine/Source/Editor/AnimGraph/Private/EditModes/AnimNodeEditMode.cpp
@@ -296,31 +296,56 @@ bool FAnimNodeEditMode::StartTracking(FEditorViewportClient* InViewportClient, F
 
 	if (WidgetMode != UE::Widget::WM_None && CurrentAxis != EAxisList::None)
 	{
-		if (!bInTransaction)
-		{
-			GEditor->BeginTransaction(LOCTEXT("EditSkelControlNodeTransaction", "Edit Skeletal Control Node"));
+		return HandleBeginTransform();
+	}
+	
+	return IAnimNodeEditMode::StartTracking(InViewportClient, InViewport);
+}
+
+bool FAnimNodeEditMode::EndTracking(FEditorViewportClient* InViewportClient, FViewport* InViewport)
+{
+	if (HandleEndTransform())
+	{
+		return true;
+	}
+
+	return IAnimNodeEditMode::EndTracking(InViewportClient, InViewport);
+}
+
+bool FAnimNodeEditMode::BeginTransform(const FGizmoState& InState)
+{
+	return HandleBeginTransform();
+}
+
+bool FAnimNodeEditMode::EndTransform(const FGizmoState& InState)
+{
+	return HandleEndTransform();
+}
 
-			for (EditorRuntimeNodePair CurrentNodePair : SelectedAnimNodes)
+bool FAnimNodeEditMode::HandleBeginTransform()
+{
+	if (!bInTransaction)
+	{
+		GEditor->BeginTransaction(LOCTEXT("EditSkelControlNodeTransaction", "Edit Skeletal Control Node"));
+
+		for (EditorRuntimeNodePair CurrentNodePair : SelectedAnimNodes)
+		{
+			if (CurrentNodePair.EditorAnimNode != nullptr)
 			{
-				if (CurrentNodePair.EditorAnimNode != nullptr)
-				{
-					CurrentNodePair.EditorAnimNode->SetFlags(RF_Transactional);
-					CurrentNodePair.EditorAnimNode->Modify();
-				}
+				CurrentNodePair.EditorAnimNode->SetFlags(RF_Transactional);
+				CurrentNodePair.EditorAnimNode->Modify();
 			}
-
-			bInTransaction = true;
 		}
 
-		bManipulating = true;
-
-		return true;
+		bInTransaction = true;
 	}
-	
-	return IAnimNodeEditMode::StartTracking(InViewportClient, InViewport);
+
+	bManipulating = true;
+
+	return true;
 }
 
-bool FAnimNodeEditMode::EndTracking(FEditorViewportClient* InViewportClient, FViewport* InViewport)
+bool FAnimNodeEditMode::HandleEndTransform()
 {
 	if (bManipulating)
 	{
@@ -334,8 +359,7 @@ bool FAnimNodeEditMode::EndTracking(FEditorViewportClient* InViewportClient, FVi
 
 		return true;
 	}
-
-	return IAnimNodeEditMode::EndTracking(InViewportClient, InViewport);
+	return false;
 }
 
 bool FAnimNodeEditMode::InputKey(FEditorViewportClient* InViewportClient, FViewport* InViewport, FKey InKey, EInputEvent InEvent)
diff --git a/Engine/Source/Editor/AnimGraph/Public/AnimNodeEditMode.h b/Engine/Source/Editor/AnimGraph/Public/AnimNodeEditMode.h
--- a/Engine/Source/Editor/AnimGraph/Public/AnimNodeEditMode.h
+++ b/Engine/Source/Editor/AnimGraph/Public/AnimNodeEditMode.h
@@ -89,6 +89,11 @@ public:
 		FAnimNode_Base* RuntimeAnimNode;
 	};
 
+	// Start IGizmoEdModeInterface overrides
+	virtual bool BeginTransform(const FGizmoState& InState) override;
+	virtual bool EndTransform(const FGizmoState& InState) override;
+	// End IGizmoEdModeInterface overrides
+	
 protected:
 	// local conversion functions for drawing
 	static void ConvertToComponentSpaceTransform(const USkeletalMeshComponent* SkelComp, const FTransform & InTransform, FTransform & OutCSTransform, int32 BoneIndex, EBoneControlSpace Space);
@@ -107,6 +112,10 @@ protected:
 
 	const bool IsManipulatingWidget() const { return bManipulating; }
 
+	// Manage the start and end of a transform action in the viewport.
+	bool HandleBeginTransform();
+	bool HandleEndTransform();
+	
 	TArray< EditorRuntimeNodePair > SelectedAnimNodes;	// Selected Anim Graph Nodes
 	TArray< EditorRuntimeNodePair > PoseWatchedAnimNodes; 	// Pose Watched Anim Graph Nodes. 
 
diff --git a/Engine/Source/Editor/PhysicsAssetEditor/Private/PhysicsAssetEditorEditMode.cpp b/Engine/Source/Editor/PhysicsAssetEditor/Private/PhysicsAssetEditorEditMode.cpp
--- a/Engine/Source/Editor/PhysicsAssetEditor/Private/PhysicsAssetEditorEditMode.cpp
+++ b/Engine/Source/Editor/PhysicsAssetEditor/Private/PhysicsAssetEditorEditMode.cpp
@@ -122,7 +122,32 @@ void FPhysicsAssetEditorEditMode::GetOnScreenDebugInfo(TArray<FText>& OutDebugIn
 bool FPhysicsAssetEditorEditMode::StartTracking(FEditorViewportClient* InViewportClient, FViewport* InViewport)
 {
 	const EAxisList::Type CurrentAxis = InViewportClient->GetCurrentWidgetAxis();
-	if(!SharedData->bRunningSimulation && !SharedData->bManipulating && CurrentAxis != EAxisList::None)
+	if (CurrentAxis == EAxisList::None)
+	{
+		return false; // not manipulating a required axis
+	}
+		
+	return HandleBeginTransform();
+}
+
+bool FPhysicsAssetEditorEditMode::EndTracking(FEditorViewportClient* InViewportClient, FViewport* InViewport)
+{
+	return HandleEndTransform(InViewportClient);
+}
+
+bool FPhysicsAssetEditorEditMode::BeginTransform(const FGizmoState& InState)
+{
+	return HandleBeginTransform();
+}
+
+bool FPhysicsAssetEditorEditMode::EndTransform(const FGizmoState& InState)
+{
+	return HandleEndTransform(Owner->GetFocusedViewportClient());
+}
+
+bool FPhysicsAssetEditorEditMode::HandleBeginTransform()
+{
+	if (!SharedData->bRunningSimulation && !SharedData->bManipulating)
 	{
 		if(SharedData->GetSelectedBody() || SharedData->GetSelectedConstraint())
 		{
@@ -176,7 +201,7 @@ bool FPhysicsAssetEditorEditMode::StartTracking(FEditorViewportClient* InViewpor
 	return SharedData->bManipulating;
 }
 
-bool FPhysicsAssetEditorEditMode::EndTracking(FEditorViewportClient* InViewportClient, FViewport* InViewport)
+bool FPhysicsAssetEditorEditMode::HandleEndTransform(FEditorViewportClient* InViewportClient) const
 {
 	if (SharedData->bManipulating)
 	{
@@ -202,7 +227,11 @@ bool FPhysicsAssetEditorEditMode::EndTracking(FEditorViewportClient* InViewportC
 
 		GEditor->EndTransaction();
 		SharedData->RefreshPhysicsAssetChange(SharedData->PhysicsAsset, false);
-		InViewport->Invalidate();
+
+		if (InViewportClient)
+		{
+			InViewportClient->Invalidate();
+		}
 
 		return true;
 	}
diff --git a/Engine/Source/Editor/PhysicsAssetEditor/Private/PhysicsAssetEditorEditMode.h b/Engine/Source/Editor/PhysicsAssetEditor/Private/PhysicsAssetEditorEditMode.h
--- a/Engine/Source/Editor/PhysicsAssetEditor/Private/PhysicsAssetEditorEditMode.h
+++ b/Engine/Source/Editor/PhysicsAssetEditor/Private/PhysicsAssetEditorEditMode.h
@@ -53,6 +53,11 @@ public:
 	virtual bool ReceivedFocus(FEditorViewportClient* ViewportClient, FViewport* Viewport);
 	virtual bool LostFocus(FEditorViewportClient* ViewportClient, FViewport* Viewport);
 
+	// Start IGizmoEdModeInterface overrides
+	virtual bool BeginTransform(const FGizmoState& InState) override;
+	virtual bool EndTransform(const FGizmoState& InState) override;
+	// End IGizmoEdModeInterface overrides
+	
 private:
 	/** Simulation mouse forces */
 	bool SimMousePress(FEditorViewportClient* InViewportClient, FKey Key);
@@ -81,6 +86,10 @@ private:
 	/** Returns the identifier for the constraint frame (child or parent) in which the manipulator widget should be drawn. */
 	EConstraintFrame::Type GetConstraintFrameForWidget() const;
 
+	// Manage the start and end of a transform action in the viewport.
+	bool HandleBeginTransform();
+	bool HandleEndTransform(FEditorViewportClient* InViewportClient) const;
+
 private:
 	/** Shared data */
 	FPhysicsAssetEditorSharedData* SharedData;
