Fix crash when destroying nodes that have no outer animation blueprint

diff --git a/Engine/Source/Editor/AnimGraph/Private/AnimGraphNode_Base.cpp b/Engine/Source/Editor/AnimGraph/Private/AnimGraphNode_Base.cpp
--- a/Engine/Source/Editor/AnimGraph/Private/AnimGraphNode_Base.cpp
+++ b/Engine/Source/Editor/AnimGraph/Private/AnimGraphNode_Base.cpp
@@ -234,11 +234,14 @@ void UAnimGraphNode_Base::PostPasteNode()
 
 void UAnimGraphNode_Base::DestroyNode()
 {
-	// This node may have been the last using its extension, so refresh
-	GetAnimBlueprint()->RequestRefreshExtensions();
-	
-	// Cleanup the pose watch if one exists on this node
-	AnimationEditorUtils::RemovePoseWatchFromNode(this, GetAnimBlueprint());
+	if(HasValidBlueprint())
+	{
+		// This node may have been the last using its extension, so refresh
+		GetAnimBlueprint()->RequestRefreshExtensions();
+		
+		// Cleanup the pose watch if one exists on this node
+		AnimationEditorUtils::RemovePoseWatchFromNode(this, GetAnimBlueprint());
+	}
 
 	Super::DestroyNode();
 }
