Fix for anim graph nodes which have an enum property bound to a property access function, and other properties connected to input pins.

diff --git a/Engine/Source/Editor/AnimGraph/Private/AnimBlueprintExtension_Base.cpp b/Engine/Source/Editor/AnimGraph/Private/AnimBlueprintExtension_Base.cpp
--- a/Engine/Source/Editor/AnimGraph/Private/AnimBlueprintExtension_Base.cpp
+++ b/Engine/Source/Editor/AnimGraph/Private/AnimBlueprintExtension_Base.cpp
@@ -509,6 +509,12 @@ struct FInternalOptionalPinManager : public FOptionalPinManager
 						NewPin->PinFriendlyName = FText::FromString(PropertyEntry.PropertyFriendlyName.IsEmpty() ? PinName.ToString() : PropertyEntry.PropertyFriendlyName);
 						NewPin->bNotConnectable = !PropertyEntry.bIsSetValuePinVisible;
 						NewPin->bDefaultValueIsIgnored = !PropertyEntry.bIsSetValuePinVisible;
+						
+						if (FEnumProperty* EnumProperty = CastField<FEnumProperty>(OuterProperty))
+						{
+							NewPin->DefaultValue = EnumProperty->GetEnum()->GetNameStringByIndex(0);
+						}
+						
 						Schema->ConstructBasicPinTooltip(*NewPin, PropertyEntry.PropertyTooltip, NewPin->PinToolTip);
 					}
 				}
