Fix issue with byte properties containing enums where default pin values weren't being set

diff --git a/Engine/Source/Editor/AnimGraph/Private/AnimBlueprintExtension_Base.cpp b/Engine/Source/Editor/AnimGraph/Private/AnimBlueprintExtension_Base.cpp
--- a/Engine/Source/Editor/AnimGraph/Private/AnimBlueprintExtension_Base.cpp
+++ b/Engine/Source/Editor/AnimGraph/Private/AnimBlueprintExtension_Base.cpp
@@ -514,7 +514,14 @@ struct FInternalOptionalPinManager : public FOptionalPinManager
 						{
 							NewPin->DefaultValue = EnumProperty->GetEnum()->GetNameStringByIndex(0);
 						}
-						
+						else if (FByteProperty* ByteProperty = CastField<FByteProperty>(OuterProperty))
+						{
+							if (ByteProperty->IsEnum())
+							{
+								NewPin->DefaultValue = ByteProperty->GetIntPropertyEnum()->GetNameStringByIndex(0);
+							}							
+						}
+
 						Schema->ConstructBasicPinTooltip(*NewPin, PropertyEntry.PropertyTooltip, NewPin->PinToolTip);
 					}
 				}
