Add a compile time flag to FAnimNode_LinkedInputPose to know if they are connected to the graph's root

review-33657663 [REDACTED_HANDLE]

diff --git a/Engine/Source/Editor/AnimGraph/Private/AnimBlueprintExtension_LinkedInputPose.cpp b/Engine/Source/Editor/AnimGraph/Private/AnimBlueprintExtension_LinkedInputPose.cpp
new file mode 100644
--- /dev/null
+++ b/Engine/Source/Editor/AnimGraph/Private/AnimBlueprintExtension_LinkedInputPose.cpp
@@ -0,0 +1,19 @@
+// Copyright Epic Games, Inc. All Rights Reserved.
+
+#include "AnimBlueprintExtension_LinkedInputPose.h"
+#include "AnimGraphNode_Base.h"
+#include "AnimGraphNode_LinkedInputPose.h"
+#include "IAnimBlueprintCompilerCreationContext.h"
+#include "IAnimBlueprintCompilationContext.h"
+#include "IAnimBlueprintGeneratedClassCompiledData.h"
+
+void UAnimBlueprintExtension_LinkedInputPose::HandlePreProcessAnimationNodes(TArrayView<UAnimGraphNode_Base*> InAnimNodes, IAnimBlueprintCompilationContext& InCompilationContext, IAnimBlueprintGeneratedClassCompiledData& OutCompiledData)
+{
+	for(UAnimGraphNode_Base* AnimNode : InAnimNodes)
+	{
+		if(UAnimGraphNode_LinkedInputPose* LinkedInputPose = Cast<UAnimGraphNode_LinkedInputPose>(AnimNode))
+		{
+			LinkedInputPose->AnalyzeLinks(InAnimNodes);
+		}
+	}
+}
\ No newline at end of file
diff --git a/Engine/Source/Editor/AnimGraph/Private/AnimBlueprintExtension_LinkedInputPose.h b/Engine/Source/Editor/AnimGraph/Private/AnimBlueprintExtension_LinkedInputPose.h
new file mode 100644
--- /dev/null
+++ b/Engine/Source/Editor/AnimGraph/Private/AnimBlueprintExtension_LinkedInputPose.h
@@ -0,0 +1,22 @@
+// Copyright Epic Games, Inc. All Rights Reserved.
+
+#pragma once
+
+#include "CoreMinimal.h"
+#include "AnimBlueprintExtension.h"
+#include "AnimBlueprintExtension_LinkedInputPose.generated.h"
+
+class UAnimGraphNode_Base;
+class IAnimBlueprintCompilerCreationContext;
+class IAnimBlueprintCompilationContext;
+class IAnimBlueprintGeneratedClassCompiledData;
+
+UCLASS(MinimalAPI)
+class UAnimBlueprintExtension_LinkedInputPose : public UAnimBlueprintExtension
+{
+	GENERATED_BODY()
+
+private:
+	// UAnimBlueprintExtension interface
+	virtual void HandlePreProcessAnimationNodes(TArrayView<UAnimGraphNode_Base*> InAnimNodes, IAnimBlueprintCompilationContext& InCompilationContext, IAnimBlueprintGeneratedClassCompiledData& OutCompiledData) override;
+};
\ No newline at end of file
diff --git a/Engine/Source/Editor/AnimGraph/Private/AnimGraphNode_LinkedInputPose.cpp b/Engine/Source/Editor/AnimGraph/Private/AnimGraphNode_LinkedInputPose.cpp
--- a/Engine/Source/Editor/AnimGraph/Private/AnimGraphNode_LinkedInputPose.cpp
+++ b/Engine/Source/Editor/AnimGraph/Private/AnimGraphNode_LinkedInputPose.cpp
@@ -1,6 +1,7 @@
 // Copyright Epic Games, Inc. All Rights Reserved.
 
 #include "AnimGraphNode_LinkedInputPose.h"
+#include "AnimBlueprintExtension_LinkedInputPose.h"
 #include "Kismet2/BlueprintEditorUtils.h"
 #include "GraphEditorSettings.h"
 #include "BlueprintActionFilter.h"
@@ -549,6 +550,12 @@ void UAnimGraphNode_LinkedInputPose::OnCopyTermDefaultsToDefaultObject(IAnimBlue
 	DestinationNode->Graph = TrueNode->GetGraph()->GetFName();
 }
 
+void UAnimGraphNode_LinkedInputPose::GetRequiredExtensions(TArray<TSubclassOf<UAnimBlueprintExtension>>& OutExtensions) const
+{
+	Super::GetRequiredExtensions(OutExtensions);
+	OutExtensions.Add(UAnimBlueprintExtension_LinkedInputPose::StaticClass());
+}
+
 TSharedRef<SWidget> UAnimGraphNode_LinkedInputPose::MakeNameWidget(IDetailLayoutBuilder& DetailBuilder)
 {
 	TSharedPtr<IPropertyHandle> NamePropertyHandle = DetailBuilder.GetProperty(GET_MEMBER_NAME_CHECKED(UAnimGraphNode_LinkedInputPose, Node.Name), GetClass());
@@ -690,4 +697,72 @@ void UAnimGraphNode_LinkedInputPose::GetOutputLinkAttributes(FNodeAttributeArray
 	});
 }
 
+void UAnimGraphNode_LinkedInputPose::AnalyzeLinks(TArrayView<UAnimGraphNode_Base*> InAnimNodes)
+{
+	Node.bIsOutputLinked = false;
+
+	// Find root nodes
+	TArray<UAnimGraphNode_Base*> RootSet;
+	for (UAnimGraphNode_Base* InNode : InAnimNodes)
+	{
+		if (InNode->IsNodeRootSet())
+		{
+			RootSet.Add(InNode);
+		}
+	}
+
+	struct FNodeVisitorDownPoseWires
+	{
+		TSet<UEdGraphNode*> VisitedNodes;
+		const UAnimationGraphSchema* Schema;
+
+		FNodeVisitorDownPoseWires()
+		{
+			Schema = GetDefault<UAnimationGraphSchema>();
+		}
+
+		// Traverse root set looking for this node
+		bool TraverseNodes(UEdGraphNode* Node, UAnimGraphNode_LinkedInputPose* SearchedNode)
+		{
+			VisitedNodes.Add(Node);
+
+			// Follow every exec output pin
+			for (int32 i = 0; i < Node->Pins.Num(); ++i)
+			{
+				UEdGraphPin* MyPin = Node->Pins[i];
+
+				if ((MyPin->Direction == EGPD_Input) && (Schema->IsPosePin(MyPin->PinType)))
+				{
+					for (int32 j = 0; j < MyPin->LinkedTo.Num(); ++j)
+					{
+						UEdGraphPin* OtherPin = MyPin->LinkedTo[j];
+						UEdGraphNode* OtherNode = OtherPin->GetOwningNode();
+						if (!VisitedNodes.Contains(OtherNode))
+						{
+							// Found ?
+							if (OtherNode == SearchedNode || TraverseNodes(OtherNode, SearchedNode))
+							{
+								return true;
+							}
+						}
+					}
+				}
+			}
+			return false;
+		}
+	};
+
+	FNodeVisitorDownPoseWires Visitor;
+
+	// Search every root node
+	for (UAnimGraphNode_Base* RootNode : RootSet)
+	{
+		if (Visitor.TraverseNodes(RootNode, this))
+		{
+			Node.bIsOutputLinked = true;
+			break;
+		}
+	}
+}
+
 #undef LOCTEXT_NAMESPACE
diff --git a/Engine/Source/Editor/AnimGraph/Public/AnimGraphNode_LinkedInputPose.h b/Engine/Source/Editor/AnimGraph/Public/AnimGraphNode_LinkedInputPose.h
--- a/Engine/Source/Editor/AnimGraph/Public/AnimGraphNode_LinkedInputPose.h
+++ b/Engine/Source/Editor/AnimGraph/Public/AnimGraphNode_LinkedInputPose.h
@@ -89,6 +89,7 @@ public:
 	virtual void GetOutputLinkAttributes(FNodeAttributeArray& OutAttributes) const override;
 	virtual bool ShouldShowAttributesOnPins() const override { return false; }
 	virtual void OnCopyTermDefaultsToDefaultObject(IAnimBlueprintCopyTermDefaultsContext& InCompilationContext, IAnimBlueprintNodeCopyTermDefaultsContext& InPerNodeContext, IAnimBlueprintGeneratedClassCompiledData& OutCompiledData) override;
+	virtual void GetRequiredExtensions(TArray<TSubclassOf<UAnimBlueprintExtension>>& OutExtensions) const override;
 
 	/** Make a name widget for this linked input pose node */
 	TSharedRef<SWidget> MakeNameWidget(IDetailLayoutBuilder& DetailBuilder);
@@ -124,6 +125,10 @@ private:
 	/** Create pins from the stub function FunctionReference */
 	void CreatePinsFromStubFunction(const UFunction* Function);
 
+	friend class UAnimBlueprintExtension_LinkedInputPose;
+	// Called pre-compilation to evaluate is this input pose is used
+	void AnalyzeLinks(TArrayView<UAnimGraphNode_Base*> InAnimNodes);
+
 	/** Reconstruct any layer nodes in this BP post-edit */
 	static void ReconstructLayerNodes(UBlueprint* InBlueprint);
 	
diff --git a/Engine/Source/Runtime/Engine/Classes/Animation/AnimNode_LinkedInputPose.h b/Engine/Source/Runtime/Engine/Classes/Animation/AnimNode_LinkedInputPose.h
--- a/Engine/Source/Runtime/Engine/Classes/Animation/AnimNode_LinkedInputPose.h
+++ b/Engine/Source/Runtime/Engine/Classes/Animation/AnimNode_LinkedInputPose.h
@@ -20,6 +20,8 @@ struct FAnimNode_LinkedInputPose : public FAnimNode_Base
 	FAnimNode_LinkedInputPose()
 		: Name(DefaultInputPoseName)
 		, Graph(NAME_None)
+		, bIsCachedInputPoseInitialized(false)
+		, bIsOutputLinked(true)
 		, OuterGraphNodeIndex(INDEX_NONE)
 		, InputProxy(nullptr)
 	{
@@ -47,11 +49,17 @@ struct FAnimNode_LinkedInputPose : public FAnimNode_Base
 
 	// CachedInputPose can have bone data allocated but uninitialized.
 	// This can happen if an anim graph has an Input Pose node with nothing populating it (e.g. if it's played as the only animbp on an actor).
-	bool bIsCachedInputPoseInitialized = false;
+	uint8 bIsCachedInputPoseInitialized : 1;
+
+	// True if this linked input pose output is connected to the graph root
+	UPROPERTY(meta = (BlueprintCompilerGeneratedDefaults))
+	uint8 bIsOutputLinked : 1;
 
 	// The node index of the currently-linked outer node
 	int32 OuterGraphNodeIndex;
 
+
+
 	// FAnimNode_Base interface
 #if ENABLE_ANIMGRAPH_TRAVERSAL_DEBUG
 	ENGINE_API virtual void Initialize_AnyThread(const FAnimationInitializeContext& Context) override;
