Fix GCC error that won't allow a nested struct to share a variable name with the class its nested in...

#rnx

diff --git a/Engine/Source/Editor/AnimGraph/Private/AnimGraphNode_LinkedInputPose.cpp b/Engine/Source/Editor/AnimGraph/Private/AnimGraphNode_LinkedInputPose.cpp
--- a/Engine/Source/Editor/AnimGraph/Private/AnimGraphNode_LinkedInputPose.cpp
+++ b/Engine/Source/Editor/AnimGraph/Private/AnimGraphNode_LinkedInputPose.cpp
@@ -722,14 +722,14 @@ void UAnimGraphNode_LinkedInputPose::AnalyzeLinks(TArrayView<UAnimGraphNode_Base
 		}
 
 		// Traverse root set looking for this node
-		bool TraverseNodes(UEdGraphNode* Node, UAnimGraphNode_LinkedInputPose* SearchedNode)
+		bool TraverseNodes(UEdGraphNode* CurrentNode, UAnimGraphNode_LinkedInputPose* SearchedNode)
 		{
-			VisitedNodes.Add(Node);
+			VisitedNodes.Add(CurrentNode);
 
 			// Follow every exec output pin
-			for (int32 i = 0; i < Node->Pins.Num(); ++i)
+			for (int32 i = 0; i < CurrentNode->Pins.Num(); ++i)
 			{
-				UEdGraphPin* MyPin = Node->Pins[i];
+				UEdGraphPin* MyPin = CurrentNode->Pins[i];
 
 				if ((MyPin->Direction == EGPD_Input) && (Schema->IsPosePin(MyPin->PinType)))
 				{
