Add editor option to hide long list of enum blueprint actions in the ABP editor.
* Replaced with a single "Blend poses by enum" node that has an unbound enum.
* The property window now shows the bound enum and can be modified

diff --git a/Engine/Source/Editor/AnimGraph/AnimGraph.Build.cs b/Engine/Source/Editor/AnimGraph/AnimGraph.Build.cs
--- a/Engine/Source/Editor/AnimGraph/AnimGraph.Build.cs
+++ b/Engine/Source/Editor/AnimGraph/AnimGraph.Build.cs
@@ -37,6 +37,7 @@ public class AnimGraph : ModuleRules
 				"EditorWidgets",
 				"ToolWidgets",
 				"AnimationEditMode",
+				"AnimationBlueprintEditor",
 			}
 		);
 
diff --git a/Engine/Source/Editor/AnimGraph/Private/AnimGraphNode_BlendListByEnum.cpp b/Engine/Source/Editor/AnimGraph/Private/AnimGraphNode_BlendListByEnum.cpp
--- a/Engine/Source/Editor/AnimGraph/Private/AnimGraphNode_BlendListByEnum.cpp
+++ b/Engine/Source/Editor/AnimGraph/Private/AnimGraphNode_BlendListByEnum.cpp
@@ -10,12 +10,19 @@
 #include "Kismet2/CompilerResultsLog.h"
 #include "BlueprintActionDatabaseRegistrar.h"
 #include "BlueprintNodeSpawner.h"
+#include "AnimationBlueprintEditorSettings.h"
+#include "BlueprintEditorSettings.h"
+#include "DetailLayoutBuilder.h"
+#include "DetailWidgetRow.h"
+#include "SSearchableComboBox.h"
 
 #define LOCTEXT_NAMESPACE "BlendListByEnum"
 
 /////////////////////////////////////////////////////
 // UAnimGraphNode_BlendListByEnum
 
+TArray<TSharedPtr<FString>> UAnimGraphNode_BlendListByEnum::AvailableEnums;
+
 UAnimGraphNode_BlendListByEnum::UAnimGraphNode_BlendListByEnum(const FObjectInitializer& ObjectInitializer)
 	: Super(ObjectInitializer)
 {
@@ -25,7 +32,7 @@ UAnimGraphNode_BlendListByEnum::UAnimGraphNode_BlendListByEnum(const FObjectInit
 
 FText UAnimGraphNode_BlendListByEnum::GetMenuCategory() const
 {
-	return LOCTEXT("AnimGraphNode_BlendListByEnum_GetMenuCategory", "Animation|Blends|Blend by Enum");
+	return LOCTEXT("AnimGraphNode_BlendListByEnum_GetMenuCategory", "Animation|Blends");
 }
 
 FText UAnimGraphNode_BlendListByEnum::GetTooltipText() const
@@ -38,7 +45,7 @@ FText UAnimGraphNode_BlendListByEnum::GetNodeTitle(ENodeTitleType::Type TitleTyp
 {
 	if (BoundEnum == nullptr)
 	{
-		return LOCTEXT("AnimGraphNode_BlendListByEnum_TitleError", "ERROR: Blend Poses (by missing enum)");
+		return LOCTEXT("AnimGraphNode_BlendListByEnum_TitleError", "Blend Poses by enum");
 	}
 	// @TODO: don't know enough about this node type to comfortably assert that
 	//        the BoundEnum won't change after the node has spawned... until
@@ -55,6 +62,12 @@ FText UAnimGraphNode_BlendListByEnum::GetNodeTitle(ENodeTitleType::Type TitleTyp
 
 void UAnimGraphNode_BlendListByEnum::GetMenuActions(FBlueprintActionDatabaseRegistrar& ActionRegistrar) const
 {
+	if (!GetDefault<UAnimationBlueprintEditorSettings>()->bShowInstancedEnumBlendAnimNodeBlueprintActions)
+	{
+		Super::GetMenuActions(ActionRegistrar);
+		return;
+	}
+
 	struct GetMenuActions_Utils
 	{
 		static void SetNodeEnum(UEdGraphNode* NewNode, bool /*bIsTemplateNode*/, TWeakObjectPtr<UEnum> NonConstEnumPtr)
@@ -288,7 +301,7 @@ void UAnimGraphNode_BlendListByEnum::ValidateAnimNodeDuringCompilation(class USk
 
 	if (BoundEnum == NULL)
 	{
-		MessageLog.Error(TEXT("@@ references an unknown enum; please delete the node and recreate it"), this);
+		MessageLog.Error(TEXT("@@ references an unknown enum; please select the enum to use"), this);
 	}
 }
 
@@ -340,6 +353,104 @@ void UAnimGraphNode_BlendListByEnum::ReloadEnum(class UEnum* InEnum)
 	CachedNodeTitle.MarkDirty();
 }
 
+void UAnimGraphNode_BlendListByEnum::CustomizeDetails(IDetailLayoutBuilder& DetailBuilder)
+{
+	Super::CustomizeDetails(DetailBuilder);
+
+	IDetailCategoryBuilder& SettingsCategory = DetailBuilder.EditCategory("Settings", LOCTEXT("SettingsCategoryName", "Settings"));
+
+	RefreshAvailableEnums();
+
+	TSharedPtr<FString> SelectedEnum = nullptr;
+	if (BoundEnum != nullptr)
+	{
+		TSharedPtr<FString>* FoundEnum = AvailableEnums.FindByPredicate([this](const TSharedPtr<FString> Candidate)
+		{
+			return *Candidate == BoundEnum.GetPathName();
+		});
+
+		if (FoundEnum == nullptr)
+		{
+			ClearEnum();
+		}
+		else
+		{
+			SelectedEnum = *FoundEnum;
+		}
+	}
+
+	SettingsCategory.AddCustomRow(LOCTEXT("BoundEnum", "Bound Enum"))
+		.NameContent()
+		[
+			SNew(STextBlock)
+			.Text(LOCTEXT("BoundEnum", "Bound Enum"))
+			.Font(FAppStyle::Get().GetFontStyle("PropertyWindow.NormalFont"))
+		]
+		.ValueContent()
+		[
+			SNew(SSearchableComboBox)
+			.InitiallySelectedItem(SelectedEnum)
+			.OptionsSource(&AvailableEnums)
+			.OnGenerateWidget_Lambda([](TSharedPtr<FString> InItem)
+				{
+					return SNew(STextBlock).Text(FText::FromString(*InItem));
+				})
+			.OnSelectionChanged_Lambda([this](TSharedPtr<FString> NewChoice, ESelectInfo::Type SelectType)
+				{
+					FScopedTransaction Transaction( LOCTEXT("ModifyBoundEnum", "Modify Bound Enum") );
+					Modify();
+
+					ClearEnum();
+
+					BoundEnum = FindObject<UEnum>(nullptr, **NewChoice);
+					ReconstructNode();
+				})
+			[
+				SNew(STextBlock)
+				.Text_Lambda([this]()
+				{
+					return FText::FromString(BoundEnum->GetFullName());
+				})
+				.Font(FAppStyle::GetFontStyle(TEXT("PropertyWindow.NormalFont")))
+			]
+		];
+}
+
+void UAnimGraphNode_BlendListByEnum::RefreshAvailableEnums()
+{
+	AvailableEnums.Empty();
+	const FPathPermissionList& Permissions = GetMutableDefault<UBlueprintEditorSettings>()->GetEnumPermissions();
+
+	for (TObjectIterator<UEnum> EnumIt; EnumIt; ++EnumIt)
+	{
+		if (!EnumIt->HasAnyEnumFlags(EEnumFlags::NewerVersionExists) && UEdGraphSchema_K2::IsAllowableBlueprintVariableType(*EnumIt))
+		{
+			if (Permissions.HasFiltering())
+			{
+				TStringBuilder<256> ResultBuilder;
+				EnumIt->GetPathName(nullptr, ResultBuilder);
+				if (!Permissions.PassesFilter(ResultBuilder.ToView()))
+				{
+					return;
+				}
+			}
+
+			AvailableEnums.Add(MakeShared<FString>(EnumIt->GetPathName()));
+		}
+	}
+}
+
+void UAnimGraphNode_BlendListByEnum::ClearEnum()
+{
+	for (int32 EnumPoseIndex = VisibleEnumEntries.Num() - 1; EnumPoseIndex >= 0; --EnumPoseIndex)
+	{
+		Node.RemovePose(EnumPoseIndex + 1);
+	}
+	VisibleEnumEntries.Empty();
+
+	BoundEnum = nullptr;
+}
+
 /////////////////////////////////////////////////////
 
 #undef LOCTEXT_NAMESPACE
diff --git a/Engine/Source/Editor/AnimGraph/Public/AnimGraphNode_BlendListByEnum.h b/Engine/Source/Editor/AnimGraph/Public/AnimGraphNode_BlendListByEnum.h
--- a/Engine/Source/Editor/AnimGraph/Public/AnimGraphNode_BlendListByEnum.h
+++ b/Engine/Source/Editor/AnimGraph/Public/AnimGraphNode_BlendListByEnum.h
@@ -45,6 +45,7 @@ public:
 	virtual void ValidateAnimNodeDuringCompilation(class USkeleton* ForSkeleton, class FCompilerResultsLog& MessageLog) override;
 	virtual void BakeDataDuringCompilation(class FCompilerResultsLog& MessageLog) override;
 	virtual void PreloadRequiredAssets() override;
+	virtual void CustomizeDetails(IDetailLayoutBuilder& DetailBuilder) override;
 	// End of UAnimGraphNode_Base interface
 
 	//@TODO: Generalize this behavior (returning a list of actions/delegates maybe?)
@@ -63,6 +64,13 @@ protected:
 	static void GetPinInformation(const FString& InPinName, int32& Out_PinIndex, bool& Out_bIsPosePin, bool& Out_bIsTimePin);
 
 private:
+	void ClearEnum();
+
+	static void RefreshAvailableEnums();
+
 	/** Constructing FText strings can be costly, so we cache the node's title */
 	FNodeTextCache CachedNodeTitle;
+
+	/** Shared array that contains all possible enums to select from */
+	static TArray<TSharedPtr<FString>> AvailableEnums;
 };
diff --git a/Engine/Source/Editor/AnimationBlueprintEditor/Private/AnimationBlueprintEditorSettings.cpp b/Engine/Source/Editor/AnimationBlueprintEditor/Private/AnimationBlueprintEditorSettings.cpp
--- a/Engine/Source/Editor/AnimationBlueprintEditor/Private/AnimationBlueprintEditorSettings.cpp
+++ b/Engine/Source/Editor/AnimationBlueprintEditor/Private/AnimationBlueprintEditorSettings.cpp
@@ -1,9 +1,17 @@
 // Copyright Epic Games, Inc. All Rights Reserved.
 
 #include "AnimationBlueprintEditorSettings.h"
+#include "BlueprintActionDatabase.h"
 
 void UAnimationBlueprintEditorSettings::PostEditChangeProperty(FPropertyChangedEvent& PropertyChangedEvent)
 {
+	const FName PropertyName = (PropertyChangedEvent.Property ? PropertyChangedEvent.Property->GetFName() : NAME_None);
+
+	if (PropertyName == GET_MEMBER_NAME_STRING_CHECKED(UAnimationBlueprintEditorSettings, bShowInstancedEnumBlendAnimNodeBlueprintActions))
+	{
+		FBlueprintActionDatabase::Get().RefreshAll();
+	}
+
 	Super::PostEditChangeProperty(PropertyChangedEvent);
 	OnSettingsChange.Broadcast(this, PropertyChangedEvent.ChangeType);
 }
diff --git a/Engine/Source/Editor/AnimationBlueprintEditor/Public/AnimationBlueprintEditorSettings.h b/Engine/Source/Editor/AnimationBlueprintEditor/Public/AnimationBlueprintEditorSettings.h
--- a/Engine/Source/Editor/AnimationBlueprintEditor/Public/AnimationBlueprintEditorSettings.h
+++ b/Engine/Source/Editor/AnimationBlueprintEditor/Public/AnimationBlueprintEditorSettings.h
@@ -16,11 +16,18 @@
 
 // Settings for the Animation Blueprint Editor
 UCLASS(config = EditorPerProjectUserSettings)
-class UAnimationBlueprintEditorSettings : public UObject
+class ANIMATIONBLUEPRINTEDITOR_API UAnimationBlueprintEditorSettings : public UObject
 {
 	GENERATED_BODY()
 
 public:
+	/**
+	 * If true, populates the blueprint action menu with pre-bound blend-by-enum nodes for supported enums.
+	 * For large projects this can clutter the context menu so it may be preferable to hide those entries.
+	 * */
+	UPROPERTY(EditAnywhere, config, Category = "Workflow")
+	bool bShowInstancedEnumBlendAnimNodeBlueprintActions = true;
+
 	/** If true, automatically pose watch selected nodes. */
 	UPROPERTY(EditAnywhere, config, Category = "Debugging")
 	bool bPoseWatchSelectedNodes = false;
