Added sorting of Anim Graph Node input pins based on Display Priority

diff --git a/Engine/Source/Editor/AnimGraph/Private/AnimBlueprintNodeOptionalPinManager.cpp b/Engine/Source/Editor/AnimGraph/Private/AnimBlueprintNodeOptionalPinManager.cpp
--- a/Engine/Source/Editor/AnimGraph/Private/AnimBlueprintNodeOptionalPinManager.cpp
+++ b/Engine/Source/Editor/AnimGraph/Private/AnimBlueprintNodeOptionalPinManager.cpp
@@ -127,13 +127,34 @@ void FAnimBlueprintNodeOptionalPinManager::AllocateDefaultPins(UStruct* SourceSt
 	RebuildPropertyList(BaseNode->ShowPinForProperties, SourceStruct);
 
 	const UAnimationGraphSchema* Schema = GetDefault<UAnimationGraphSchema>();
-	
-	// Sort pins by property offset (makes pose pins consistent across derived nodes amongst other things)
+	static const FName NAME_DisplayPriority("DisplayPriority");
+
+	constexpr auto GetDisplayPriority = [](const FProperty* Prop)
+	{
+		const FString& DisplayPriorityStr = Prop->GetMetaData(NAME_DisplayPriority);
+		int32 DisplayPriority = (DisplayPriorityStr.IsEmpty() ? MAX_int32 : FCString::Atoi(*DisplayPriorityStr));
+		if (DisplayPriority == 0 && !FCString::IsNumeric(*DisplayPriorityStr))
+		{
+			// If there was a malformed display priority str Atoi will say it is 0, but we want to treat it as unset
+			DisplayPriority = MAX_int32;
+		}
+
+		return DisplayPriority;
+	};
+
+	// Sort pins by display priority and property offset (makes pose pins consistent across derived nodes amongst other things)
 	Algo::StableSort(BaseNode->ShowPinForProperties, [SourceStruct, Schema](const FOptionalPinFromProperty& InPin0, const FOptionalPinFromProperty& InPin1)
 	{
 		FProperty* Property0 = FindFieldChecked<FProperty>(SourceStruct, InPin0.PropertyName);
 		FProperty* Property1 = FindFieldChecked<FProperty>(SourceStruct, InPin1.PropertyName);
-		
+
+		const int32 Priority0 = GetDisplayPriority(Property0);
+		const int32 Priority1 = GetDisplayPriority(Property1);
+		if (Priority0 != Priority1)
+		{
+			return Priority0 < Priority1;
+		}
+
 		return Property0->GetOffset_ForInternal() < Property1->GetOffset_ForInternal();
 	});
 	
