[UE] Anim Dynamics Node - Ensure preview node instance physics body properties are updated when they are changed in the editor preview panel.
- A seperate instance of the anim dynamics runtime node is created when an AnimDynamics graph node is selected in editor. This second instance is used to define the physics simulation and the debug draw shown in the viewport. This change ensurtes that when a physics body definition is changed in the details panel the same changes are applied to this second instance of the node.

diff --git a/Engine/Source/Editor/AnimGraph/Private/AnimGraphNode_AnimDynamics.cpp b/Engine/Source/Editor/AnimGraph/Private/AnimGraphNode_AnimDynamics.cpp
--- a/Engine/Source/Editor/AnimGraph/Private/AnimGraphNode_AnimDynamics.cpp
+++ b/Engine/Source/Editor/AnimGraph/Private/AnimGraphNode_AnimDynamics.cpp
@@ -111,12 +111,25 @@ void UAnimGraphNode_AnimDynamics::CustomizeDetails(IDetailLayoutBuilder& DetailB
 	}
 }
 
-
 void UAnimGraphNode_AnimDynamics::ValidateAnimNodeDuringCompilation(USkeleton* ForSkeleton, FCompilerResultsLog& MessageLog)
 {
 	Super::ValidateAnimNodeDuringCompilation(ForSkeleton, MessageLog);
 }
 
+void UAnimGraphNode_AnimDynamics::UpdatePreviewDynamicsNode(const int32 PhysicsBodyIndex)
+{
+	// Update the preview node from details panel and reset physics simulation.
+
+	if(FAnimNode_AnimDynamics* const ActivePreviewNode = GetPreviewDynamicsNode())
+	{
+		if(ActivePreviewNode->PhysicsBodyDefinitions.IsValidIndex(PhysicsBodyIndex) && Node.PhysicsBodyDefinitions.IsValidIndex(PhysicsBodyIndex))
+		{
+			ActivePreviewNode->PhysicsBodyDefinitions[PhysicsBodyIndex] = Node.PhysicsBodyDefinitions[PhysicsBodyIndex];
+			ResetSim();
+		}
+	}
+}
+
 FEditorModeID UAnimGraphNode_AnimDynamics::GetEditorMode() const
 {
 	return AnimNodeEditModes::AnimDynamics;
@@ -271,7 +284,7 @@ void UAnimGraphNode_AnimDynamics::Serialize(FArchive& Ar)
 
 	if (Ar.CustomVer(FFortniteMainBranchObjectVersion::GUID) < FFortniteMainBranchObjectVersion::AnimDynamicsEditableChainParameters)
 	{
-		// Initialise first physics body using deprecated parameters then re-initialize chain.
+		// Initialize first physics body using deprecated parameters then re-initialize chain.
 		Node.PhysicsBodyDefinitions.Reset();		
 		FAnimPhysBodyDefinition PhysBodyDef;
 		PhysBodyDef.BoundBone = Node.BoundBone;
@@ -310,6 +323,8 @@ void UAnimGraphNode_AnimDynamics::OnPhysicsBodyDefCustomizeDetails(TSharedRef<IP
 				.Text_Lambda([this, ElementIndex]() -> FText { return this->BodyDefinitionUINameString(ElementIndex); })
 			];
 	}
+
+	ElementProperty->SetOnChildPropertyValueChanged(FSimpleDelegate::CreateLambda([this, ElementIndex]() { this->UpdatePreviewDynamicsNode(ElementIndex); }));
 }
 
 FText UAnimGraphNode_AnimDynamics::BodyDefinitionUINameString(const uint32 BodyIndex) const
diff --git a/Engine/Source/Editor/AnimGraph/Public/AnimGraphNode_AnimDynamics.h b/Engine/Source/Editor/AnimGraph/Public/AnimGraphNode_AnimDynamics.h
--- a/Engine/Source/Editor/AnimGraph/Public/AnimGraphNode_AnimDynamics.h
+++ b/Engine/Source/Editor/AnimGraph/Public/AnimGraphNode_AnimDynamics.h
@@ -92,6 +92,8 @@ public:
 	void OnPhysicsBodyDefCustomizeDetails(TSharedRef<IPropertyHandle> ElementProperty, int32 ElementIndex, IDetailChildrenBuilder& ChildrenBuilder, IDetailLayoutBuilder* DetailLayout);
 	FText BodyDefinitionUINameString(const uint32 BodyIndex) const;
 
+	void UpdatePreviewDynamicsNode(const int32 PhysicsBodyIndex);
+
 protected:
 
 	virtual FText GetControllerDescription() const override;
diff --git a/Engine/Source/Runtime/AnimGraphRuntime/Public/BoneControllers/AnimNode_AnimDynamics.h b/Engine/Source/Runtime/AnimGraphRuntime/Public/BoneControllers/AnimNode_AnimDynamics.h
--- a/Engine/Source/Runtime/AnimGraphRuntime/Public/BoneControllers/AnimNode_AnimDynamics.h
+++ b/Engine/Source/Runtime/AnimGraphRuntime/Public/BoneControllers/AnimNode_AnimDynamics.h
@@ -246,6 +246,18 @@ struct FAnimPhysBodyDefinition
 	, SphereCollisionRadius(10.0f)
 	{}
 
+	FAnimPhysBodyDefinition& operator=(const FAnimPhysBodyDefinition& Other)
+	{
+		this->BoundBone = Other.BoundBone;
+		this->BoxExtents = Other.BoxExtents;
+		this->LocalJointOffset = Other.LocalJointOffset;
+		this->ConstraintSetup = Other.ConstraintSetup;
+		this->CollisionType = Other.CollisionType;
+		this->SphereCollisionRadius = Other.SphereCollisionRadius;
+
+		return *this;
+	}
+
 	UPROPERTY(VisibleAnywhere, Category = PhysicsBodyDefinition, meta = (EditCondition = "false"))
 	FBoneReference BoundBone;
 
