Fix for animation blueprints incorrectly showing animation graphs using animation layers default shared group in my blueprint.

diff --git a/Engine/Source/Editor/AnimGraph/Private/AnimBlueprintCompiler.cpp b/Engine/Source/Editor/AnimGraph/Private/AnimBlueprintCompiler.cpp
--- a/Engine/Source/Editor/AnimGraph/Private/AnimBlueprintCompiler.cpp
+++ b/Engine/Source/Editor/AnimGraph/Private/AnimBlueprintCompiler.cpp
@@ -1647,7 +1647,11 @@ void FAnimBlueprintCompilerContext::CreateAnimGraphStubFunctions()
 				UK2Node_FunctionEntry* EntryNode = SpawnIntermediateNode<UK2Node_FunctionEntry>(RootNode, StubGraph);
 				EntryNode->NodePosX = -200;
 				EntryNode->CustomGeneratedFunctionName = GraphToUseforSignature->GetFName();	// Note that the function generated from this temporary graph is undecorated
-				EntryNode->MetaData.Category = (RootNode->Node.GetGroup() == NAME_None) ? FText::GetEmpty() : FText::FromName(RootNode->Node.GetGroup());
+
+				const bool bIsDefaultAnimGraph = GraphToUseforSignature->GetFName() == UEdGraphSchema_K2::GN_AnimGraph;
+				EntryNode->MetaData.Category = ((RootNode->Node.GetGroup() == NAME_None) || (bIsDefaultAnimGraph && RootNode->Node.GetGroup() == FAnimNode_Root::DefaultSharedGroup)) 
+					? FText::GetEmpty()
+					: FText::FromName(RootNode->Node.GetGroup());
 
 				// Add linked input poses as parameters
 				for(UAnimGraphNode_LinkedInputPose* LinkedInputPoseNode : LinkedInputPoseNodes)
