Fix K2Node_AnimGetter title not updating on state rename

diff --git a/Engine/Source/Editor/AnimGraph/Private/K2Node_AnimGetter.cpp b/Engine/Source/Editor/AnimGraph/Private/K2Node_AnimGetter.cpp
--- a/Engine/Source/Editor/AnimGraph/Private/K2Node_AnimGetter.cpp
+++ b/Engine/Source/Editor/AnimGraph/Private/K2Node_AnimGetter.cpp
@@ -40,6 +40,54 @@ class UBlueprint;
 
 #define LOCTEXT_NAMESPACE "AnimGetter"
 
+void UK2Node_AnimGetter::PostLoad()
+{
+	Super::PostLoad();
+
+	if (!HasAnyFlags(RF_ClassDefaultObject))
+	{
+		RegisterDelegates();
+	}
+}
+
+void UK2Node_AnimGetter::PostPlacedNewNode()
+{
+	Super::PostPlacedNewNode();
+
+	RegisterDelegates();
+}
+
+void UK2Node_AnimGetter::RegisterDelegates()
+{
+	if (!HasAnyFlags(RF_Transient) && HasValidBlueprint())
+	{
+		if (UAnimBlueprint* ABP = CastChecked<UAnimBlueprint>(GetBlueprint()))
+		{
+			GraphRenameHandle = ABP->OnGraphRenamedEvent().AddLambda([this](UEdGraph*, FName, FName)
+			{
+				UpdateCachedTitle();
+			});
+		}
+	}
+}
+
+void UK2Node_AnimGetter::UnregisterDelegates()
+{
+	if (!HasAnyFlags(RF_Transient) && HasValidBlueprint())
+	{
+		if (UAnimBlueprint* ABP = CastChecked<UAnimBlueprint>(GetBlueprint()))
+		{
+			ABP->OnGraphRenamedEvent().Remove(GraphRenameHandle);
+		}
+	}
+}
+
+void UK2Node_AnimGetter::BeginDestroy()
+{
+	Super::BeginDestroy();
+
+	UnregisterDelegates();	
+}
 
 void UK2Node_AnimGetter::Serialize(FArchive& Ar)
 {
diff --git a/Engine/Source/Editor/AnimGraph/Public/K2Node_AnimGetter.h b/Engine/Source/Editor/AnimGraph/Public/K2Node_AnimGetter.h
--- a/Engine/Source/Editor/AnimGraph/Public/K2Node_AnimGetter.h
+++ b/Engine/Source/Editor/AnimGraph/Public/K2Node_AnimGetter.h
@@ -71,6 +71,8 @@ public:
 	// UObject interface
 	virtual void Serialize(FArchive& Ar) override;
 	virtual void PostPasteNode() override;
+	virtual void PostLoad() override;
+	virtual void BeginDestroy() override;
 	// End of UObject interface
 
 	// UEdGraphNode interface
@@ -78,6 +80,7 @@ public:
 	virtual FText GetNodeTitle(ENodeTitleType::Type TitleType) const override;
 	virtual bool CanCreateUnderSpecifiedSchema(const UEdGraphSchema* Schema) const override;
 	virtual bool IsActionFilteredOut(FBlueprintActionFilter const& Filter) override;
+	virtual void PostPlacedNewNode();
 	// End of UEdGraphNode interface
 
 	// UK2Node interface
@@ -141,4 +144,10 @@ protected:
 	// UK2Node_CallFunction interface
 	virtual bool CanToggleNodePurity() const override { return false; }
 	// End of UK2Node_CallFunction interface
+
+	virtual void RegisterDelegates();
+	virtual void UnregisterDelegates();
+
+private:
+	FDelegateHandle GraphRenameHandle;
 };
diff --git a/Engine/Source/Runtime/Engine/Classes/Animation/AnimBlueprint.h b/Engine/Source/Runtime/Engine/Classes/Animation/AnimBlueprint.h
--- a/Engine/Source/Runtime/Engine/Classes/Animation/AnimBlueprint.h
+++ b/Engine/Source/Runtime/Engine/Classes/Animation/AnimBlueprint.h
@@ -244,6 +244,16 @@ public:
 	// Get the default binding type that any new nodes will use when created
 	ENGINE_API UClass* GetDefaultBindingClass() const { return DefaultBindingClass; }
 
+	virtual void NotifyGraphRenamed(class UEdGraph* Graph, FName OldName, FName NewName) override;
+
+	// Event that is broadcast to inform observers that the node title has changed
+	// The default SAnimationGraphNode uses this to invalidate cached node title text
+	DECLARE_EVENT_ThreeParams(UAnimBlueprint, FOnGraphRenamedEvent, UEdGraph*, FName, FName);
+	FOnGraphRenamedEvent& OnGraphRenamedEvent() { return GraphRenameEvent; }
+
+private:
+	FOnGraphRenamedEvent GraphRenameEvent;
+
 public:
 	// Array of overrides to asset containing nodes in the parent that have been overridden
 	UPROPERTY()
diff --git a/Engine/Source/Runtime/Engine/Private/AnimBlueprint.cpp b/Engine/Source/Runtime/Engine/Private/AnimBlueprint.cpp
--- a/Engine/Source/Runtime/Engine/Private/AnimBlueprint.cpp
+++ b/Engine/Source/Runtime/Engine/Private/AnimBlueprint.cpp
@@ -462,4 +462,9 @@ bool UAnimBlueprint::IsCompatibleByAssetString(const FString& InSkeletonAsset, b
 
 	return (TargetSkeleton != nullptr && TargetSkeleton->IsCompatibleForEditor(InSkeletonAsset));
 }
+
+void UAnimBlueprint::NotifyGraphRenamed(class UEdGraph* Graph, FName OldName, FName NewName)
+{
+	OnGraphRenamedEvent().Broadcast(Graph, OldName, NewName);
+}
 #endif
