Fixed rewiring animation blueprint nodes occasionally not picking up changes correctly

When rewiring nodes of the same type (say two sequence players each playing different animations), depending on the combination of default values and the current value in the running instance, re-instancing could potentially end up not giving the user what they expect (say, playing a different animation). This is because nodes are given numeric property names during compilation, and distinct nodes can be considered 'the same' depending on how nodes are rewired and processed. When it comes to reinstancing, the state of the node in the existing instance can end up 'winning', despite the node's info coming from a different source.

The solution to this is to uniquely name each node in an invariant way according to the source UEdGraphNode that they come from. In this case, to preserve name table size, I've chosen to use the hash of the node's GUID so the resulting FName can continue to just use a numeric extension rather than a whole new entry.

[RN] Fixed rewiring animation blueprint nodes occasionally not picking up changes correctly

diff --git a/Engine/Source/Editor/AnimGraph/Private/AnimGraphNode_Base.cpp b/Engine/Source/Editor/AnimGraph/Private/AnimGraphNode_Base.cpp
--- a/Engine/Source/Editor/AnimGraph/Private/AnimGraphNode_Base.cpp
+++ b/Engine/Source/Editor/AnimGraph/Private/AnimGraphNode_Base.cpp
@@ -246,6 +246,15 @@ void UAnimGraphNode_Base::DestroyNode()
 	Super::DestroyNode();
 }
 
+FString UAnimGraphNode_Base::GetDescriptiveCompiledName() const
+{
+	// Compiled names for nodes should be stable according to the locus in the graph, so if we rewire nodes we can distinguish between different nodes
+	// of the same type. In this case we use the hash of the node GUID to get a unique (or near as possible) value per node in the anim BP while still
+	// keeping the impact on the name table the same as incrementally numbering nodes during compilation via FNetNameMapping.
+	const uint32 Hash = GetTypeHash(NodeGuid);
+	return Super::GetDescriptiveCompiledName() + FString::Printf(TEXT("_%u"), Hash);
+}
+
 void UAnimGraphNode_Base::CreateOutputPins()
 {
 	if (!IsSinkNode())
diff --git a/Engine/Source/Editor/AnimGraph/Public/AnimGraphNode_Base.h b/Engine/Source/Editor/AnimGraph/Public/AnimGraphNode_Base.h
--- a/Engine/Source/Editor/AnimGraph/Public/AnimGraphNode_Base.h
+++ b/Engine/Source/Editor/AnimGraph/Public/AnimGraphNode_Base.h
@@ -257,6 +257,7 @@ public:
 	virtual void PostPlacedNewNode() override;
 	virtual void PostPasteNode() override;
 	virtual void DestroyNode() override;
+	virtual FString GetDescriptiveCompiledName() const override;
 	// End of UEdGraphNode interface
 
 	// UK2Node interface
