Fix copy-pasting shared transition rules not behaving correctly
* This fixes an issue where copy-pasting a shared transition rule will not display the graph in the 'My Blueprint' tab and will result in operations such as Find References not taking into account all graphs.

diff --git a/Engine/Source/Editor/AnimGraph/Private/AnimStateTransitionNode.cpp b/Engine/Source/Editor/AnimGraph/Private/AnimStateTransitionNode.cpp
--- a/Engine/Source/Editor/AnimGraph/Private/AnimStateTransitionNode.cpp
+++ b/Engine/Source/Editor/AnimGraph/Private/AnimStateTransitionNode.cpp
@@ -140,6 +140,15 @@ void UAnimStateTransitionNode::PostLoad()
 	{
 		ValidateBlendProfile();
 	}
+
+	// Fix up previously pasted graphs that were not correctly added to the parent sub-graphs
+	if (BoundGraph)
+	{
+		if (UEdGraph* ParentGraph = GetGraph())
+		{
+			ParentGraph->SubGraphs.AddUnique(BoundGraph);
+		}
+	}
 }
 
 bool UAnimStateTransitionNode::ValidateBlendProfile()
@@ -178,6 +187,15 @@ void UAnimStateTransitionNode::PostPasteNode()
 		CreateBoundGraph();
 	}
 
+	// Ensure transition graph is added to blueprint
+	if (BoundGraph)
+	{
+		if (UEdGraph* ParentGraph = GetGraph())
+		{
+			ParentGraph->SubGraphs.AddUnique(BoundGraph);
+		}
+	}
+
 	for (UEdGraphNode* GraphNode : BoundGraph->Nodes)
 	{
 		GraphNode->CreateNewGuid();
