Fix crash in 'call function from anim graph' node BP menu population

Also fix double-entry in the context menu coming from parent/child classes

[RN] Fix crash in 'call function from anim graph' node BP menu population

diff --git a/Engine/Source/Editor/AnimGraph/Private/AnimGraphNode_CallFunction.cpp b/Engine/Source/Editor/AnimGraph/Private/AnimGraphNode_CallFunction.cpp
--- a/Engine/Source/Editor/AnimGraph/Private/AnimGraphNode_CallFunction.cpp
+++ b/Engine/Source/Editor/AnimGraph/Private/AnimGraphNode_CallFunction.cpp
@@ -167,6 +167,11 @@ void UAnimGraphNode_CallFunction::ExpandNode(FKismetCompilerContext& CompilerCon
 
 void UAnimGraphNode_CallFunction::SetupFromFunction(UFunction* InFunction)
 {
+	if(InFunction == nullptr)
+	{
+		return;
+	}
+
 	// Create graph and inner node
 	InnerGraph = FBlueprintEditorUtils::CreateNewGraph(this, NAME_None, UEdGraph::StaticClass(), UEdGraphSchema_K2::StaticClass());
 
@@ -274,10 +279,16 @@ void UAnimGraphNode_CallFunction::GetMenuActions(FBlueprintActionDatabaseRegistr
 		{
 			if (UEdGraphSchema_K2::CanUserKismetCallFunction(InFunction) && ValidateFunction(InFunction))
 			{
-				auto CustomizeNode = [InFunction](UEdGraphNode* InNode, bool bIsTemplate)
+				auto CustomizeNode = [Function = TWeakObjectPtr<UFunction>(InFunction)](UEdGraphNode* InNode, bool bIsTemplate)
 				{
+					UFunction* ResolvedFunction = Function.Get();
+					if(ResolvedFunction == nullptr)
+					{
+						return;
+					}
+
 					UAnimGraphNode_CallFunction* CallFunctionNode = CastChecked<UAnimGraphNode_CallFunction>(InNode);
-					CallFunctionNode->SetupFromFunction(InFunction);
+					CallFunctionNode->SetupFromFunction(ResolvedFunction);
 				};
 
 				UBlueprintNodeSpawner* Spawner = UBlueprintNodeSpawner::Create(UAnimGraphNode_CallFunction::StaticClass(), nullptr, UBlueprintNodeSpawner::FCustomizeNodeDelegate::CreateLambda(CustomizeNode));
@@ -298,7 +309,7 @@ void UAnimGraphNode_CallFunction::GetMenuActions(FBlueprintActionDatabaseRegistr
 			}
 		};
 
-		for (TFieldIterator<UFunction> It(InClass); It; ++It)
+		for (TFieldIterator<UFunction> It(InClass, EFieldIterationFlags::IncludeInterfaces); It; ++It)
 		{
 			MakeFunctionAction(*It);
 		}
@@ -319,6 +330,11 @@ void UAnimGraphNode_CallFunction::GetMenuActions(FBlueprintActionDatabaseRegistr
 	}
 }
 
+bool UAnimGraphNode_CallFunction::IsActionFilteredOut(FBlueprintActionFilter const& Filter)
+{
+	return CallFunctionPrototype == nullptr || InnerGraph == nullptr || CallFunctionPrototype->IsActionFilteredOut(Filter);
+}
+
 void UAnimGraphNode_CallFunction::GetRequiredExtensions(TArray<TSubclassOf<UAnimBlueprintExtension>>& OutExtensions) const
 {
 	OutExtensions.Add(UAnimBlueprintExtension_CallFunction::StaticClass());
diff --git a/Engine/Source/Editor/AnimGraph/Private/AnimGraphNode_CallFunction.h b/Engine/Source/Editor/AnimGraph/Private/AnimGraphNode_CallFunction.h
--- a/Engine/Source/Editor/AnimGraph/Private/AnimGraphNode_CallFunction.h
+++ b/Engine/Source/Editor/AnimGraph/Private/AnimGraphNode_CallFunction.h
@@ -50,6 +50,7 @@ private:
 	virtual FText GetMenuCategory() const override;
 	virtual void ExpandNode(FKismetCompilerContext& CompilerContext, UEdGraph* SourceGraph) override;
 	virtual void GetMenuActions(FBlueprintActionDatabaseRegistrar& ActionRegistrar) const override;
+	virtual bool IsActionFilteredOut(class FBlueprintActionFilter const& Filter) override;
 
 	// UAnimGraphNode_Base interface
 	virtual void GetRequiredExtensions(TArray<TSubclassOf<UAnimBlueprintExtension>>& OutExtensions) const override;
