Fixed compilation non-determinism in animation blueprints

Name-generation code in UAnimGraphNode_Base::GetDescriptiveCompiledName is taking the hash of a node's GUID, expecting it to be stable, but the calling code was passing the duplicated node, not the original, so the GUID varied each time. This fix passes the 'source' node instead for stability. Any additional duplicates are then handled as before by the NetNameMap incrementing the name's number.

#rnx

diff --git a/Engine/Source/Editor/AnimGraph/Private/AnimBlueprintCompiler.cpp b/Engine/Source/Editor/AnimGraph/Private/AnimBlueprintCompiler.cpp
--- a/Engine/Source/Editor/AnimGraph/Private/AnimBlueprintCompiler.cpp
+++ b/Engine/Source/Editor/AnimGraph/Private/AnimBlueprintCompiler.cpp
@@ -262,7 +262,8 @@ void FAnimBlueprintCompilerContext::ProcessAnimationNode(UAnimGraphNode_Base* Vi
 	}
 
 	// Create a property for the node
-	const FString NodeVariableName = ClassScopeNetNameMap.MakeValidName(VisualAnimNode);
+	UAnimGraphNode_Base* TrueSourceObject = MessageLog.FindSourceObjectTypeChecked<UAnimGraphNode_Base>(VisualAnimNode);
+	const FString NodeVariableName = ClassScopeNetNameMap.MakeValidName(TrueSourceObject);
 
 	const UAnimationGraphSchema* AnimGraphDefaultSchema = GetDefault<UAnimationGraphSchema>();
 
@@ -309,7 +310,6 @@ void FAnimBlueprintCompilerContext::ProcessAnimationNode(UAnimGraphNode_Base* Vi
 	AllocatedAnimNodeIndices.Add(VisualAnimNode, AllocatedIndex);
 	AllocatedPropertiesByIndex.Add(AllocatedIndex, NewProperty);
 
-	UAnimGraphNode_Base* TrueSourceObject = MessageLog.FindSourceObjectTypeChecked<UAnimGraphNode_Base>(VisualAnimNode);
 	SourceNodeToProcessedNodeMap.Add(TrueSourceObject, VisualAnimNode);
 
 	// Register the slightly more permanent debug information
