[Backout] - CL36699932

Causes issues in motion matching samples with parameters not updating correctly. Possibly due to blend stack graphs with same name.

Original CL Desc
-----------------------------------------------------------------
Fixed compilation non-determinism in animation blueprints

Name-generation code in UAnimGraphNode_Base::GetDescriptiveCompiledName is taking the hash of a node's GUID, expecting it to be stable, but the calling code was passing the duplicated node, not the original, so the GUID varied each time. This fix passes the 'source' node instead for stability. Any additional duplicates are then handled as before by the NetNameMap incrementing the name's number.

#rnx
#tests PIE, Anim BP Compile

diff --git a/Engine/Source/Editor/AnimGraph/Private/AnimBlueprintCompiler.cpp b/Engine/Source/Editor/AnimGraph/Private/AnimBlueprintCompiler.cpp
--- a/Engine/Source/Editor/AnimGraph/Private/AnimBlueprintCompiler.cpp
+++ b/Engine/Source/Editor/AnimGraph/Private/AnimBlueprintCompiler.cpp
@@ -262,8 +262,7 @@ void FAnimBlueprintCompilerContext::ProcessAnimationNode(UAnimGraphNode_Base* Vi
 	}
 
 	// Create a property for the node
-	UAnimGraphNode_Base* TrueSourceObject = MessageLog.FindSourceObjectTypeChecked<UAnimGraphNode_Base>(VisualAnimNode);
-	const FString NodeVariableName = ClassScopeNetNameMap.MakeValidName(TrueSourceObject);
+	const FString NodeVariableName = ClassScopeNetNameMap.MakeValidName(VisualAnimNode);
 
 	const UAnimationGraphSchema* AnimGraphDefaultSchema = GetDefault<UAnimationGraphSchema>();
 
@@ -310,6 +309,7 @@ void FAnimBlueprintCompilerContext::ProcessAnimationNode(UAnimGraphNode_Base* Vi
 	AllocatedAnimNodeIndices.Add(VisualAnimNode, AllocatedIndex);
 	AllocatedPropertiesByIndex.Add(AllocatedIndex, NewProperty);
 
+	UAnimGraphNode_Base* TrueSourceObject = MessageLog.FindSourceObjectTypeChecked<UAnimGraphNode_Base>(VisualAnimNode);
 	SourceNodeToProcessedNodeMap.Add(TrueSourceObject, VisualAnimNode);
 
 	// Register the slightly more permanent debug information
