Fix the graph ordering constraint when using anim node reference node

diff --git a/Engine/Source/Editor/AnimGraph/Private/AnimBlueprintExtension_Tag.cpp b/Engine/Source/Editor/AnimGraph/Private/AnimBlueprintExtension_Tag.cpp
--- a/Engine/Source/Editor/AnimGraph/Private/AnimBlueprintExtension_Tag.cpp
+++ b/Engine/Source/Editor/AnimGraph/Private/AnimBlueprintExtension_Tag.cpp
@@ -6,6 +6,7 @@
 #include "AnimGraphNode_LinkedAnimGraphBase.h"
 #include "IAnimBlueprintCompilationBracketContext.h"
 #include "IAnimBlueprintCompilationContext.h"
+#include "K2Node_CallFunction.h"
 
 #define LOCTEXT_NAMESPACE "UAnimBlueprintExtension_Tag"
 
@@ -13,6 +14,28 @@ void UAnimBlueprintExtension_Tag::HandleStartCompilingClass(const UClass* InClas
 {
 	TaggedNodes.Empty();
 	Subsystem.NodeIndices.Empty();
+	RequestedNodes.Empty();
+}
+
+void UAnimBlueprintExtension_Tag::HandlePostProcessAnimationNodes(TArrayView<UAnimGraphNode_Base*> InAnimNodes, IAnimBlueprintCompilationContext& InCompilationContext, IAnimBlueprintGeneratedClassCompiledData& OutCompiledData)
+{
+	for (const FNodeReferenceTag RefTag : RequestedNodes)
+	{
+		if (UAnimGraphNode_Base* const* TaggedNode = TaggedNodes.Find(RefTag.Tag))
+		{
+			if (UEdGraphPin* IndexPin = RefTag.CallFunction->FindPin(TEXT("Index"), EGPD_Input))
+			{
+				if (const int32* NodeIndex = InCompilationContext.GetAllocatedAnimNodeIndices().Find(*TaggedNode))
+				{
+					IndexPin->DefaultValue = FString::FromInt(*NodeIndex);
+				}
+			}
+		}
+		else
+		{
+			InCompilationContext.GetMessageLog().Error(*FText::Format(LOCTEXT("MissingTaggedNodeError", "@@ cannot find referenced node with tag '{0}', ensure it is present and connected to the graph"), FText::FromName(RefTag.Tag)).ToString(), this);
+		}
+	}
 }
 
 void UAnimBlueprintExtension_Tag::HandleFinishCompilingClass(const UClass* InClass, IAnimBlueprintCompilationBracketContext& InCompilationContext, IAnimBlueprintGeneratedClassCompiledData& OutCompiledData)
@@ -41,14 +64,10 @@ void UAnimBlueprintExtension_Tag::AddTaggedNode(UAnimGraphNode_Base* InNode, IAn
 	}
 }
 
-UAnimGraphNode_Base* UAnimBlueprintExtension_Tag::FindTaggedNode(FName InTag) const
+void UAnimBlueprintExtension_Tag::RequestTaggedNode(UK2Node_AnimNodeReference* InNode, const FName InTag, UK2Node_CallFunction* InCallFunction)
 {
-	if(UAnimGraphNode_Base*const* ExistingNode = TaggedNodes.Find(InTag))
-	{
-		return *ExistingNode;
-	}
-
-	return nullptr;
+	check(InNode && InTag != NAME_None && InCallFunction);
+	RequestedNodes.Add({ InNode, InCallFunction, InTag });
 }
 
 #undef LOCTEXT_NAMESPACE
\ No newline at end of file
diff --git a/Engine/Source/Editor/AnimGraph/Private/AnimBlueprintExtension_Tag.h b/Engine/Source/Editor/AnimGraph/Private/AnimBlueprintExtension_Tag.h
--- a/Engine/Source/Editor/AnimGraph/Private/AnimBlueprintExtension_Tag.h
+++ b/Engine/Source/Editor/AnimGraph/Private/AnimBlueprintExtension_Tag.h
@@ -5,8 +5,10 @@
 #include "CoreMinimal.h"
 #include "AnimBlueprintExtension.h"
 #include "Animation/AnimSubsystem_Tag.h"
+#include "K2Node_AnimNodeReference.h"
 #include "AnimBlueprintExtension_Tag.generated.h"
 
+class UK2Node_CallFunction;
 class UAnimGraphNode_Base;
 class IAnimBlueprintCompilationContext;
 
@@ -19,17 +21,28 @@ public:
 	// Add a tagged node
 	void AddTaggedNode(UAnimGraphNode_Base* InNode, IAnimBlueprintCompilationContext& InCompilationContext);
 
-	// Find a tagged node in this animation blueprint
-	UAnimGraphNode_Base* FindTaggedNode(FName InTag) const;
+	// Register the provided node to be hooked-up to the given tag during HandlePostProcessAnimationNodes, if the tag exists, error otherwise
+	void RequestTaggedNode(UK2Node_AnimNodeReference* InNode, const FName InTag, UK2Node_CallFunction* InCallFunction);
 
 private:
 	// UAnimBlueprintExtension interface
 	virtual void HandleStartCompilingClass(const UClass* InClass, IAnimBlueprintCompilationBracketContext& InCompilationContext, IAnimBlueprintGeneratedClassCompiledData& OutCompiledData) override;
+	virtual void HandlePostProcessAnimationNodes(TArrayView<UAnimGraphNode_Base*> InAnimNodes, IAnimBlueprintCompilationContext& InCompilationContext, IAnimBlueprintGeneratedClassCompiledData& OutCompiledData) override;
 	virtual void HandleFinishCompilingClass(const UClass* InClass, IAnimBlueprintCompilationBracketContext& InCompilationContext, IAnimBlueprintGeneratedClassCompiledData& OutCompiledData) override;
 
 private:
 	// Map of tag -> tagged nodes
 	TMap<FName, UAnimGraphNode_Base*> TaggedNodes;
+
+	struct FNodeReferenceTag
+	{
+		UK2Node_AnimNodeReference* Node = nullptr;
+		UK2Node_CallFunction* CallFunction = nullptr;
+		FName Tag;
+	};
+
+	// Array of nodes to process during HandlePostProcessAnimationNodes
+	TArray<FNodeReferenceTag> RequestedNodes;
 	
 	UPROPERTY()
 	FAnimSubsystem_Tag Subsystem;
diff --git a/Engine/Source/Editor/AnimGraph/Private/K2Node_AnimNodeReference.cpp b/Engine/Source/Editor/AnimGraph/Private/K2Node_AnimNodeReference.cpp
--- a/Engine/Source/Editor/AnimGraph/Private/K2Node_AnimNodeReference.cpp
+++ b/Engine/Source/Editor/AnimGraph/Private/K2Node_AnimNodeReference.cpp
@@ -39,13 +39,6 @@ struct FSearchTagDataPair;
 
 void UK2Node_AnimNodeReference::ExpandNode(FKismetCompilerContext& InCompilerContext, UEdGraph* InSourceGraph)
 {
-	UAnimGraphNode_Base* LabelledNode = nullptr;
-	UAnimBlueprintExtension_Tag* Extension = UAnimBlueprintExtension::FindExtension<UAnimBlueprintExtension_Tag>(Cast<UAnimBlueprint>(GetBlueprint()));
-	if(Extension)
-	{
-		LabelledNode = Extension->FindTaggedNode(Tag);
-	}
-
 	// Expand to function call
 	UK2Node_CallFunction* CallFunctionNode = InCompilerContext.SpawnIntermediateNode<UK2Node_CallFunction>(this, InSourceGraph);
 	CallFunctionNode->SetFromFunction(GetDefault<UAnimExecutionContextLibrary>()->FindFunctionChecked(GET_FUNCTION_NAME_CHECKED(UAnimExecutionContextLibrary, GetAnimNodeReference)));
@@ -53,22 +46,15 @@ void UK2Node_AnimNodeReference::ExpandNode(FKismetCompilerContext& InCompilerCon
 
 	UEdGraphPin* ValuePin = FindPinChecked(TEXT("Value"), EGPD_Output);
 
-	if(UEdGraphPin* IndexPin = CallFunctionNode->FindPin(TEXT("Index"), EGPD_Input))
-	{
-		FAnimBlueprintCompilerContext& AnimBlueprintCompilerContext = static_cast<FAnimBlueprintCompilerContext&>(InCompilerContext);
-
-		const int32 NodeIndex = LabelledNode ? AnimBlueprintCompilerContext.GetAllocationIndexOfNode(LabelledNode) : INDEX_NONE;
-		IndexPin->DefaultValue = FString::FromInt(NodeIndex);
-	}
-
-	if(UEdGraphPin* ReturnValuePin = CallFunctionNode->FindPin(TEXT("ReturnValue"), EGPD_Output))
+	if (UEdGraphPin* ReturnValuePin = CallFunctionNode->FindPin(TEXT("ReturnValue"), EGPD_Output))
 	{
 		InCompilerContext.MovePinLinksToIntermediate(*ValuePin, *ReturnValuePin);
 	}
 
-	if(LabelledNode == nullptr)
+	UAnimBlueprintExtension_Tag* Extension = UAnimBlueprintExtension::FindExtension<UAnimBlueprintExtension_Tag>(Cast<UAnimBlueprint>(GetBlueprint()));
+	if (Extension)
 	{
-		InCompilerContext.MessageLog.Error(*FText::Format(LOCTEXT("MissingTaggedNodeError", "@@ cannot find referenced node with tag '{0}', ensure it is present and connected to the graph"), FText::FromName(Tag)).ToString(), this);
+		Extension->RequestTaggedNode(this, Tag, CallFunctionNode);
 	}
 }
 
