Fix the graph ordering constraint when using anim node reference node

diff --git a/Engine/Source/Editor/AnimGraph/Private/AnimBlueprintCompilationContext.cpp b/Engine/Source/Editor/AnimGraph/Private/AnimBlueprintCompilationContext.cpp
--- a/Engine/Source/Editor/AnimGraph/Private/AnimBlueprintCompilationContext.cpp
+++ b/Engine/Source/Editor/AnimGraph/Private/AnimBlueprintCompilationContext.cpp
@@ -158,3 +158,8 @@ const FKismetCompilerOptions& FAnimBlueprintPostExpansionStepContext::GetCompile
 {
 	return CompilerContext->CompileOptions;
 }
+
+const TMap<UAnimGraphNode_Base*, int32>& FAnimBlueprintPostExpansionStepContext::GetAllocatedAnimNodeIndicesImpl() const
+{
+	return CompilerContext->AllocatedAnimNodeIndices;
+}
\ No newline at end of file
diff --git a/Engine/Source/Editor/AnimGraph/Private/AnimBlueprintCompilationContext.h b/Engine/Source/Editor/AnimGraph/Private/AnimBlueprintCompilationContext.h
--- a/Engine/Source/Editor/AnimGraph/Private/AnimBlueprintCompilationContext.h
+++ b/Engine/Source/Editor/AnimGraph/Private/AnimBlueprintCompilationContext.h
@@ -145,6 +145,7 @@ private:
 	virtual FCompilerResultsLog& GetMessageLogImpl() const override;
 	virtual UEdGraph* GetConsolidatedEventGraphImpl() const override;
 	virtual const FKismetCompilerOptions& GetCompileOptionsImpl() const override;
+	virtual const TMap<UAnimGraphNode_Base*, int32>& GetAllocatedAnimNodeIndicesImpl() const override;
 
 	FAnimBlueprintCompilerContext* CompilerContext;
 };
diff --git a/Engine/Source/Editor/AnimGraph/Private/AnimBlueprintExtension_Tag.cpp b/Engine/Source/Editor/AnimGraph/Private/AnimBlueprintExtension_Tag.cpp
--- a/Engine/Source/Editor/AnimGraph/Private/AnimBlueprintExtension_Tag.cpp
+++ b/Engine/Source/Editor/AnimGraph/Private/AnimBlueprintExtension_Tag.cpp
@@ -6,6 +6,8 @@
 #include "AnimGraphNode_LinkedAnimGraphBase.h"
 #include "IAnimBlueprintCompilationBracketContext.h"
 #include "IAnimBlueprintCompilationContext.h"
+#include "K2Node_CallFunction.h"
+#include "IAnimBlueprintPostExpansionStepContext.h"
 
 #define LOCTEXT_NAMESPACE "UAnimBlueprintExtension_Tag"
 
@@ -13,13 +15,39 @@ void UAnimBlueprintExtension_Tag::HandleStartCompilingClass(const UClass* InClas
 {
 	TaggedNodes.Empty();
 	Subsystem.NodeIndices.Empty();
+	RequestedNodes.Empty();
+}
+
+void UAnimBlueprintExtension_Tag::HandlePostExpansionStep(const UEdGraph* InGraph, IAnimBlueprintPostExpansionStepContext& InCompilationContext, IAnimBlueprintGeneratedClassCompiledData& OutCompiledData)
+{
+	for (int32 Index = RequestedNodes.Num() - 1; Index >= 0; --Index)
+	{
+		const FNodeReferenceTag& RefTag = RequestedNodes[Index];
+		
+		if (UAnimGraphNode_Base* const* TaggedNode = TaggedNodes.Find(RefTag.Tag))
+		{
+			if (UEdGraphPin* IndexPin = RefTag.CallFunction->FindPin(TEXT("Index"), EGPD_Input))
+			{
+				if (const int32* NodeIndex = InCompilationContext.GetAllocatedAnimNodeIndices().Find(*TaggedNode))
+				{
+					IndexPin->DefaultValue = FString::FromInt(*NodeIndex);
+					RequestedNodes.RemoveAt(Index, EAllowShrinking::No);
+				}
+			}
+		}
+	}
 }
 
 void UAnimBlueprintExtension_Tag::HandleFinishCompilingClass(const UClass* InClass, IAnimBlueprintCompilationBracketContext& InCompilationContext, IAnimBlueprintGeneratedClassCompiledData& OutCompiledData)
 {
-	for(const TPair<FName, UAnimGraphNode_Base*>& TaggedNodePair : TaggedNodes)
+	for (const FNodeReferenceTag& RefTag : RequestedNodes)
 	{
-		if(const int32* IndexPtr = InCompilationContext.GetAllocatedAnimNodeIndices().Find(TaggedNodePair.Value))
+		InCompilationContext.GetMessageLog().Error(*FText::Format(LOCTEXT("MissingTaggedNodeError", "@@ cannot find referenced node with tag '{0}', ensure it is present and connected to the graph"), FText::FromName(RefTag.Tag)).ToString(), this);
+	}
+
+	for (const TPair<FName, UAnimGraphNode_Base*>& TaggedNodePair : TaggedNodes)
+	{
+		if (const int32* IndexPtr = InCompilationContext.GetAllocatedAnimNodeIndices().Find(TaggedNodePair.Value))
 		{
 			Subsystem.NodeIndices.Add(TaggedNodePair.Value->GetTag(), *IndexPtr);
 		}
@@ -28,9 +56,9 @@ void UAnimBlueprintExtension_Tag::HandleFinishCompilingClass(const UClass* InCla
 
 void UAnimBlueprintExtension_Tag::AddTaggedNode(UAnimGraphNode_Base* InNode, IAnimBlueprintCompilationContext& InCompilationContext)
 {
-	if(InNode->GetTag() != NAME_None)
+	if (InNode->GetTag() != NAME_None)
 	{
-		if(UAnimGraphNode_Base** ExistingNode = TaggedNodes.Find(InNode->GetTag()))
+		if (UAnimGraphNode_Base** ExistingNode = TaggedNodes.Find(InNode->GetTag()))
 		{
 			InCompilationContext.GetMessageLog().Error(*FText::Format(LOCTEXT("DuplicateLabelError", "Nodes @@ and @@ have the same reference tag '{0}'"), FText::FromName(InNode->GetTag())).ToString(), InNode, *ExistingNode);
 		}
@@ -41,14 +69,10 @@ void UAnimBlueprintExtension_Tag::AddTaggedNode(UAnimGraphNode_Base* InNode, IAn
 	}
 }
 
-UAnimGraphNode_Base* UAnimBlueprintExtension_Tag::FindTaggedNode(FName InTag) const
+void UAnimBlueprintExtension_Tag::RequestTaggedNode(UK2Node_AnimNodeReference* InNode, const FName InTag, UK2Node_CallFunction* InCallFunction)
 {
-	if(UAnimGraphNode_Base*const* ExistingNode = TaggedNodes.Find(InTag))
-	{
-		return *ExistingNode;
-	}
-
-	return nullptr;
+	check(InNode && InTag != NAME_None && InCallFunction);
+	RequestedNodes.Add({ InNode, InCallFunction, InTag });
 }
 
 #undef LOCTEXT_NAMESPACE
\ No newline at end of file
diff --git a/Engine/Source/Editor/AnimGraph/Private/AnimBlueprintExtension_Tag.h b/Engine/Source/Editor/AnimGraph/Private/AnimBlueprintExtension_Tag.h
--- a/Engine/Source/Editor/AnimGraph/Private/AnimBlueprintExtension_Tag.h
+++ b/Engine/Source/Editor/AnimGraph/Private/AnimBlueprintExtension_Tag.h
@@ -5,8 +5,10 @@
 #include "CoreMinimal.h"
 #include "AnimBlueprintExtension.h"
 #include "Animation/AnimSubsystem_Tag.h"
+#include "K2Node_AnimNodeReference.h"
 #include "AnimBlueprintExtension_Tag.generated.h"
 
+class UK2Node_CallFunction;
 class UAnimGraphNode_Base;
 class IAnimBlueprintCompilationContext;
 
@@ -19,18 +21,29 @@ public:
 	// Add a tagged node
 	void AddTaggedNode(UAnimGraphNode_Base* InNode, IAnimBlueprintCompilationContext& InCompilationContext);
 
-	// Find a tagged node in this animation blueprint
-	UAnimGraphNode_Base* FindTaggedNode(FName InTag) const;
+	// Register the provided node to be hooked-up to the given tag during HandlePostProcessAnimationNodes, if the tag exists, error otherwise
+	void RequestTaggedNode(UK2Node_AnimNodeReference* InNode, const FName InTag, UK2Node_CallFunction* InCallFunction);
 
 private:
 	// UAnimBlueprintExtension interface
 	virtual void HandleStartCompilingClass(const UClass* InClass, IAnimBlueprintCompilationBracketContext& InCompilationContext, IAnimBlueprintGeneratedClassCompiledData& OutCompiledData) override;
+	virtual void HandlePostExpansionStep(const UEdGraph* InGraph, IAnimBlueprintPostExpansionStepContext& InCompilationContext, IAnimBlueprintGeneratedClassCompiledData& OutCompiledData) override;
 	virtual void HandleFinishCompilingClass(const UClass* InClass, IAnimBlueprintCompilationBracketContext& InCompilationContext, IAnimBlueprintGeneratedClassCompiledData& OutCompiledData) override;
 
 private:
 	// Map of tag -> tagged nodes
 	TMap<FName, UAnimGraphNode_Base*> TaggedNodes;
-	
+
+	struct FNodeReferenceTag
+	{
+		UK2Node_AnimNodeReference* Node = nullptr;
+		UK2Node_CallFunction* CallFunction = nullptr;
+		FName Tag;
+	};
+
+	// Array of nodes to process during HandlePostProcessAnimationNodes
+	TArray<FNodeReferenceTag> RequestedNodes;
+
 	UPROPERTY()
 	FAnimSubsystem_Tag Subsystem;
 };
\ No newline at end of file
diff --git a/Engine/Source/Editor/AnimGraph/Private/K2Node_AnimNodeReference.cpp b/Engine/Source/Editor/AnimGraph/Private/K2Node_AnimNodeReference.cpp
--- a/Engine/Source/Editor/AnimGraph/Private/K2Node_AnimNodeReference.cpp
+++ b/Engine/Source/Editor/AnimGraph/Private/K2Node_AnimNodeReference.cpp
@@ -39,13 +39,6 @@ struct FSearchTagDataPair;
 
 void UK2Node_AnimNodeReference::ExpandNode(FKismetCompilerContext& InCompilerContext, UEdGraph* InSourceGraph)
 {
-	UAnimGraphNode_Base* LabelledNode = nullptr;
-	UAnimBlueprintExtension_Tag* Extension = UAnimBlueprintExtension::FindExtension<UAnimBlueprintExtension_Tag>(Cast<UAnimBlueprint>(GetBlueprint()));
-	if(Extension)
-	{
-		LabelledNode = Extension->FindTaggedNode(Tag);
-	}
-
 	// Expand to function call
 	UK2Node_CallFunction* CallFunctionNode = InCompilerContext.SpawnIntermediateNode<UK2Node_CallFunction>(this, InSourceGraph);
 	CallFunctionNode->SetFromFunction(GetDefault<UAnimExecutionContextLibrary>()->FindFunctionChecked(GET_FUNCTION_NAME_CHECKED(UAnimExecutionContextLibrary, GetAnimNodeReference)));
@@ -53,22 +46,15 @@ void UK2Node_AnimNodeReference::ExpandNode(FKismetCompilerContext& InCompilerCon
 
 	UEdGraphPin* ValuePin = FindPinChecked(TEXT("Value"), EGPD_Output);
 
-	if(UEdGraphPin* IndexPin = CallFunctionNode->FindPin(TEXT("Index"), EGPD_Input))
-	{
-		FAnimBlueprintCompilerContext& AnimBlueprintCompilerContext = static_cast<FAnimBlueprintCompilerContext&>(InCompilerContext);
-
-		const int32 NodeIndex = LabelledNode ? AnimBlueprintCompilerContext.GetAllocationIndexOfNode(LabelledNode) : INDEX_NONE;
-		IndexPin->DefaultValue = FString::FromInt(NodeIndex);
-	}
-
-	if(UEdGraphPin* ReturnValuePin = CallFunctionNode->FindPin(TEXT("ReturnValue"), EGPD_Output))
+	if (UEdGraphPin* ReturnValuePin = CallFunctionNode->FindPin(TEXT("ReturnValue"), EGPD_Output))
 	{
 		InCompilerContext.MovePinLinksToIntermediate(*ValuePin, *ReturnValuePin);
 	}
 
-	if(LabelledNode == nullptr)
+	UAnimBlueprintExtension_Tag* Extension = UAnimBlueprintExtension::FindExtension<UAnimBlueprintExtension_Tag>(Cast<UAnimBlueprint>(GetBlueprint()));
+	if (Extension)
 	{
-		InCompilerContext.MessageLog.Error(*FText::Format(LOCTEXT("MissingTaggedNodeError", "@@ cannot find referenced node with tag '{0}', ensure it is present and connected to the graph"), FText::FromName(Tag)).ToString(), this);
+		Extension->RequestTaggedNode(this, Tag, CallFunctionNode);
 	}
 }
 
diff --git a/Engine/Source/Editor/AnimGraph/Public/IAnimBlueprintPostExpansionStepContext.h b/Engine/Source/Editor/AnimGraph/Public/IAnimBlueprintPostExpansionStepContext.h
--- a/Engine/Source/Editor/AnimGraph/Public/IAnimBlueprintPostExpansionStepContext.h
+++ b/Engine/Source/Editor/AnimGraph/Public/IAnimBlueprintPostExpansionStepContext.h
@@ -6,6 +6,7 @@
 
 class FCompilerResultsLog;
 class UEdGraph;
+class UAnimGraphNode_Base;
 struct FKismetCompilerOptions;
 
 // Interface passed to PostExpansionStep delegate
@@ -17,6 +18,9 @@ public:
 	// Get the message log for the current compilation
 	FCompilerResultsLog& GetMessageLog() const { return GetMessageLogImpl(); }
 
+	// Index of the nodes (must match up with the runtime discovery process of nodes, which runs thru the property chain)
+	const TMap<UAnimGraphNode_Base*, int32>& GetAllocatedAnimNodeIndices() const { return GetAllocatedAnimNodeIndicesImpl(); }
+
 	// Get the consolidated uber graph during compilation
 	UEdGraph* GetConsolidatedEventGraph() const { return GetConsolidatedEventGraphImpl(); }
 
@@ -27,6 +31,8 @@ protected:
 	// Get the message log for the current compilation
 	virtual FCompilerResultsLog& GetMessageLogImpl() const = 0;
 
+	virtual const TMap<UAnimGraphNode_Base*, int32>& GetAllocatedAnimNodeIndicesImpl() const = 0;
+
 	// Get the consolidated uber graph during compilation
 	virtual UEdGraph* GetConsolidatedEventGraphImpl() const = 0;
 
