Anim State Machine: Prevent renaming states to 'None'

'None' is an invalid name for a graph, which the state wraps. This avoids a check() when setting a graph name to 'None'.

diff --git a/Engine/Source/Editor/AnimGraph/Private/AnimStateNodeBase.cpp b/Engine/Source/Editor/AnimGraph/Private/AnimStateNodeBase.cpp
--- a/Engine/Source/Editor/AnimGraph/Private/AnimStateNodeBase.cpp
+++ b/Engine/Source/Editor/AnimGraph/Private/AnimStateNodeBase.cpp
@@ -23,6 +23,8 @@ public:
 	FAnimStateNodeNameValidator(const UAnimStateNodeBase* InStateNode)
 		: FStringSetNameValidator(FString())
 	{
+		Names.Add(TEXT("None"));
+
 		TArray<UAnimStateNodeBase*> Nodes;
 		UAnimationStateMachineGraph* StateMachine = CastChecked<UAnimationStateMachineGraph>(InStateNode->GetOuter());
 
