Anim State Machine: Prevent renaming state machines to 'None'

Similar to CL 42014078, but for machines

diff --git a/Engine/Source/Editor/AnimGraph/Private/AnimGraphNode_StateMachineBase.cpp b/Engine/Source/Editor/AnimGraph/Private/AnimGraphNode_StateMachineBase.cpp
--- a/Engine/Source/Editor/AnimGraph/Private/AnimGraphNode_StateMachineBase.cpp
+++ b/Engine/Source/Editor/AnimGraph/Private/AnimGraphNode_StateMachineBase.cpp
@@ -41,6 +41,8 @@ public:
 	FAnimStateMachineNodeNameValidator(const UAnimGraphNode_StateMachineBase* InStateMachineNode)
 		: FStringSetNameValidator(FString())
 	{
+		Names.Add(TEXT("None"));
+
 		TArray<UAnimGraphNode_StateMachineBase*> Nodes;
 
 		UAnimationGraph* StateMachine = CastChecked<UAnimationGraph>(InStateMachineNode->GetOuter());
