Anim State Machine: Fix copy-paste of transition nodes creating duplicate GUIDs/assets

diff --git a/Engine/Source/Editor/AnimGraph/Private/AnimStateNodeBase.cpp b/Engine/Source/Editor/AnimGraph/Private/AnimStateNodeBase.cpp
--- a/Engine/Source/Editor/AnimGraph/Private/AnimStateNodeBase.cpp
+++ b/Engine/Source/Editor/AnimGraph/Private/AnimStateNodeBase.cpp
@@ -70,7 +70,7 @@ UAnimStateNodeBase::UAnimStateNodeBase(const FObjectInitializer& ObjectInitializ
 void UAnimStateNodeBase::PostPasteNode()
 {
 	Super::PostPasteNode();
-
+	
 	for(UEdGraph* SubGraph : GetSubGraphs())
 	{
 		if(SubGraph)
@@ -83,9 +83,12 @@ void UAnimStateNodeBase::PostPasteNode()
 				ParentGraph->SubGraphs.Add(SubGraph);
 			}
 
-			//@TODO: CONDUIT: Merge conflict - May no longer be necessary due to other changes?
-	//		FBlueprintEditorUtils::RenameGraphWithSuggestion(SubGraph, NameValidator, GetDesiredNewNodeName());
-			//@ENDTODO
+			for (UEdGraphNode* SubGraphNode : SubGraph->Nodes)
+			{
+				SubGraphNode->CreateNewGuid();
+				SubGraphNode->PostPasteNode();
+				SubGraphNode->ReconstructNode();
+			}
 
 			// Restore transactional flag that is lost during copy/paste process
 			SubGraph->SetFlags(RF_Transactional);
diff --git a/Engine/Source/Editor/AnimGraph/Private/AnimStateTransitionNode.cpp b/Engine/Source/Editor/AnimGraph/Private/AnimStateTransitionNode.cpp
--- a/Engine/Source/Editor/AnimGraph/Private/AnimStateTransitionNode.cpp
+++ b/Engine/Source/Editor/AnimGraph/Private/AnimStateTransitionNode.cpp
@@ -210,7 +210,7 @@ void UAnimStateTransitionNode::PostPasteNode()
 		GraphNode->PostPasteNode();
 		GraphNode->ReconstructNode();
 	}
-
+	
 	if(CustomTransitionGraph)
 	{
 		// Needs to be added to the parent graph
@@ -223,6 +223,13 @@ void UAnimStateTransitionNode::PostPasteNode()
 
 		// Transactional flag is lost in copy/paste, restore it.
 		CustomTransitionGraph->SetFlags(RF_Transactional);
+
+		for (UEdGraphNode* TransitionGraphNode : CustomTransitionGraph->Nodes)
+		{
+			TransitionGraphNode->CreateNewGuid();
+			TransitionGraphNode->PostPasteNode();
+			TransitionGraphNode->ReconstructNode();
+		}
 	}
 
 	ValidateBlendProfile();
