Anim State Machine: Fix incorrect duplicate GUID warning when a transition rule is shared

Only add the warning if the GUID is found but the node differs (as nodes/graphs can be referenced by multiple others)

diff --git a/Engine/Source/Editor/AnimGraph/Private/AnimBlueprintCompiler.cpp b/Engine/Source/Editor/AnimGraph/Private/AnimBlueprintCompiler.cpp
--- a/Engine/Source/Editor/AnimGraph/Private/AnimBlueprintCompiler.cpp
+++ b/Engine/Source/Editor/AnimGraph/Private/AnimBlueprintCompiler.cpp
@@ -85,13 +85,13 @@ FAnimBlueprintCompilerContext::FAnimBlueprintCompilerContext(UAnimBlueprint* Sou
 		TArray<UEdGraph*> ChildGraphs;
 		ChildGraphs.Reserve(20);
 
-		TSet<FGuid> NodeGuids;
-		NodeGuids.Reserve(200);
+		TMap<FGuid, UEdGraphNode*> GuidMap;
+		GuidMap.Reserve(200);
 
 		// Tracking to see if we need to warn for deterministic cooking
 		bool bNodeGuidsRegenerated = false;
 
-		auto CheckGraph = [&bNodeGuidsRegenerated, &NodeGuids, &ChildGraphs](UEdGraph* InGraph)
+		auto CheckGraph = [&bNodeGuidsRegenerated, &GuidMap, &ChildGraphs](UEdGraph* InGraph)
 		{
 			if (AnimationEditorUtils::IsAnimGraph(InGraph))
 			{
@@ -109,15 +109,18 @@ FAnimBlueprintCompilerContext::FAnimBlueprintCompilerContext(UAnimBlueprint* Sou
 					{
 						if (Node)
 						{
-							if (NodeGuids.Contains(Node->NodeGuid))
+							if (UEdGraphNode** FoundNode = GuidMap.Find(Node->NodeGuid))
 							{
-								bNodeGuidsRegenerated = true;
-							
-								Node->CreateNewGuid(); // GUID is already being used, create a new one.
+								if (*FoundNode != Node)
+								{
+									bNodeGuidsRegenerated = true;
+								
+									Node->CreateNewGuid(); // GUID is already being used, create a new one.
+								}
 							}
 							else
 							{
-								NodeGuids.Add(Node->NodeGuid);
+								GuidMap.Add(Node->NodeGuid, Node);
 							}
 						}
 					}
@@ -135,7 +138,7 @@ FAnimBlueprintCompilerContext::FAnimBlueprintCompilerContext(UAnimBlueprint* Sou
 			for(UEdGraph* Graph : InterfaceDesc.Graphs)
 			{
 				CheckGraph(Graph);
-			}
+			} 
 		}
 
 		if(bNodeGuidsRegenerated)
