Anim BP: Fix cook non-determinism with generated 'mutables' struct

diff --git a/Engine/Source/Editor/AnimGraph/Private/AnimBlueprintCompiler.cpp b/Engine/Source/Editor/AnimGraph/Private/AnimBlueprintCompiler.cpp
--- a/Engine/Source/Editor/AnimGraph/Private/AnimBlueprintCompiler.cpp
+++ b/Engine/Source/Editor/AnimGraph/Private/AnimBlueprintCompiler.cpp
@@ -2037,6 +2037,11 @@ void FAnimBlueprintCompilerContext::ProcessFoldedPropertyRecords()
 		NewMutablesProperty = CastFieldChecked<FStructProperty>(CreateVariable(MutablesStructName, MutablesPinType));
 		NewMutablesProperty->SetMetaData(TEXT("BlueprintCompilerGeneratedDefaults"), TEXT("true"));
 	}
+	else
+	{
+		// No more references to this struct, as it was not used to generate the property above, so mark it as garbage to ensure it doesnt get saved into the package
+		NewAnimBlueprintMutables->MarkAsGarbage();
+	}
 
 	// Set up per-node mappings
 	NewAnimBlueprintClass->AnimNodeData.Empty();
