Fix look at edit mode locator gizmo not updating location when changes are made in the details panel

diff --git a/Engine/Source/Editor/AnimGraph/Private/AnimGraphNode_LookAt.cpp b/Engine/Source/Editor/AnimGraph/Private/AnimGraphNode_LookAt.cpp
--- a/Engine/Source/Editor/AnimGraph/Private/AnimGraphNode_LookAt.cpp
+++ b/Engine/Source/Editor/AnimGraph/Private/AnimGraphNode_LookAt.cpp
@@ -142,4 +142,20 @@ FEditorModeID UAnimGraphNode_LookAt::GetEditorMode() const
 	return AnimNodeEditModes::LookAt;
 }
 
+void UAnimGraphNode_LookAt::CopyNodeDataToPreviewNode(FAnimNode_Base* InPreviewNode)
+{
+	FAnimNode_LookAt* LookAt = static_cast<FAnimNode_LookAt*>(InPreviewNode);
+
+	// copies Pin values from the internal node to get data which are not compiled yet
+	LookAt->LookAtLocation = Node.LookAtLocation;
+}
+
+void UAnimGraphNode_LookAt::CopyPinDefaultsToNodeData(UEdGraphPin* InPin)
+{
+	if (InPin->GetName() == GET_MEMBER_NAME_STRING_CHECKED(FAnimNode_LookAt, LookAtLocation))
+	{
+		GetDefaultValue(GET_MEMBER_NAME_STRING_CHECKED(FAnimNode_LookAt, LookAtLocation), Node.LookAtLocation);
+	}
+}
+
 #undef LOCTEXT_NAMESPACE
diff --git a/Engine/Source/Editor/AnimGraph/Public/AnimGraphNode_LookAt.h b/Engine/Source/Editor/AnimGraph/Public/AnimGraphNode_LookAt.h
--- a/Engine/Source/Editor/AnimGraph/Public/AnimGraphNode_LookAt.h
+++ b/Engine/Source/Editor/AnimGraph/Public/AnimGraphNode_LookAt.h
@@ -25,6 +25,8 @@ public:
 	virtual FText GetNodeTitle(ENodeTitleType::Type TitleType) const override;
 	virtual FText GetTooltipText() const override;
 	virtual FEditorModeID GetEditorMode() const override;
+	virtual void CopyNodeDataToPreviewNode(FAnimNode_Base* InPreviewNode) override;
+	virtual void CopyPinDefaultsToNodeData(UEdGraphPin* InPin) override;
 	// End of UEdGraphNode interface
 
 	// UAnimGraphNode_SkeletalControlBase interface
