Updated headers using UnrealCodeFixup to make sure dllstorage is on methods/staticvars instead of types. Part 4/n

diff --git a/Engine/Source/Editor/AIGraph/Public/AIGraphEditor.h b/Engine/Source/Editor/AIGraph/Public/AIGraphEditor.h
--- a/Engine/Source/Editor/AIGraph/Public/AIGraphEditor.h
+++ b/Engine/Source/Editor/AIGraph/Public/AIGraphEditor.h
@@ -6,46 +6,48 @@
 #include "EditorUndoClient.h"
 #include "GraphEditor.h"
 
-class AIGRAPH_API FAIGraphEditor : public FEditorUndoClient
+#define UE_API AIGRAPH_API
+
+class FAIGraphEditor : public FEditorUndoClient
 {
 public:
-	FAIGraphEditor();
-	virtual ~FAIGraphEditor();
+	UE_API FAIGraphEditor();
+	UE_API virtual ~FAIGraphEditor();
 
-	FGraphPanelSelectionSet GetSelectedNodes() const;
-	virtual void OnSelectedNodesChanged(const TSet<class UObject*>& NewSelection);
+	UE_API FGraphPanelSelectionSet GetSelectedNodes() const;
+	UE_API virtual void OnSelectedNodesChanged(const TSet<class UObject*>& NewSelection);
 
 	//~ Begin FEditorUndoClient Interface
-	virtual void PostUndo(bool bSuccess) override;
-	virtual void PostRedo(bool bSuccess) override;
+	UE_API virtual void PostUndo(bool bSuccess) override;
+	UE_API virtual void PostRedo(bool bSuccess) override;
 	// End of FEditorUndoClient
 
-	void CreateCommandList();
+	UE_API void CreateCommandList();
 
 	// Delegates for graph editor commands
-	void SelectAllNodes();
-	bool CanSelectAllNodes() const;
-	void DeleteSelectedNodes();
-	bool CanDeleteNodes() const;
-	void DeleteSelectedDuplicatableNodes();
-	void CutSelectedNodes();
-	bool CanCutNodes() const;
-	void CopySelectedNodes();
-	bool CanCopyNodes() const;
-	void PasteNodes();
-	void PasteNodesHere(const FVector2D& Location);
-	bool CanPasteNodes() const;
-	void DuplicateNodes();
-	bool CanDuplicateNodes() const;
-
-	bool CanCreateComment() const;
-	void OnCreateComment();
-
-
-	virtual void OnClassListUpdated();
+	UE_API void SelectAllNodes();
+	UE_API bool CanSelectAllNodes() const;
+	UE_API void DeleteSelectedNodes();
+	UE_API bool CanDeleteNodes() const;
+	UE_API void DeleteSelectedDuplicatableNodes();
+	UE_API void CutSelectedNodes();
+	UE_API bool CanCutNodes() const;
+	UE_API void CopySelectedNodes();
+	UE_API bool CanCopyNodes() const;
+	UE_API void PasteNodes();
+	UE_API void PasteNodesHere(const FVector2D& Location);
+	UE_API bool CanPasteNodes() const;
+	UE_API void DuplicateNodes();
+	UE_API bool CanDuplicateNodes() const;
+
+	UE_API bool CanCreateComment() const;
+	UE_API void OnCreateComment();
+
+
+	UE_API virtual void OnClassListUpdated();
 
 protected:
-	virtual void FixupPastedNodes(const TSet<UEdGraphNode*>& NewPastedGraphNodes, const TMap<FGuid/*New*/, FGuid/*Old*/>& NewToOldNodeMapping);
+	UE_API virtual void FixupPastedNodes(const TSet<UEdGraphNode*>& NewPastedGraphNodes, const TMap<FGuid/*New*/, FGuid/*Old*/>& NewToOldNodeMapping);
 
 protected:
 
@@ -58,3 +60,5 @@ protected:
 	/** Handle to the registered OnClassListUpdated delegate */
 	FDelegateHandle OnClassListUpdatedDelegateHandle;
 };
+
+#undef UE_API
diff --git a/Engine/Source/Editor/AIGraph/Public/K2Node_AIMoveTo.h b/Engine/Source/Editor/AIGraph/Public/K2Node_AIMoveTo.h
--- a/Engine/Source/Editor/AIGraph/Public/K2Node_AIMoveTo.h
+++ b/Engine/Source/Editor/AIGraph/Public/K2Node_AIMoveTo.h
@@ -11,19 +11,23 @@
 
 #include "K2Node_AIMoveTo.generated.h"
 
+#define UE_API AIGRAPH_API
+
 class UObject;
 
-UCLASS()
-class AIGRAPH_API UK2Node_AIMoveTo : public UK2Node_BaseAsyncTask
+UCLASS(MinimalAPI)
+class UK2Node_AIMoveTo : public UK2Node_BaseAsyncTask
 {
 	GENERATED_UCLASS_BODY()
 
 	//~ Begin UEdGraphNode Interface
-	virtual FText GetTooltipText() const override;
-	virtual FText GetNodeTitle(ENodeTitleType::Type TitleType) const override;
+	UE_API virtual FText GetTooltipText() const override;
+	UE_API virtual FText GetNodeTitle(ENodeTitleType::Type TitleType) const override;
 	//~ End UEdGraphNode Interface
 
 	//~ Begin UK2Node Interface
-	virtual FText GetMenuCategory() const override;
+	UE_API virtual FText GetMenuCategory() const override;
 	//~ End UK2Node Interface
 };
+
+#undef UE_API
diff --git a/Engine/Source/Editor/AIGraph/Public/SGraphEditorActionMenuAI.h b/Engine/Source/Editor/AIGraph/Public/SGraphEditorActionMenuAI.h
--- a/Engine/Source/Editor/AIGraph/Public/SGraphEditorActionMenuAI.h
+++ b/Engine/Source/Editor/AIGraph/Public/SGraphEditorActionMenuAI.h
@@ -13,6 +13,8 @@
 #include "Widgets/DeclarativeSyntaxSupport.h"
 #include "Widgets/Layout/SBorder.h"
 
+#define UE_API AIGRAPH_API
+
 class SEditableTextBox;
 class SGraphActionMenu;
 class UAIGraphNode;
@@ -23,7 +25,7 @@ struct FGraphActionListBuilderBase;
 
 /////////////////////////////////////////////////////////////////////////////////////////////////
 
-class AIGRAPH_API SGraphEditorActionMenuAI : public SBorder
+class SGraphEditorActionMenuAI : public SBorder
 {
 public:
 	SLATE_BEGIN_ARGS(SGraphEditorActionMenuAI)
@@ -43,11 +45,11 @@ public:
 		SLATE_ARGUMENT( int32, SubNodeFlags)
 	SLATE_END_ARGS()
 
-	void Construct( const FArguments& InArgs );
+	UE_API void Construct( const FArguments& InArgs );
 
-	~SGraphEditorActionMenuAI();
+	UE_API ~SGraphEditorActionMenuAI();
 
-	TSharedRef<SEditableTextBox> GetFilterTextBox();
+	UE_API TSharedRef<SEditableTextBox> GetFilterTextBox();
 
 protected:
 	UEdGraph* GraphObj;
@@ -60,8 +62,10 @@ protected:
 	SGraphEditor::FActionMenuClosed OnClosedCallback;
 	TSharedPtr<SGraphActionMenu> GraphActionMenu;
 
-	void OnActionSelected( const TArray< TSharedPtr<FEdGraphSchemaAction> >& SelectedAction, ESelectInfo::Type InSelectionType );
+	UE_API void OnActionSelected( const TArray< TSharedPtr<FEdGraphSchemaAction> >& SelectedAction, ESelectInfo::Type InSelectionType );
 
 	/** Callback used to populate all actions list in SGraphActionMenu */
-	void CollectAllActions(FGraphActionListBuilderBase& OutAllActions);
+	UE_API void CollectAllActions(FGraphActionListBuilderBase& OutAllActions);
 };
+
+#undef UE_API
diff --git a/Engine/Source/Editor/AIGraph/Public/SGraphNodeAI.h b/Engine/Source/Editor/AIGraph/Public/SGraphNodeAI.h
--- a/Engine/Source/Editor/AIGraph/Public/SGraphNodeAI.h
+++ b/Engine/Source/Editor/AIGraph/Public/SGraphNodeAI.h
@@ -13,6 +13,8 @@
 #include "SGraphPin.h"
 #include "Editor/GraphEditor/Private/DragNode.h"
 
+#define UE_API AIGRAPH_API
+
 class SGraphPanel;
 class SToolTip;
 class UAIGraphNode;
@@ -33,65 +35,67 @@ protected:
 	typedef FDragNode Super;
 };
 
-class AIGRAPH_API SGraphNodeAI : public SGraphNode
+class SGraphNodeAI : public SGraphNode
 {
 public:
 	SLATE_BEGIN_ARGS(SGraphNodeAI){}
 	SLATE_END_ARGS()
 
-	void Construct(const FArguments& InArgs, UAIGraphNode* InNode);
+	UE_API void Construct(const FArguments& InArgs, UAIGraphNode* InNode);
 
 	//~ Begin SGraphNode Interface
-	virtual TSharedPtr<SToolTip> GetComplexTooltip() override;
-	virtual void OnDragEnter(const FGeometry& MyGeometry, const FDragDropEvent& DragDropEvent) override;
-	virtual FReply OnDragOver(const FGeometry& MyGeometry, const FDragDropEvent& DragDropEvent) override;
-	virtual FReply OnDrop(const FGeometry& MyGeometry, const FDragDropEvent& DragDropEvent) override;
-	virtual void OnDragLeave(const FDragDropEvent& DragDropEvent) override;
-	virtual FReply OnMouseMove(const FGeometry& SenderGeometry, const FPointerEvent& MouseEvent) override;
-	virtual void SetOwner(const TSharedRef<SGraphPanel>& OwnerPanel) override;
-	virtual void AddPin(const TSharedRef<SGraphPin>& PinToAdd) override;
+	UE_API virtual TSharedPtr<SToolTip> GetComplexTooltip() override;
+	UE_API virtual void OnDragEnter(const FGeometry& MyGeometry, const FDragDropEvent& DragDropEvent) override;
+	UE_API virtual FReply OnDragOver(const FGeometry& MyGeometry, const FDragDropEvent& DragDropEvent) override;
+	UE_API virtual FReply OnDrop(const FGeometry& MyGeometry, const FDragDropEvent& DragDropEvent) override;
+	UE_API virtual void OnDragLeave(const FDragDropEvent& DragDropEvent) override;
+	UE_API virtual FReply OnMouseMove(const FGeometry& SenderGeometry, const FPointerEvent& MouseEvent) override;
+	UE_API virtual void SetOwner(const TSharedRef<SGraphPanel>& OwnerPanel) override;
+	UE_API virtual void AddPin(const TSharedRef<SGraphPin>& PinToAdd) override;
 	//~ End SGraphNode Interface
 
 	/** handle mouse down on the node */
-	FReply OnMouseDown(const FGeometry& SenderGeometry, const FPointerEvent& MouseEvent);
+	UE_API FReply OnMouseDown(const FGeometry& SenderGeometry, const FPointerEvent& MouseEvent);
 
 	/** adds subnode widget inside current node */
-	virtual void AddSubNode(TSharedPtr<SGraphNode> SubNodeWidget);
+	UE_API virtual void AddSubNode(TSharedPtr<SGraphNode> SubNodeWidget);
 
 	/** gets decorator or service node if one is found under mouse cursor */
-	TSharedPtr<SGraphNode> GetSubNodeUnderCursor(const FGeometry& WidgetGeometry, const FPointerEvent& MouseEvent);
+	UE_API TSharedPtr<SGraphNode> GetSubNodeUnderCursor(const FGeometry& WidgetGeometry, const FPointerEvent& MouseEvent);
 
 	/** gets drag over marker visibility */
-	EVisibility GetDragOverMarkerVisibility() const;
+	UE_API EVisibility GetDragOverMarkerVisibility() const;
 
 	/** sets drag marker visible or collapsed on this node */
-	void SetDragMarker(bool bEnabled);
+	UE_API void SetDragMarker(bool bEnabled);
 
 protected:
 	TArray< TSharedPtr<SGraphNode> > SubNodes;
 
 	uint32 bDragMarkerVisible : 1;
 
-	virtual FText GetTitle() const;
-	virtual FText GetDescription() const;
-	virtual EVisibility GetDescriptionVisibility() const;
+	UE_API virtual FText GetTitle() const;
+	UE_API virtual FText GetDescription() const;
+	UE_API virtual EVisibility GetDescriptionVisibility() const;
 
-	virtual FText GetPreviewCornerText() const;
-	virtual const FSlateBrush* GetNameIcon() const;
+	UE_API virtual FText GetPreviewCornerText() const;
+	UE_API virtual const FSlateBrush* GetNameIcon() const;
 };
 
-class AIGRAPH_API SGraphPinAI : public SGraphPin
+class SGraphPinAI : public SGraphPin
 {
 public:
 	SLATE_BEGIN_ARGS(SGraphPinAI){}
 	SLATE_END_ARGS()
 
-	void Construct(const FArguments& InArgs, UEdGraphPin* InPin);
+	UE_API void Construct(const FArguments& InArgs, UEdGraphPin* InPin);
 protected:
 	//~ Begin SGraphPin Interface
-	virtual FSlateColor GetPinColor() const override;
-	virtual TSharedRef<SWidget>	GetDefaultValueWidget() override;
+	UE_API virtual FSlateColor GetPinColor() const override;
+	UE_API virtual TSharedRef<SWidget>	GetDefaultValueWidget() override;
 	//~ End SGraphPin Interface
 
-	const FSlateBrush* GetPinBorder() const;
+	UE_API const FSlateBrush* GetPinBorder() const;
 };
+
+#undef UE_API
diff --git a/Engine/Source/Editor/ActionableMessage/Public/ActionableMessageSubsystem.h b/Engine/Source/Editor/ActionableMessage/Public/ActionableMessageSubsystem.h
--- a/Engine/Source/Editor/ActionableMessage/Public/ActionableMessageSubsystem.h
+++ b/Engine/Source/Editor/ActionableMessage/Public/ActionableMessageSubsystem.h
@@ -7,6 +7,8 @@
 
 #include "ActionableMessageSubsystem.generated.h"
 
+#define UE_API ACTIONABLEMESSAGE_API
+
 struct FActionableMessage
 {
 	/**
@@ -40,8 +42,8 @@ FORCEINLINE bool operator!=(const FActionableMessage& LHS, const FActionableMess
 /**
  * Used to push messages that will be used by a viewport widget to present to the user.
  */
-UCLASS()
-class ACTIONABLEMESSAGE_API UActionableMessageSubsystem : public UWorldSubsystem
+UCLASS(MinimalAPI)
+class UActionableMessageSubsystem : public UWorldSubsystem
 {
 	GENERATED_BODY()
 
@@ -57,13 +59,13 @@ public:
 	 * @param InProvider The name of the entry.
 	 * @param InActionableMessage The content of the actionable message.
 	 */
-	void SetActionableMessage(FName InProvider, const FActionableMessage& InActionableMessage);
+	UE_API void SetActionableMessage(FName InProvider, const FActionableMessage& InActionableMessage);
 
 	/**
 	 * Removes an entry from the current actionable messages. 
 	 * @param InProvider The name of the entry to remove.
 	 */
-	void ClearActionableMessage(FName InProvider);
+	UE_API void ClearActionableMessage(FName InProvider);
 
 	/**
 	 * Returns the current stateID. Any modification in the current entries will update it. 
@@ -75,3 +77,5 @@ private:
 	TMap<FName, TSharedPtr<FActionableMessage>> ProviderActionableMessageMap;
 	uint32 StateID = 0;
 };
+
+#undef UE_API
diff --git a/Engine/Source/Editor/ActionableMessage/Public/SActionableMessageViewportWidget.h b/Engine/Source/Editor/ActionableMessage/Public/SActionableMessageViewportWidget.h
--- a/Engine/Source/Editor/ActionableMessage/Public/SActionableMessageViewportWidget.h
+++ b/Engine/Source/Editor/ActionableMessage/Public/SActionableMessageViewportWidget.h
@@ -9,31 +9,33 @@
 #include "Widgets/Views/STableRow.h"
 #include "Widgets/Views/SListView.h"
 
-class ACTIONABLEMESSAGE_API SActionableMessageEntry : public SCompoundWidget
+#define UE_API ACTIONABLEMESSAGE_API
+
+class SActionableMessageEntry : public SCompoundWidget
 {
 public:
 	SLATE_BEGIN_ARGS(SActionableMessageEntry) {}
 	SLATE_ARGUMENT(TSharedPtr<FActionableMessage>, ActionableMessage)
 	SLATE_END_ARGS()
 
-	void Construct(const FArguments& InArgs);
-	void SetActionableMessage(TSharedPtr<FActionableMessage> InActionableMessage);
+	UE_API void Construct(const FArguments& InArgs);
+	UE_API void SetActionableMessage(TSharedPtr<FActionableMessage> InActionableMessage);
 
 private:
 	TSharedPtr<FActionableMessage> ActionableMessage;
 	FOnClicked OnClicked;
 };
 
-class ACTIONABLEMESSAGE_API SActionableMessageViewportWidget : public SCompoundWidget
+class SActionableMessageViewportWidget : public SCompoundWidget
 {
 public:
 	SLATE_BEGIN_ARGS(SActionableMessageViewportWidget) {}
 	SLATE_END_ARGS()
 
-	void Construct(const FArguments& InArgs);
-	EVisibility GetVisibility();
+	UE_API void Construct(const FArguments& InArgs);
+	UE_API EVisibility GetVisibility();
 
-	TSharedRef<ITableRow> OnGenerateRow(TSharedPtr<FActionableMessage> InActionableMessage, const TSharedRef<STableViewBase>& OwnerTable);
+	UE_API TSharedRef<ITableRow> OnGenerateRow(TSharedPtr<FActionableMessage> InActionableMessage, const TSharedRef<STableViewBase>& OwnerTable);
 
 private:
 	bool IsExpanded() const { return bExpanded || bForceExpand;	}
@@ -47,3 +49,5 @@ private:
 	bool bExpanded = false;
 	bool bForceExpand = false;
 };
+
+#undef UE_API
diff --git a/Engine/Source/Editor/ActorPickerMode/Public/ActorPickerMode.h b/Engine/Source/Editor/ActorPickerMode/Public/ActorPickerMode.h
--- a/Engine/Source/Editor/ActorPickerMode/Public/ActorPickerMode.h
+++ b/Engine/Source/Editor/ActorPickerMode/Public/ActorPickerMode.h
@@ -5,6 +5,8 @@
 #include "CoreMinimal.h"
 #include "Modules/ModuleInterface.h"
 
+#define UE_API ACTORPICKERMODE_API
+
 class AActor;
 class FEditorModeTools;
 
@@ -15,12 +17,12 @@ DECLARE_DELEGATE_RetVal_OneParam(bool, FOnShouldFilterActor, const AActor*);
 /**
  * Actor picker mode module
  */
-class ACTORPICKERMODE_API FActorPickerModeModule : public IModuleInterface
+class FActorPickerModeModule : public IModuleInterface
 {
 public:
 	// IModuleInterface
-	virtual void StartupModule() override;
-	virtual void ShutdownModule() override;
+	UE_API virtual void StartupModule() override;
+	UE_API virtual void ShutdownModule() override;
 	// End of IModuleInterface
 
 	/** 
@@ -29,21 +31,23 @@ public:
 	 * @param	InOnShouldFilterActor	Delegate used to only allow particular actors (empty to accept all actors; works alongside InOnGetAllowedClasses)
 	 * @param	InOnActorSelected		Delegate to call when a valid actor is selected
 	 */
-	void BeginActorPickingMode(FOnGetAllowedClasses InOnGetAllowedClasses, FOnShouldFilterActor InOnShouldFilterActor, FOnActorSelected InOnActorSelected) const;
+	UE_API void BeginActorPickingMode(FOnGetAllowedClasses InOnGetAllowedClasses, FOnShouldFilterActor InOnShouldFilterActor, FOnActorSelected InOnActorSelected) const;
 
 	/** Exit actor picking mode */
-	void EndActorPickingMode() const;
+	UE_API void EndActorPickingMode() const;
 
 	/** @return Whether or not actor picking mode is currently active */
-	bool IsInActorPickingMode() const;
+	UE_API bool IsInActorPickingMode() const;
 	
 private:
 
 	/** Handler for when the application is deactivated. */
-	void OnApplicationDeactivated(const bool IsActive) const;
+	UE_API void OnApplicationDeactivated(const bool IsActive) const;
 
-	static FEditorModeTools* GetLevelEditorModeManager();
+	static UE_API FEditorModeTools* GetLevelEditorModeManager();
 
 private:
 	FDelegateHandle OnApplicationDeactivatedHandle;
 };
+
+#undef UE_API
diff --git a/Engine/Source/Editor/AddContentDialog/Public/FeaturePackContentSource.h b/Engine/Source/Editor/AddContentDialog/Public/FeaturePackContentSource.h
--- a/Engine/Source/Editor/AddContentDialog/Public/FeaturePackContentSource.h
+++ b/Engine/Source/Editor/AddContentDialog/Public/FeaturePackContentSource.h
@@ -16,6 +16,8 @@
 
 #include "FeaturePackContentSource.generated.h"
 
+#define UE_API ADDCONTENTDIALOG_API
+
 class FJsonValue;
 class FPakPlatformFile;
 class UObject;
@@ -142,32 +144,32 @@ class FPakPlatformFile;
 struct FSearchEntry;
 
 /** A content source which represents a content upack. */
-class ADDCONTENTDIALOG_API FFeaturePackContentSource : public IContentSource
+class FFeaturePackContentSource : public IContentSource
 {
 public:
-	FFeaturePackContentSource();
-	FFeaturePackContentSource(FString InFeaturePackPath);
+	UE_API FFeaturePackContentSource();
+	UE_API FFeaturePackContentSource(FString InFeaturePackPath);
 
-	virtual ~FFeaturePackContentSource();
+	UE_API virtual ~FFeaturePackContentSource();
 
-	virtual const TArray<FLocalizedText>& GetLocalizedNames() const override;
-	virtual const TArray<FLocalizedText>& GetLocalizedDescriptions() const override;
+	UE_API virtual const TArray<FLocalizedText>& GetLocalizedNames() const override;
+	UE_API virtual const TArray<FLocalizedText>& GetLocalizedDescriptions() const override;
 	
-	virtual const TArray<EContentSourceCategory>& GetCategories() const override;
-	virtual const TArray<FLocalizedText>& GetLocalizedAssetTypes() const override;
-	virtual const FString& GetSortKey() const override;
-	virtual const FString& GetClassTypesUsed() const override;
-	virtual const FString& GetIdent() const override;
-	virtual TSharedPtr<FImageData> GetIconData() const override;
-	virtual const TArray<TSharedPtr<FImageData>>& GetScreenshotData() const override;
-	const FString& GetFocusAssetName() const;
+	UE_API virtual const TArray<EContentSourceCategory>& GetCategories() const override;
+	UE_API virtual const TArray<FLocalizedText>& GetLocalizedAssetTypes() const override;
+	UE_API virtual const FString& GetSortKey() const override;
+	UE_API virtual const FString& GetClassTypesUsed() const override;
+	UE_API virtual const FString& GetIdent() const override;
+	UE_API virtual TSharedPtr<FImageData> GetIconData() const override;
+	UE_API virtual const TArray<TSharedPtr<FImageData>>& GetScreenshotData() const override;
+	UE_API const FString& GetFocusAssetName() const;
 
-	virtual bool InstallToProject(FString InstallPath) override;
+	UE_API virtual bool InstallToProject(FString InstallPath) override;
 
-	void InsertAdditionalFeaturePacks();
-	bool InsertAdditionalResources(TArray<FFeaturePackLevelSet> InAdditionalFeaturePacks,EFeaturePackDetailLevel RequiredLevel, const FString& InDestinationFolder,TArray<FString>& InFilesCopied);
+	UE_API void InsertAdditionalFeaturePacks();
+	UE_API bool InsertAdditionalResources(TArray<FFeaturePackLevelSet> InAdditionalFeaturePacks,EFeaturePackDetailLevel RequiredLevel, const FString& InDestinationFolder,TArray<FString>& InFilesCopied);
 
-	virtual bool IsDataValid() const override;
+	UE_API virtual bool IsDataValid() const override;
 			
 	/*
 	 * Copies the list of files specified in 'AdditionFilesToInclude' section in the config.ini of the feature pack.
@@ -177,7 +179,7 @@ public:
 	 * @param bContainsSource 	Set to true if the file list contains any source files
 	 * @returns true if config file was read and parsed successfully
 	 */
-	void CopyAdditionalFilesToFolder( const FString& DestinationFolder, TArray<FString>& FilesCopied, bool &bHasSourceFiles, FString InGameFolder = FString() );
+	UE_API void CopyAdditionalFilesToFolder( const FString& DestinationFolder, TArray<FString>& FilesCopied, bool &bHasSourceFiles, FString InGameFolder = FString() );
 
 	/*
 	 * Returns a list of additional files (including the path) as specified in the config file if one exists in the pack file.
@@ -186,18 +188,18 @@ public:
 	 * @param bContainsSource did the file list contain any source files
 	 * @returns true if config file was read and parsed successfully
 	 */
-	bool GetAdditionalFilesForPack(TArray<FString>& FileList, bool& bContainsSource);
+	UE_API bool GetAdditionalFilesForPack(TArray<FString>& FileList, bool& bContainsSource);
 
-	static void ImportPendingPacks();
+	static UE_API void ImportPendingPacks();
 	
 	/* Errors found when parsing manifest (if any) */
 	TArray<FString>	ParseErrors;
 	
-	void BuildListOfAdditionalFiles(TArray<FString>& AdditionalFileSourceList,TArray<FString>& FileList, bool& bContainsSourceFiles);
+	UE_API void BuildListOfAdditionalFiles(TArray<FString>& AdditionalFileSourceList,TArray<FString>& FileList, bool& bContainsSourceFiles);
 
 private:
-	static void ParseAndImportPacks();
-	bool LoadPakFileToBuffer(FPakPlatformFile& PakPlatformFile, FString Path, TArray<uint8>& Buffer);
+	static UE_API void ParseAndImportPacks();
+	UE_API bool LoadPakFileToBuffer(FPakPlatformFile& PakPlatformFile, FString Path, TArray<uint8>& Buffer);
 	
 	
 	/*
@@ -207,22 +209,22 @@ private:
 	 * @param FileList		  array to receive list of files
 	 * @param bContainsSource did the file list contain any source files
 	 */
-	bool ExtractListOfAdditionalFiles(const FString& ConfigFile, TArray<FString>& FileList,bool& bContainsSource);
+	UE_API bool ExtractListOfAdditionalFiles(const FString& ConfigFile, TArray<FString>& FileList,bool& bContainsSource);
 
-	void RecordAndLogError(const FString& ErrorString);
+	UE_API void RecordAndLogError(const FString& ErrorString);
 
 	/* Load the images for the icon and screen shots directly from disk */
-	bool LoadFeaturePackImageData();
+	UE_API bool LoadFeaturePackImageData();
 	
 	/* extract the images for the icon and screen shots from a pak file */
-	bool LoadFeaturePackImageDataFromPackFile(FPakPlatformFile& PakPlatformFile);
+	UE_API bool LoadFeaturePackImageDataFromPackFile(FPakPlatformFile& PakPlatformFile);
 	
 	/* Parse the manifest string describing this pack file */
-	bool ParseManifestString(const FString& ManifestString);
+	UE_API bool ParseManifestString(const FString& ManifestString);
 
 	/** Selects an FLocalizedText from an array which matches either the supplied language code, or the default language code. */
-	FLocalizedTextArray ChooseLocalizedTextArray(TArray<FLocalizedTextArray> Choices, FString LanguageCode);
-	FLocalizedText ChooseLocalizedText(TArray<FLocalizedText> Choices, FString LanguageCode);
+	UE_API FLocalizedTextArray ChooseLocalizedTextArray(TArray<FLocalizedTextArray> Choices, FString LanguageCode);
+	UE_API FLocalizedText ChooseLocalizedText(TArray<FLocalizedText> Choices, FString LanguageCode);
 
 	/* The path of the file we used to create this feature pack instance */
 	FString FeaturePackPath;
@@ -281,3 +283,5 @@ private:
 	FString Identity;
 	FString VersionNumber;
 };
+
+#undef UE_API
diff --git a/Engine/Source/Editor/AdvancedPreviewScene/Public/AdvancedPreviewScene.h b/Engine/Source/Editor/AdvancedPreviewScene/Public/AdvancedPreviewScene.h
--- a/Engine/Source/Editor/AdvancedPreviewScene/Public/AdvancedPreviewScene.h
+++ b/Engine/Source/Editor/AdvancedPreviewScene/Public/AdvancedPreviewScene.h
@@ -14,6 +14,8 @@
 #include "PreviewScene.h"
 #include "TickableEditorObject.h"
 
+#define UE_API ADVANCEDPREVIEWSCENE_API
+
 class FViewport;
 class UAssetViewerSettings;
 class UMaterialInstanceConstant;
@@ -24,82 +26,82 @@ class USphereReflectionCaptureComponent;
 struct FPreviewSceneProfile;
 class FUICommandList;
 
-class ADVANCEDPREVIEWSCENE_API FAdvancedPreviewScene : public FPreviewScene, public FTickableEditorObject
+class FAdvancedPreviewScene : public FPreviewScene, public FTickableEditorObject
 {
 public:
 	/** An event for when a profile changes. Includes the changed profile and (optionally) the name of the property in the profile that changed. */
 	DECLARE_MULTICAST_DELEGATE_TwoParams(FOnProfileChanged, const FPreviewSceneProfile&, FName);
 
-	FAdvancedPreviewScene(ConstructionValues CVS, float InFloorOffset = 0.0f);
-	~FAdvancedPreviewScene();
+	UE_API FAdvancedPreviewScene(ConstructionValues CVS, float InFloorOffset = 0.0f);
+	UE_API ~FAdvancedPreviewScene();
 
-	void UpdateScene(FPreviewSceneProfile& Profile, bool bUpdateSkyLight = true, bool bUpdateEnvironment = true, bool bUpdatePostProcessing = true, bool bUpdateDirectionalLight = true);
+	UE_API void UpdateScene(FPreviewSceneProfile& Profile, bool bUpdateSkyLight = true, bool bUpdateEnvironment = true, bool bUpdatePostProcessing = true, bool bUpdateDirectionalLight = true);
 
 	/** Begin FPreviewScene */
-	virtual FLinearColor GetBackgroundColor() const override;
+	UE_API virtual FLinearColor GetBackgroundColor() const override;
 	/** End FPreviewScene */
 
 	/* Begin FTickableEditorObject */
-	virtual void Tick(float DeltaTime) override;
+	UE_API virtual void Tick(float DeltaTime) override;
 	virtual ETickableTickType GetTickableTickType() const override { return ETickableTickType::Always; }
-	virtual TStatId GetStatId() const override;
+	UE_API virtual TStatId GetStatId() const override;
 	/* End FTickableEditorObject */
 	
 	FOnProfileChanged& OnProfileChanged() { return OnProfileChangedDelegate; }
 
-	const bool HandleViewportInput(FViewport* InViewport, FInputDeviceId DeviceId, FKey Key, float Delta, float DeltaTime, int32 NumSamples, bool bGamepad);
-	const bool HandleInputKey(const FInputKeyEventArgs& EventArgs);
+	UE_API const bool HandleViewportInput(FViewport* InViewport, FInputDeviceId DeviceId, FKey Key, float Delta, float DeltaTime, int32 NumSamples, bool bGamepad);
+	UE_API const bool HandleInputKey(const FInputKeyEventArgs& EventArgs);
 
-	void SetSkyRotation(const float SkyRotation);
+	UE_API void SetSkyRotation(const float SkyRotation);
 	/* Sets the visiblity state for the floor/environment by storing it in the scene profile and refreshing the scene, in case bDirect is true it sets the visibility directly and leaves the profile untouched. */
-	bool GetFloorVisibility() const;
-	void SetFloorVisibility(const bool bVisible, const bool bDirect = false);
-	void SetEnvironmentVisibility(const bool bVisible, const bool bDirect = false);
-	float GetFloorOffset() const;
-	void SetFloorOffset(const float InFloorOffset);
-	void SetProfileIndex(const int32 InProfileIndex);
-	FPreviewSceneProfile* GetCurrentProfile() const;
-
-	const UStaticMeshComponent* GetFloorMeshComponent() const;
-	const float GetSkyRotation() const;
-	const int32 GetCurrentProfileIndex() const;
-	const bool IsUsingPostProcessing() const;
+	UE_API bool GetFloorVisibility() const;
+	UE_API void SetFloorVisibility(const bool bVisible, const bool bDirect = false);
+	UE_API void SetEnvironmentVisibility(const bool bVisible, const bool bDirect = false);
+	UE_API float GetFloorOffset() const;
+	UE_API void SetFloorOffset(const float InFloorOffset);
+	UE_API void SetProfileIndex(const int32 InProfileIndex);
+	UE_API FPreviewSceneProfile* GetCurrentProfile() const;
+
+	UE_API const UStaticMeshComponent* GetFloorMeshComponent() const;
+	UE_API const float GetSkyRotation() const;
+	UE_API const int32 GetCurrentProfileIndex() const;
+	UE_API const bool IsUsingPostProcessing() const;
 
 	TSharedPtr<FUICommandList> GetCommandList() const { return UICommandList; }
 
 	/** Toggle the sky sphere on and off */
-	void HandleToggleEnvironment();
+	UE_API void HandleToggleEnvironment();
 	
 	/** Whether the sky sphere is on. */
-	bool IsEnvironmentEnabled() const;
+	UE_API bool IsEnvironmentEnabled() const;
 
 	/** Toggle the floor mesh on and off */
-	void HandleToggleFloor();
+	UE_API void HandleToggleFloor();
 
 	/** Whether the floor is on */
-	bool IsFloorEnabled() const;
+	UE_API bool IsFloorEnabled() const;
 
 	/** Toggle the grid on and off. */
-	void HandleToggleGrid();
+	UE_API void HandleToggleGrid();
 	
 	/** Whether the grid is enabled. */
-	bool IsGridEnabled() const;
+	UE_API bool IsGridEnabled() const;
 
 	/** Toggle post processing on and off */
-	void HandleTogglePostProcessing();
+	UE_API void HandleTogglePostProcessing();
 	
 	/** Whether post processing is enabled */
-	bool IsPostProcessingEnabled() const;
+	UE_API bool IsPostProcessingEnabled() const;
 
 protected:
 	/** Create and map the command list. */
-	virtual void BindCommands();
+	UE_API virtual void BindCommands();
 	
 	/** Add commands to a provided command list. */
-	virtual TSharedRef<FUICommandList> CreateCommandList();
+	UE_API virtual TSharedRef<FUICommandList> CreateCommandList();
 	
 	/** Handle refreshing the scene when settings change */
-	void OnAssetViewerSettingsRefresh(const FName& InPropertyName);
+	UE_API void OnAssetViewerSettingsRefresh(const FName& InPropertyName);
 
 protected:
 	FOnProfileChanged OnProfileChangedDelegate;
@@ -126,3 +128,5 @@ protected:
 	/** Delegate handle used to refresh the scene when settings change */
 	FDelegateHandle RefreshDelegate;
 };
+
+#undef UE_API
diff --git a/Engine/Source/Editor/AdvancedPreviewScene/Public/AdvancedPreviewSceneMenus.h b/Engine/Source/Editor/AdvancedPreviewScene/Public/AdvancedPreviewSceneMenus.h
--- a/Engine/Source/Editor/AdvancedPreviewScene/Public/AdvancedPreviewSceneMenus.h
+++ b/Engine/Source/Editor/AdvancedPreviewScene/Public/AdvancedPreviewSceneMenus.h
@@ -13,7 +13,7 @@ struct FPreviewSceneProfile;
 namespace UE::AdvancedPreviewScene::Menus
 {
 
-struct ADVANCEDPREVIEWSCENE_API FSettingsOptions
+struct FSettingsOptions
 {
 	bool bShowToggleEnvironment = true;
 	bool bShowToggleGrid = true;
diff --git a/Engine/Source/Editor/AdvancedPreviewScene/Public/AssetViewerSettings.h b/Engine/Source/Editor/AdvancedPreviewScene/Public/AssetViewerSettings.h
--- a/Engine/Source/Editor/AdvancedPreviewScene/Public/AssetViewerSettings.h
+++ b/Engine/Source/Editor/AdvancedPreviewScene/Public/AssetViewerSettings.h
@@ -13,6 +13,8 @@
 #include "ShowFlags.h"
 #include "AssetViewerSettings.generated.h"
 
+#define UE_API ADVANCEDPREVIEWSCENE_API
+
 /**
 * Preview scene profile settings structure.
 */
@@ -176,8 +178,8 @@ struct FPreviewSceneProfile
 	}
 };
 
-UCLASS()
-class ADVANCEDPREVIEWSCENE_API UDefaultEditorProfiles : public UObject
+UCLASS(MinimalAPI)
+class UDefaultEditorProfiles : public UObject
 {
 	GENERATED_BODY()
 public:
@@ -221,11 +223,11 @@ public:
 		Profiles.Add(GreyAmbientProfile);
 	}
 
-	const FPreviewSceneProfile* GetProfile(const FString& ProfileName);
+	UE_API const FPreviewSceneProfile* GetProfile(const FString& ProfileName);
 
-	static FName DefaultProfileName;
-	static FName EditingProfileName;
-	static FName GreyAmbientProfileName;
+	static UE_API FName DefaultProfileName;
+	static UE_API FName EditingProfileName;
+	static UE_API FName GreyAmbientProfileName;
 	
 	/** Collection of default engine-provided profiles used in various editing environments*/
 	UPROPERTY()
@@ -255,23 +257,23 @@ public:
 /**
 * Default asset viewer settings.
 */
-UCLASS(config=Editor, meta=(DisplayName = "Asset Viewer"))
-class ADVANCEDPREVIEWSCENE_API UAssetViewerSettings : public UObject, public FEditorUndoClient
+UCLASS(MinimalAPI, config=Editor, meta=(DisplayName = "Asset Viewer"))
+class UAssetViewerSettings : public UObject, public FEditorUndoClient
 {
 	GENERATED_BODY()
 public:
-	UAssetViewerSettings();
-	virtual ~UAssetViewerSettings();
+	UE_API UAssetViewerSettings();
+	UE_API virtual ~UAssetViewerSettings();
 
-	static UAssetViewerSettings* Get();
+	static UE_API UAssetViewerSettings* Get();
 
-	static FPreviewSceneProfile& GetCurrentUserProjectProfile();
+	static UE_API FPreviewSceneProfile& GetCurrentUserProjectProfile();
 	
 	/**
 	 * Saves the config data out to the ini files
 	 * @param bWarnIfFail Should we log a warning if a ini file couldn't be saved.
 	 */
-	void Save(bool bWarnIfFail = true);
+	UE_API void Save(bool bWarnIfFail = true);
 
 	DECLARE_EVENT_OneParam(UAssetViewerSettings, FOnAssetViewerSettingsChangedEvent, const FName&);
 	FOnAssetViewerSettingsChangedEvent& OnAssetViewerSettingsChanged() { return OnAssetViewerSettingsChangedEvent; }
@@ -283,13 +285,13 @@ public:
 	FOnAssetViewerSettingsPostUndo& OnAssetViewerSettingsPostUndo() { return OnAssetViewerSettingsPostUndoEvent; }
 
 	/** Begin UObject */
-	virtual void PostEditChangeProperty(struct FPropertyChangedEvent& PropertyChangedEvent) override;
-	virtual void PostInitProperties() override;
+	UE_API virtual void PostEditChangeProperty(struct FPropertyChangedEvent& PropertyChangedEvent) override;
+	UE_API virtual void PostInitProperties() override;
 	/** End UObject */
 
 	/** Begin FEditorUndoClient */
-	virtual void PostUndo(bool bSuccess) override;
-	virtual void PostRedo(bool bSuccess) override;
+	UE_API virtual void PostUndo(bool bSuccess) override;
+	UE_API virtual void PostRedo(bool bSuccess) override;
 	/** End FEditorUndoClient */
 
 	/** Collection of scene profiles */
@@ -310,3 +312,5 @@ protected:
 	UPROPERTY(Config)
 	bool bFakeConfigValue_HACK;
 };
+
+#undef UE_API
diff --git a/Engine/Source/Editor/AdvancedPreviewScene/Public/PreviewProfileController.h b/Engine/Source/Editor/AdvancedPreviewScene/Public/PreviewProfileController.h
--- a/Engine/Source/Editor/AdvancedPreviewScene/Public/PreviewProfileController.h
+++ b/Engine/Source/Editor/AdvancedPreviewScene/Public/PreviewProfileController.h
@@ -5,6 +5,8 @@
 #include "CoreMinimal.h"
 #include "IPreviewProfileController.h"
 
+#define UE_API ADVANCEDPREVIEWSCENE_API
+
 class UAssetViewerSettings;
 class UEditorPerProjectUserSettings;
 
@@ -15,23 +17,23 @@ class UEditorPerProjectUserSettings;
  * @note This class was created to decouple UnrealEd from AdvancedPreviewScene and prevent
  *       circular dependencies between the modules.
  */
-class ADVANCEDPREVIEWSCENE_API FPreviewProfileController : public IPreviewProfileController
+class FPreviewProfileController : public IPreviewProfileController
 {
 public:
-	FPreviewProfileController();
-	virtual ~FPreviewProfileController();
+	UE_API FPreviewProfileController();
+	UE_API virtual ~FPreviewProfileController();
 
 	/** Returns the list of available preview profiles names. */
-	virtual TArray<FString> GetPreviewProfiles(int32& OutCurrentProfileIndex) const override;
+	UE_API virtual TArray<FString> GetPreviewProfiles(int32& OutCurrentProfileIndex) const override;
 
 	/** Set the specified preview profiles as the active one. */
-	virtual bool SetActiveProfile(const FString& ProfileName) override;
+	UE_API virtual bool SetActiveProfile(const FString& ProfileName) override;
 
 	/** Returns the preview profiles currently active. */
-	virtual FString GetActiveProfile() const override;
+	UE_API virtual FString GetActiveProfile() const override;
 
 	/** Returns true if user has added one or more of their own profiles */
-	virtual bool HasAnyUserProfiles() const override;
+	UE_API virtual bool HasAnyUserProfiles() const override;
 
 	/** Invoked after the list of available profiles has changed. */
 	virtual FOnPreviewProfileListChanged& OnPreviewProfileListChanged() override { return OnPreviewProfileListChangedDelegate; }
@@ -40,8 +42,8 @@ public:
 	virtual FOnPreviewProfileChanged& OnPreviewProfileChanged() override { return OnPreviewSettingChangedDelegate; }
 
 private:
-	void UpdateAssetViewerProfiles();
-	void EnsureProfilesStateCoherence() const;
+	UE_API void UpdateAssetViewerProfiles();
+	UE_API void EnsureProfilesStateCoherence() const;
 
 private:
 	/** Provides the list of available preview profiles along with the profile change delegates (to update the profile combo box selection) */
@@ -62,3 +64,5 @@ private:
 	FDelegateHandle AssetViewerSettingsProfileAddRemoveHandle;
 	FDelegateHandle AssetViewerSettingsChangedHandle;
 };
+
+#undef UE_API
diff --git a/Engine/Source/Editor/AdvancedPreviewScene/Public/SAdvancedPreviewDetailsTab.h b/Engine/Source/Editor/AdvancedPreviewScene/Public/SAdvancedPreviewDetailsTab.h
--- a/Engine/Source/Editor/AdvancedPreviewScene/Public/SAdvancedPreviewDetailsTab.h
+++ b/Engine/Source/Editor/AdvancedPreviewScene/Public/SAdvancedPreviewDetailsTab.h
@@ -9,12 +9,14 @@
 #include "PropertyEditorDelegates.h"
 #include "AdvancedPreviewSceneModule.h"
 
+#define UE_API ADVANCEDPREVIEWSCENE_API
+
 class FAdvancedPreviewScene;
 class IDetailsView;
 class UAssetViewerSettings;
 class UEditorPerProjectUserSettings;
 
-class ADVANCEDPREVIEWSCENE_API SAdvancedPreviewDetailsTab : public SCompoundWidget
+class SAdvancedPreviewDetailsTab : public SCompoundWidget
 {
 public:
 	SLATE_BEGIN_ARGS(SAdvancedPreviewDetailsTab)
@@ -37,24 +39,24 @@ public:
 
 public:
 	/** **/
-	SAdvancedPreviewDetailsTab();
-	~SAdvancedPreviewDetailsTab();
+	UE_API SAdvancedPreviewDetailsTab();
+	UE_API ~SAdvancedPreviewDetailsTab();
 
 	/** SWidget functions */
-	void Construct(const FArguments& InArgs, const TSharedRef<FAdvancedPreviewScene>& InPreviewScene);
+	UE_API void Construct(const FArguments& InArgs, const TSharedRef<FAdvancedPreviewScene>& InPreviewScene);
 
-	void Refresh();
+	UE_API void Refresh();
 
 protected:
-	void CreateSettingsView();
-	void ComboBoxSelectionChanged(TSharedPtr<FString> NewSelection, ESelectInfo::Type /*SelectInfo*/);
-	void UpdateSettingsView();
-	void UpdateProfileNames();
-	FReply AddProfileButtonClick();
-	FReply RemoveOrResetProfileButtonClick();
+	UE_API void CreateSettingsView();
+	UE_API void ComboBoxSelectionChanged(TSharedPtr<FString> NewSelection, ESelectInfo::Type /*SelectInfo*/);
+	UE_API void UpdateSettingsView();
+	UE_API void UpdateProfileNames();
+	UE_API FReply AddProfileButtonClick();
+	UE_API FReply RemoveOrResetProfileButtonClick();
 protected:
-	void OnAssetViewerSettingsRefresh(const FName& InPropertyName);
-	void OnAssetViewerSettingsPostUndo();
+	UE_API void OnAssetViewerSettingsRefresh(const FName& InPropertyName);
+	UE_API void OnAssetViewerSettingsPostUndo();
 protected:
 	/** Property viewing widget */
 	TSharedPtr<IDetailsView> SettingsView;
@@ -66,7 +68,7 @@ protected:
 	TArray<TSharedPtr<FString>> ProfileNames;
 	int32 ProfileIndex;
 
-	void OnPreviewSceneChanged(TSharedRef<FAdvancedPreviewScene> PreviewScene);
+	UE_API void OnPreviewSceneChanged(TSharedRef<FAdvancedPreviewScene> PreviewScene);
 
 	FDelegateHandle RefreshDelegate;
 	FDelegateHandle AddRemoveProfileDelegate;
@@ -80,3 +82,5 @@ protected:
 
 	TArray<FAdvancedPreviewSceneModule::FDetailDelegates> Delegates;
 };
+
+#undef UE_API
diff --git a/Engine/Source/Editor/AnimGraph/Public/AnimBlueprintExtension.h b/Engine/Source/Editor/AnimGraph/Public/AnimBlueprintExtension.h
--- a/Engine/Source/Editor/AnimGraph/Public/AnimBlueprintExtension.h
+++ b/Engine/Source/Editor/AnimGraph/Public/AnimBlueprintExtension.h
@@ -7,6 +7,8 @@
 
 #include "AnimBlueprintExtension.generated.h"
 
+#define UE_API ANIMGRAPH_API
+
 class UAnimBlueprint;
 class IAnimBlueprintCompilationContext;
 class IAnimBlueprintCopyTermDefaultsContext;
@@ -19,8 +21,8 @@ class UAnimGraphNode_Base;
 class UEdGraph;
 
 /** Extension that allows per-system data to be held on the anim blueprint, and per-system logic to be executed during compilation */
-UCLASS()
-class ANIMGRAPH_API UAnimBlueprintExtension : public UBlueprintExtension
+UCLASS(MinimalAPI)
+class UAnimBlueprintExtension : public UBlueprintExtension
 {
 	GENERATED_BODY()
 
@@ -33,7 +35,7 @@ public:
 	}
 	
 	/** Request an anim blueprint extension for an anim blueprint. It is illegal to perform this operation once compilation has commenced, use GetExtension instead. */
-	static UAnimBlueprintExtension* RequestExtension(UAnimBlueprint* InAnimBlueprint, TSubclassOf<UAnimBlueprintExtension> InExtensionType);
+	static UE_API UAnimBlueprintExtension* RequestExtension(UAnimBlueprint* InAnimBlueprint, TSubclassOf<UAnimBlueprintExtension> InExtensionType);
 
 	/** Get an already-requested extension for an anim blueprint. Will assert if the extension is not found. */
 	template<typename ExtensionType>
@@ -50,31 +52,31 @@ public:
 	}
 	
 	/** Get an already-requested an anim blueprint extension for an anim blueprint. */
-	static UAnimBlueprintExtension* GetExtension(UAnimBlueprint* InAnimBlueprint, TSubclassOf<UAnimBlueprintExtension> InExtensionType);
+	static UE_API UAnimBlueprintExtension* GetExtension(UAnimBlueprint* InAnimBlueprint, TSubclassOf<UAnimBlueprintExtension> InExtensionType);
 	
 	/** Iterate over all registered UAnimBlueprintExtensions in an anim BP */
-	static void ForEachExtension(UAnimBlueprint* InAnimBlueprint, TFunctionRef<void(UAnimBlueprintExtension*)> InFunction);
+	static UE_API void ForEachExtension(UAnimBlueprint* InAnimBlueprint, TFunctionRef<void(UAnimBlueprintExtension*)> InFunction);
 
 	/** Get all subsystems currently present on an anim blueprint */
-	static TArray<UAnimBlueprintExtension*> GetExtensions(UAnimBlueprint* InAnimBlueprint);
+	static UE_API TArray<UAnimBlueprintExtension*> GetExtensions(UAnimBlueprint* InAnimBlueprint);
 
 	/** Request all extensions that a node needs */
-	static void RequestExtensionsForNode(UAnimGraphNode_Base* InAnimGraphNode);
+	static UE_API void RequestExtensionsForNode(UAnimGraphNode_Base* InAnimGraphNode);
 	
 	/** Refresh all extensions according to nodes present in an anim BP */
-	static void RefreshExtensions(UAnimBlueprint* InAnimBlueprint);
+	static UE_API void RefreshExtensions(UAnimBlueprint* InAnimBlueprint);
 	
 	/** Get the structure that will be added to any BP-derived UAnimInstance */
-	const UScriptStruct* GetInstanceDataType() const;
+	UE_API const UScriptStruct* GetInstanceDataType() const;
 
 	/** Get the structure that will be added to any anim BP class */
-	const UScriptStruct* GetClassDataType() const;
+	UE_API const UScriptStruct* GetClassDataType() const;
 
 	/** Get the property of the instance data */
-	const FStructProperty* GetInstanceDataProperty() const;
+	UE_API const FStructProperty* GetInstanceDataProperty() const;
 
 	/** Get the property of the class data */
-	const FStructProperty* GetClassDataProperty() const;
+	UE_API const FStructProperty* GetClassDataProperty() const;
 
 	/** Get the defaults that will be held on the class */
 	template<typename ClassDataType>
@@ -92,7 +94,7 @@ public:
 
 protected:
 	/** Get the anim blueprint that hosts this extension */
-	UAnimBlueprint* GetAnimBlueprint() const;
+	UE_API UAnimBlueprint* GetAnimBlueprint() const;
 	
 private:
 	friend class FAnimBlueprintCompilerContext;
@@ -167,7 +169,7 @@ private:
 	 * Called when a compiler context is created for the anim blueprint
 	 * @param	InCreationContext	The compiler context for the current compilation
 	 */	
-	void BeginCompilation(IAnimBlueprintCompilerCreationContext& InCreationContext);
+	UE_API void BeginCompilation(IAnimBlueprintCompilerCreationContext& InCreationContext);
 
 	/** 
 	 * Called when the class starts compiling. The class may be new or recycled.
@@ -175,7 +177,7 @@ private:
 	 * @param	InCompilerContext	The compiler context for the current compilation
 	 * @param	OutCompiledData	The compiled data that this handler can write to
 	 */
-	void StartCompilingClass(const UClass* InClass, IAnimBlueprintCompilationBracketContext& InCompilationContext, IAnimBlueprintGeneratedClassCompiledData& OutCompiledData);
+	UE_API void StartCompilingClass(const UClass* InClass, IAnimBlueprintCompilationBracketContext& InCompilationContext, IAnimBlueprintGeneratedClassCompiledData& OutCompiledData);
 
 	/** 
 	 * Called before all animation nodes are processed
@@ -183,7 +185,7 @@ private:
 	 * @param	InCompilerContext	The compiler context for the current compilation
 	 * @param	OutCompiledData	The compiled data that this handler can write to
 	 */
-	void PreProcessAnimationNodes(TArrayView<UAnimGraphNode_Base*> InAnimNodes, IAnimBlueprintCompilationContext& InCompilationContext, IAnimBlueprintGeneratedClassCompiledData& OutCompiledData);
+	UE_API void PreProcessAnimationNodes(TArrayView<UAnimGraphNode_Base*> InAnimNodes, IAnimBlueprintCompilationContext& InCompilationContext, IAnimBlueprintGeneratedClassCompiledData& OutCompiledData);
 
 	/** 
 	 * Called after all animation nodes are processed
@@ -191,7 +193,7 @@ private:
 	 * @param	InCompilerContext	The compiler context for the current compilation
 	 * @param	OutCompiledData	The compiled data that this handler can write to
 	 */	
-	void PostProcessAnimationNodes(TArrayView<UAnimGraphNode_Base*> InAnimNodes, IAnimBlueprintCompilationContext& InCompilationContext, IAnimBlueprintGeneratedClassCompiledData& OutCompiledData);
+	UE_API void PostProcessAnimationNodes(TArrayView<UAnimGraphNode_Base*> InAnimNodes, IAnimBlueprintCompilationContext& InCompilationContext, IAnimBlueprintGeneratedClassCompiledData& OutCompiledData);
 
 	/** 
 	 * Called post graph expansion
@@ -199,7 +201,7 @@ private:
 	 * @param	InCompilerContext	The compiler context for the current compilation
 	 * @param	OutCompiledData	The compiled data that this handler can write to
 	 */		
-	void PostExpansionStep(const UEdGraph* InGraph, IAnimBlueprintPostExpansionStepContext& InCompilationContext, IAnimBlueprintGeneratedClassCompiledData& OutCompiledData);
+	UE_API void PostExpansionStep(const UEdGraph* InGraph, IAnimBlueprintPostExpansionStepContext& InCompilationContext, IAnimBlueprintGeneratedClassCompiledData& OutCompiledData);
 
 	/** 
 	 * Called when the class has finished compiling
@@ -207,7 +209,7 @@ private:
 	 * @param	InCompilerContext	The compiler context for the current compilation
 	 * @param	OutCompiledData	The compiled data that this handler can write to
 	 */
-	void FinishCompilingClass(const UClass* InClass, IAnimBlueprintCompilationBracketContext& InCompilationContext, IAnimBlueprintGeneratedClassCompiledData& OutCompiledData);
+	UE_API void FinishCompilingClass(const UClass* InClass, IAnimBlueprintCompilationBracketContext& InCompilationContext, IAnimBlueprintGeneratedClassCompiledData& OutCompiledData);
 	
 	/** 
 	 * Called when data is being copied to the CDO
@@ -215,7 +217,7 @@ private:
 	 * @param	InCompilerContext	The compiler context for the current compilation
 	 * @param	InPerExtensionContext	The extension context for the current compilation
 	 */
-	void CopyTermDefaultsToDefaultObject(UObject* InDefaultObject, IAnimBlueprintCopyTermDefaultsContext& InCompilationContext, IAnimBlueprintExtensionCopyTermDefaultsContext& InPerExtensionContext);
+	UE_API void CopyTermDefaultsToDefaultObject(UObject* InDefaultObject, IAnimBlueprintCopyTermDefaultsContext& InCompilationContext, IAnimBlueprintExtensionCopyTermDefaultsContext& InPerExtensionContext);
 
 	/** 
 	 * Called when data is being copied to the sparse class data
@@ -223,14 +225,16 @@ private:
 	 * @param	InCompilerContext	The compiler context for the current compilation
 	 * @param	InPerExtensionContext	The extension context for the current compilation
 	 */
-	void CopyTermDefaultsToSparseClassData(IAnimBlueprintCopyTermDefaultsContext& InCompilationContext, IAnimBlueprintExtensionCopyTermDefaultsContext& InPerExtensionContext);
+	UE_API void CopyTermDefaultsToSparseClassData(IAnimBlueprintCopyTermDefaultsContext& InCompilationContext, IAnimBlueprintExtensionCopyTermDefaultsContext& InPerExtensionContext);
 
 	/** Called when a compiler context is destroyed for the anim blueprint. Can be used to clean up resources. */
-	void EndCompilation();
+	UE_API void EndCompilation();
 	
 	/** Get the defaults that will be held on the class */
-	void* GetClassDataInternal();
+	UE_API void* GetClassDataInternal();
 
 	/** Get the defaults that will be held on the instance */
-	void* GetInstanceDataInternal();
+	UE_API void* GetInstanceDataInternal();
 };
+
+#undef UE_API
diff --git a/Engine/Source/Editor/AnimGraph/Public/AnimBlueprintExtension_PropertyAccess.h b/Engine/Source/Editor/AnimGraph/Public/AnimBlueprintExtension_PropertyAccess.h
--- a/Engine/Source/Editor/AnimGraph/Public/AnimBlueprintExtension_PropertyAccess.h
+++ b/Engine/Source/Editor/AnimGraph/Public/AnimBlueprintExtension_PropertyAccess.h
@@ -18,6 +18,8 @@
 
 #include "AnimBlueprintExtension_PropertyAccess.generated.h"
 
+#define UE_API ANIMGRAPH_API
+
 class FString;
 class IAnimBlueprintCompilationBracketContext;
 class IAnimBlueprintCompilerCreationContext;
@@ -34,8 +36,8 @@ class UEdGraphPin;
 // Delegate called when the library is compiled (whether successfully or not)
 DECLARE_MULTICAST_DELEGATE_TwoParams(FOnPostLibraryCompiled, IAnimBlueprintCompilationBracketContext& /*InCompilationContext*/, IAnimBlueprintGeneratedClassCompiledData& /*OutCompiledData*/)
 
-UCLASS()
-class ANIMGRAPH_API UAnimBlueprintExtension_PropertyAccess : public UAnimBlueprintExtension
+UCLASS(MinimalAPI)
+class UAnimBlueprintExtension_PropertyAccess : public UAnimBlueprintExtension
 {
 	GENERATED_BODY()
 
@@ -43,27 +45,27 @@ class ANIMGRAPH_API UAnimBlueprintExtension_PropertyAccess : public UAnimBluepri
 
 public:	
 	// Try to determine the context in which this property access can safely be performed.
-	static const FName ContextId_Automatic;
+	static UE_API const FName ContextId_Automatic;
 
 	// Can safely be executed on worker threads
-	static const FName ContextId_UnBatched_ThreadSafe;
+	static UE_API const FName ContextId_UnBatched_ThreadSafe;
 
 	// Executed batched on the game thread before the event graph is run
-	static const FName ContextId_Batched_WorkerThreadPreEventGraph;
+	static UE_API const FName ContextId_Batched_WorkerThreadPreEventGraph;
 
 	// Executed batched on the game thread after the event graph is run
-	static const FName ContextId_Batched_WorkerThreadPostEventGraph;
+	static UE_API const FName ContextId_Batched_WorkerThreadPostEventGraph;
 
 	// Executed batched on the game thread before the event graph is run
-	static const FName ContextId_Batched_GameThreadPreEventGraph;
+	static UE_API const FName ContextId_Batched_GameThreadPreEventGraph;
 
 	// Executed batched on the game thread after the event graph is run
-	static const FName ContextId_Batched_GameThreadPostEventGraph;
+	static UE_API const FName ContextId_Batched_GameThreadPostEventGraph;
 	
 public:
 	// Add a copy to the property access library we are compiling
 	// @return an integer handle to the pending copy. This can be resolved to a true copy index by calling GetCompiledHandle
-	FPropertyAccessHandle AddCopy(TArrayView<FString> InSourcePath, TArrayView<FString> InDestPath, const FName& InContextId, UObject* InObject = nullptr);
+	UE_API FPropertyAccessHandle AddCopy(TArrayView<FString> InSourcePath, TArrayView<FString> InDestPath, const FName& InContextId, UObject* InObject = nullptr);
 	
 	UE_DEPRECATED(5.0, "Please use AddCopy with a context ID")
 	int32 AddCopy(TArrayView<FString> InSourcePath, TArrayView<FString> InDestPath, EPropertyAccessBatchType InBatchType, UObject* InObject = nullptr)
@@ -74,7 +76,7 @@ public:
 
 	// Add an access to the property access library we are compiling
 	// @return an integer handle to the pending access. This can be resolved to a true access index by calling GetCompiledHandle
-	FPropertyAccessHandle AddAccess(TArrayView<FString> InPath, UObject* InObject = nullptr);
+	UE_API FPropertyAccessHandle AddAccess(TArrayView<FString> InPath, UObject* InObject = nullptr);
 
 	// Delegate called when the library is compiled (whether successfully or not)
 	FSimpleMulticastDelegate& OnPreLibraryCompiled() { return OnPreLibraryCompiledDelegate; }
@@ -83,27 +85,27 @@ public:
 	FOnPostLibraryCompiled& OnPostLibraryCompiled() { return OnPostLibraryCompiledDelegate; }
 
 	// Maps the initial copy handle to a true handle, post compilation
-	FCompiledPropertyAccessHandle GetCompiledHandle(FPropertyAccessHandle InHandle) const;
+	UE_API FCompiledPropertyAccessHandle GetCompiledHandle(FPropertyAccessHandle InHandle) const;
 
 	// Get the access type for the specified handle
-	EPropertyAccessCopyType GetCompiledHandleAccessType(FPropertyAccessHandle InHandle) const;
+	UE_API EPropertyAccessCopyType GetCompiledHandleAccessType(FPropertyAccessHandle InHandle) const;
 
 	// Expands a property access path to a pure chain of BP nodes
-	void ExpandPropertyAccess(FKismetCompilerContext& InCompilerContext, TArrayView<FString> InSourcePath, UEdGraph* InParentGraph, UEdGraphPin* InTargetPin) const;
+	UE_API void ExpandPropertyAccess(FKismetCompilerContext& InCompilerContext, TArrayView<FString> InSourcePath, UEdGraph* InParentGraph, UEdGraphPin* InTargetPin) const;
 
 	// Maps a compiled handle back to a human-readable context 
-	static FText GetCompiledHandleContext(FCompiledPropertyAccessHandle InHandle);
+	static UE_API FText GetCompiledHandleContext(FCompiledPropertyAccessHandle InHandle);
 	
 	// Maps a compiled handle back to a human-readable context description
-	static FText GetCompiledHandleContextDesc(FCompiledPropertyAccessHandle InHandle);
+	static UE_API FText GetCompiledHandleContextDesc(FCompiledPropertyAccessHandle InHandle);
 
 	// Whether a context name requires an auto-generated variable to be cached for an access
-	static bool ContextRequiresCachedVariable(FName InName);
+	static UE_API bool ContextRequiresCachedVariable(FName InName);
 	
 private:
 	// UAnimBlueprintExtension interface
-	virtual void HandleStartCompilingClass(const UClass* InClass, IAnimBlueprintCompilationBracketContext& InCompilationContext, IAnimBlueprintGeneratedClassCompiledData& OutCompiledData) override;
-	virtual void HandleFinishCompilingClass(const UClass* InClass, IAnimBlueprintCompilationBracketContext& InCompilationContext, IAnimBlueprintGeneratedClassCompiledData& OutCompiledData) override;
+	UE_API virtual void HandleStartCompilingClass(const UClass* InClass, IAnimBlueprintCompilationBracketContext& InCompilationContext, IAnimBlueprintGeneratedClassCompiledData& OutCompiledData) override;
+	UE_API virtual void HandleFinishCompilingClass(const UClass* InClass, IAnimBlueprintCompilationBracketContext& InCompilationContext, IAnimBlueprintGeneratedClassCompiledData& OutCompiledData) override;
 
 private:	
 	// Property access library compiler
@@ -117,4 +119,6 @@ private:
 
 	UPROPERTY()
 	FAnimSubsystem_PropertyAccess Subsystem;
-};
\ No newline at end of file
+};
+
+#undef UE_API
diff --git a/Engine/Source/Editor/AnimGraph/Public/AnimBlueprintPostCompileValidation.h b/Engine/Source/Editor/AnimGraph/Public/AnimBlueprintPostCompileValidation.h
--- a/Engine/Source/Editor/AnimGraph/Public/AnimBlueprintPostCompileValidation.h
+++ b/Engine/Source/Editor/AnimGraph/Public/AnimBlueprintPostCompileValidation.h
@@ -12,6 +12,8 @@
 
 #include "AnimBlueprintPostCompileValidation.generated.h"
 
+#define UE_API ANIMGRAPH_API
+
 // Encapsulated Parameters, to upgrade these without changing function signature.
 struct FAnimBPCompileValidationParams
 {
@@ -35,19 +37,19 @@ struct FAnimBPCompileValidationParams
 };
 
 // This class is a base class for performing AnimBlueprint Post Compilation Validation.
-UCLASS()
-class ANIMGRAPH_API UAnimBlueprintPostCompileValidation : public UObject
+UCLASS(MinimalAPI)
+class UAnimBlueprintPostCompileValidation : public UObject
 {
 	GENERATED_UCLASS_BODY()
 
 public:
-	virtual void DoPostCompileValidation(FAnimBPCompileValidationParams& InParams) const;
+	UE_API virtual void DoPostCompileValidation(FAnimBPCompileValidationParams& InParams) const;
 
 protected:
-	static void PCV_PreloadObject(const UObject* const ReferencedObject);
+	static UE_API void PCV_PreloadObject(const UObject* const ReferencedObject);
 
-	static void PCV_GatherAnimSequences(TArray<const UAnimSequence*>& OutAnimSequences, const UAnimSequenceBase* const InAnimSequenceBase);
-	static void PCV_GatherAnimSequences(TArray<const UAnimSequence*>& OutAnimSequences, const class UBlendSpace* const InBlendSpace);
+	static UE_API void PCV_GatherAnimSequences(TArray<const UAnimSequence*>& OutAnimSequences, const UAnimSequenceBase* const InAnimSequenceBase);
+	static UE_API void PCV_GatherAnimSequences(TArray<const UAnimSequence*>& OutAnimSequences, const class UBlendSpace* const InBlendSpace);
 
 	struct FPCV_GatherParams
 	{
@@ -84,8 +86,8 @@ protected:
 		{}
 	};
 
-	static void PCV_GatherAnimSequencesFromGraph(TArray<const UAnimSequence*>& OutAnimSequences, FAnimBPCompileValidationParams& PCV_Params, const FPCV_GatherParams& GatherParams);
-	static void PCV_GatherBlendSpacesFromGraph(TArray<const class UBlendSpace*>& OutBlendSpaces, FAnimBPCompileValidationParams& PCV_Params);
+	static UE_API void PCV_GatherAnimSequencesFromGraph(TArray<const UAnimSequence*>& OutAnimSequences, FAnimBPCompileValidationParams& PCV_Params, const FPCV_GatherParams& GatherParams);
+	static UE_API void PCV_GatherBlendSpacesFromGraph(TArray<const class UBlendSpace*>& OutBlendSpaces, FAnimBPCompileValidationParams& PCV_Params);
 
 	struct FPCV_ReferencedAnimSequence
 	{
@@ -117,11 +119,13 @@ protected:
 		{}
 	};
 
-	static void PCV_GatherAllReferencedAnimSequences(TArray<FPCV_ReferencedAnimSequence>& OutRefAnimSequences, FAnimBPCompileValidationParams& PCV_Params);
-	static void PCV_GatherAnimSequencesFromStruct(TArray<FPCV_ReferencedAnimSequence>& OutRefAnimSequences, FAnimBPCompileValidationParams& PCV_Params, const UStruct* InStruct, const void* InData, TArray<FPCV_PropertyAndValue> InPropertyCallChain);
-	static void PCV_GatherAnimSequencesFromProperty(TArray<FPCV_ReferencedAnimSequence>& OutRefAnimSequences, FAnimBPCompileValidationParams& PCV_Params, const FProperty* InProperty, const void* InData, TArray<FPCV_PropertyAndValue> InPropertyCallChain);
+	static UE_API void PCV_GatherAllReferencedAnimSequences(TArray<FPCV_ReferencedAnimSequence>& OutRefAnimSequences, FAnimBPCompileValidationParams& PCV_Params);
+	static UE_API void PCV_GatherAnimSequencesFromStruct(TArray<FPCV_ReferencedAnimSequence>& OutRefAnimSequences, FAnimBPCompileValidationParams& PCV_Params, const UStruct* InStruct, const void* InData, TArray<FPCV_PropertyAndValue> InPropertyCallChain);
+	static UE_API void PCV_GatherAnimSequencesFromProperty(TArray<FPCV_ReferencedAnimSequence>& OutRefAnimSequences, FAnimBPCompileValidationParams& PCV_Params, const FProperty* InProperty, const void* InData, TArray<FPCV_PropertyAndValue> InPropertyCallChain);
 
 private:
 	virtual bool NeedsLoadForClient() const override { return false; }
 	virtual bool NeedsLoadForServer() const override { return false; }
 };
+
+#undef UE_API
diff --git a/Engine/Source/Editor/AnimGraph/Public/AnimGraphAttributes.h b/Engine/Source/Editor/AnimGraph/Public/AnimGraphAttributes.h
--- a/Engine/Source/Editor/AnimGraph/Public/AnimGraphAttributes.h
+++ b/Engine/Source/Editor/AnimGraph/Public/AnimGraphAttributes.h
@@ -15,6 +15,8 @@
 
 #include "AnimGraphAttributes.generated.h"
 
+#define UE_API ANIMGRAPH_API
+
 struct FPropertyChangedEvent;
 template <typename FuncType> class TFunctionRef;
 
@@ -89,27 +91,27 @@ struct FAnimGraphAttributeDesc
 	int32 SortOrder = INDEX_NONE;
 };
 
-UCLASS(Config = EditorPerProjectUserSettings)
-class ANIMGRAPH_API UAnimGraphAttributes : public UObject
+UCLASS(MinimalAPI, Config = EditorPerProjectUserSettings)
+class UAnimGraphAttributes : public UObject
 {
 	GENERATED_BODY()
 
 public:
 	// Register an anim graph attribute
-	void Register(const FAnimGraphAttributeDesc& InDesc);
+	UE_API void Register(const FAnimGraphAttributeDesc& InDesc);
 
 	// Get a previously registered attribute desc
 	// @param	InName	The name of a previously registered attribute
 	// @return nullptr if no desc is registered with the specified name
-	const FAnimGraphAttributeDesc* FindAttributeDesc(FName InName) const;
+	UE_API const FAnimGraphAttributeDesc* FindAttributeDesc(FName InName) const;
 
 	// Iterate over all attributes that are registered
 	// @param	InFunction	FUnction receiving the various attribute descriptors
-	void ForEachAttribute(TFunctionRef<void(const FAnimGraphAttributeDesc&)> InFunction) const;
+	UE_API void ForEachAttribute(TFunctionRef<void(const FAnimGraphAttributeDesc&)> InFunction) const;
 
 private:
 	// UObject interface
-	virtual void PostEditChangeProperty(FPropertyChangedEvent& PropertyChangedEvent) override;
+	UE_API virtual void PostEditChangeProperty(FPropertyChangedEvent& PropertyChangedEvent) override;
 
 private:
 	// All attributes
@@ -118,4 +120,6 @@ private:
 
 	// Cached attribute index map for faster lookup
 	TMap<FName, int32> AttributesMap;
-};
\ No newline at end of file
+};
+
+#undef UE_API
diff --git a/Engine/Source/Editor/AnimGraph/Public/AnimGraphNode_AssetPlayerBase.h b/Engine/Source/Editor/AnimGraph/Public/AnimGraphNode_AssetPlayerBase.h
--- a/Engine/Source/Editor/AnimGraph/Public/AnimGraphNode_AssetPlayerBase.h
+++ b/Engine/Source/Editor/AnimGraph/Public/AnimGraphNode_AssetPlayerBase.h
@@ -9,6 +9,8 @@
 #include "EdGraph/EdGraphNodeUtils.h"
 #include "AnimGraphNode_AssetPlayerBase.generated.h"
 
+#define UE_API ANIMGRAPH_API
+
 /** Get the default anim node class for playing a particular asset */
 ANIMGRAPH_API UClass* GetNodeClassForAsset(const UClass* AssetClass);
 
@@ -16,8 +18,8 @@ ANIMGRAPH_API UClass* GetNodeClassForAsset(const UClass* AssetClass);
 ANIMGRAPH_API bool SupportNodeClassForAsset(const UClass* AssetClass, UClass* NodeClass);
 
 /** Helper / intermediate for asset player graphical nodes */
-UCLASS(Abstract)
-class ANIMGRAPH_API UAnimGraphNode_AssetPlayerBase : public UAnimGraphNode_Base
+UCLASS(MinimalAPI, Abstract)
+class UAnimGraphNode_AssetPlayerBase : public UAnimGraphNode_Base
 {
 	GENERATED_BODY()
 public:
@@ -26,28 +28,28 @@ public:
 	FAnimationGroupReference SyncGroup_DEPRECATED;
 
 	/** UObject interface */
-	void Serialize(FArchive& Ar) override;
-	void PostEditChangeProperty(FPropertyChangedEvent& PropertyChangedEvent) override;
+	UE_API void Serialize(FArchive& Ar) override;
+	UE_API void PostEditChangeProperty(FPropertyChangedEvent& PropertyChangedEvent) override;
 
 	/** UEdGraphNode interface */
-	virtual void PinConnectionListChanged(UEdGraphPin* Pin) override;
-	virtual void PinDefaultValueChanged(UEdGraphPin* Pin) override;
-	virtual FText GetTooltipText() const override;
+	UE_API virtual void PinConnectionListChanged(UEdGraphPin* Pin) override;
+	UE_API virtual void PinDefaultValueChanged(UEdGraphPin* Pin) override;
+	UE_API virtual FText GetTooltipText() const override;
 
 	/** UK2Node interface */
-	virtual bool IsActionFilteredOut(class FBlueprintActionFilter const& Filter) override;
+	UE_API virtual bool IsActionFilteredOut(class FBlueprintActionFilter const& Filter) override;
 
 	/** UAnimGraphNode_Base interface */
-	virtual void OnProcessDuringCompilation(IAnimBlueprintCompilationContext& InCompilationContext, IAnimBlueprintGeneratedClassCompiledData& OutCompiledData) override;
-	virtual void GetOutputLinkAttributes(FNodeAttributeArray& OutAttributes) const override;
-	virtual void ValidateAnimNodeDuringCompilation(USkeleton* ForSkeleton, FCompilerResultsLog& MessageLog) override;
+	UE_API virtual void OnProcessDuringCompilation(IAnimBlueprintCompilationContext& InCompilationContext, IAnimBlueprintGeneratedClassCompiledData& OutCompiledData) override;
+	UE_API virtual void GetOutputLinkAttributes(FNodeAttributeArray& OutAttributes) const override;
+	UE_API virtual void ValidateAnimNodeDuringCompilation(USkeleton* ForSkeleton, FCompilerResultsLog& MessageLog) override;
 	virtual void SetAnimationAsset(UAnimationAsset* Asset) { check(false); /*Base function called*/ }
 
 	// Override this to copy any relevant settings from the animation asset used when creating a node with a drag/drop operation
 	virtual void CopySettingsFromAnimationAsset(UAnimationAsset* Asset) {}
 
 	// Helper function to gather menu actions from specific asset types supported by this node
-	static void GetMenuActionsHelper(
+	static UE_API void GetMenuActionsHelper(
 		FBlueprintActionDatabaseRegistrar& InActionRegistrar,
 		TSubclassOf<UAnimGraphNode_Base> InNodeClass,
 		const TArray<TSubclassOf<UObject>>& InAssetTypes,
@@ -59,20 +61,22 @@ public:
 		const TFunction<FText(const FAssetData&)>& InMenuCategoryFunction = nullptr);
 
 	// Helper function to validate player nodes
-	void ValidateAnimNodeDuringCompilationHelper(USkeleton* ForSkeleton, FCompilerResultsLog& MessageLog, UAnimationAsset* InAsset, TSubclassOf<UAnimationAsset> InAssetType, UEdGraphPin* InExposedPin, FName InPropertyName);
+	UE_API void ValidateAnimNodeDuringCompilationHelper(USkeleton* ForSkeleton, FCompilerResultsLog& MessageLog, UAnimationAsset* InAsset, TSubclassOf<UAnimationAsset> InAssetType, UEdGraphPin* InExposedPin, FName InPropertyName);
 
 	// Helper function to preload required assets (including those referenced on pins)
-	void PreloadRequiredAssetsHelper(UAnimationAsset* InAsset, UEdGraphPin* InExposedPin);
+	UE_API void PreloadRequiredAssetsHelper(UAnimationAsset* InAsset, UEdGraphPin* InExposedPin);
 	
 protected:
 	// Helper functions to build a title for an asset player node
-	FText GetNodeTitleHelper(ENodeTitleType::Type InTitleType, UEdGraphPin* InAssetPin, const FText& InAssetDesc, const TFunction<FText(UAnimationAsset*)> InPostFixFunctionRef = nullptr) const;
-	FText GetNodeTitleForAsset(ENodeTitleType::Type InTitleType, UAnimationAsset* InAsset, const FText& InAssetDesc, const TFunction<FText(UAnimationAsset*)> InPostFixFunctionRef = nullptr) const;
-	void AddSyncGroupToNodeTitle(ENodeTitleType::Type InTitleType, FText & InOutTitle) const;
+	UE_API FText GetNodeTitleHelper(ENodeTitleType::Type InTitleType, UEdGraphPin* InAssetPin, const FText& InAssetDesc, const TFunction<FText(UAnimationAsset*)> InPostFixFunctionRef = nullptr) const;
+	UE_API FText GetNodeTitleForAsset(ENodeTitleType::Type InTitleType, UAnimationAsset* InAsset, const FText& InAssetDesc, const TFunction<FText(UAnimationAsset*)> InPostFixFunctionRef = nullptr) const;
+	UE_API void AddSyncGroupToNodeTitle(ENodeTitleType::Type InTitleType, FText & InOutTitle) const;
 	
 	// Default setup function that can be used with GetMenuActionsHalper
-	static void SetupNewNode(UEdGraphNode* InNewNode, bool bInIsTemplateNode, const FAssetData InAssetData);
+	static UE_API void SetupNewNode(UEdGraphNode* InNewNode, bool bInIsTemplateNode, const FAssetData InAssetData);
 
 	/** Used for filtering in the Blueprint context menu when the sequence asset this node uses is unloaded */
 	FString UnloadedSkeletonName;
 };
+
+#undef UE_API
diff --git a/Engine/Source/Editor/AnimGraph/Public/AnimGraphNode_Base.h b/Engine/Source/Editor/AnimGraph/Public/AnimGraphNode_Base.h
--- a/Engine/Source/Editor/AnimGraph/Public/AnimGraphNode_Base.h
+++ b/Engine/Source/Editor/AnimGraph/Public/AnimGraphNode_Base.h
@@ -16,6 +16,8 @@
 
 #include "AnimGraphNode_Base.generated.h"
 
+#define UE_API ANIMGRAPH_API
+
 class FAnimGraphNodeDetails;
 class FBlueprintActionDatabaseRegistrar;
 class FCanvas;
@@ -194,8 +196,8 @@ struct FAnimGraphNodePropertyBinding
   *
   * Any concrete implementations will be paired with a runtime graph node derived from FAnimNode_Base
   */
-UCLASS(Abstract, BlueprintType)
-class ANIMGRAPH_API UAnimGraphNode_Base : public UK2Node
+UCLASS(MinimalAPI, Abstract, BlueprintType)
+class UAnimGraphNode_Base : public UK2Node
 {
 	GENERATED_UCLASS_BODY()
 
@@ -236,50 +238,50 @@ private:
 
 public:
 	// UObject interface
-	virtual void PostEditChangeProperty(FPropertyChangedEvent& PropertyChangedEvent) override;
-	virtual void PreEditChange(FProperty* PropertyAboutToChange) override;
-	virtual void PostEditUndo() override;
-	virtual void Serialize(FArchive& Ar) override;
+	UE_API virtual void PostEditChangeProperty(FPropertyChangedEvent& PropertyChangedEvent) override;
+	UE_API virtual void PreEditChange(FProperty* PropertyAboutToChange) override;
+	UE_API virtual void PostEditUndo() override;
+	UE_API virtual void Serialize(FArchive& Ar) override;
 	// End of UObject interface
 
 	// UEdGraphNode interface
-	virtual void AllocateDefaultPins() override;
-	virtual FLinearColor GetNodeTitleColor() const override;
-	virtual FString GetDocumentationLink() const override;
-	virtual void GetPinHoverText(const UEdGraphPin& Pin, FString& HoverTextOut) const override;
+	UE_API virtual void AllocateDefaultPins() override;
+	UE_API virtual FLinearColor GetNodeTitleColor() const override;
+	UE_API virtual FString GetDocumentationLink() const override;
+	UE_API virtual void GetPinHoverText(const UEdGraphPin& Pin, FString& HoverTextOut) const override;
 	virtual bool ShowPaletteIconOnNode() const override{ return false; }
-	virtual void PinDefaultValueChanged(UEdGraphPin* Pin) override;
-	virtual FString GetPinMetaData(FName InPinName, FName InKey) override;
-	virtual void AddSearchMetaDataInfo(TArray<struct FSearchTagDataPair>& OutTaggedMetaData) const override;
-	virtual void AddPinSearchMetaDataInfo(const UEdGraphPin* Pin, TArray<struct FSearchTagDataPair>& OutTaggedMetaData) const override;
-	virtual void PinConnectionListChanged(UEdGraphPin* Pin) override;
-	virtual void AutowireNewNode(UEdGraphPin* FromPin) override;
-	virtual void PostPlacedNewNode() override;
-	virtual void PostPasteNode() override;
-	virtual void DestroyNode() override;
+	UE_API virtual void PinDefaultValueChanged(UEdGraphPin* Pin) override;
+	UE_API virtual FString GetPinMetaData(FName InPinName, FName InKey) override;
+	UE_API virtual void AddSearchMetaDataInfo(TArray<struct FSearchTagDataPair>& OutTaggedMetaData) const override;
+	UE_API virtual void AddPinSearchMetaDataInfo(const UEdGraphPin* Pin, TArray<struct FSearchTagDataPair>& OutTaggedMetaData) const override;
+	UE_API virtual void PinConnectionListChanged(UEdGraphPin* Pin) override;
+	UE_API virtual void AutowireNewNode(UEdGraphPin* FromPin) override;
+	UE_API virtual void PostPlacedNewNode() override;
+	UE_API virtual void PostPasteNode() override;
+	UE_API virtual void DestroyNode() override;
 	// End of UEdGraphNode interface
 
 	// UK2Node interface
-	virtual void ReconstructNode() override;
+	UE_API virtual void ReconstructNode() override;
 	virtual bool NodeCausesStructuralBlueprintChange() const override { return true; }
 	virtual bool ShouldShowNodeProperties() const override { return true; }
 	virtual bool CanPlaceBreakpoints() const override { return false; }
-	virtual void ReallocatePinsDuringReconstruction(TArray<UEdGraphPin*>& OldPins) override;
-	virtual bool CanCreateUnderSpecifiedSchema(const UEdGraphSchema* DesiredSchema) const override;
-	virtual void GetNodeAttributes(TArray<TKeyValuePair<FString, FString>>& OutNodeAttributes) const override;
-	virtual void GetMenuActions(FBlueprintActionDatabaseRegistrar& ActionRegistrar) const override;
-	virtual FText GetMenuCategory() const override;
-	virtual void ExpandNode(class FKismetCompilerContext& CompilerContext, UEdGraph* SourceGraph) override;
+	UE_API virtual void ReallocatePinsDuringReconstruction(TArray<UEdGraphPin*>& OldPins) override;
+	UE_API virtual bool CanCreateUnderSpecifiedSchema(const UEdGraphSchema* DesiredSchema) const override;
+	UE_API virtual void GetNodeAttributes(TArray<TKeyValuePair<FString, FString>>& OutNodeAttributes) const override;
+	UE_API virtual void GetMenuActions(FBlueprintActionDatabaseRegistrar& ActionRegistrar) const override;
+	UE_API virtual FText GetMenuCategory() const override;
+	UE_API virtual void ExpandNode(class FKismetCompilerContext& CompilerContext, UEdGraph* SourceGraph) override;
 
 	// By default return any animation assets we have
 	virtual UObject* GetJumpTargetForDoubleClick() const override { return GetAnimationAsset(); }
-	virtual bool CanJumpToDefinition() const override;
-	virtual void JumpToDefinition() const override;
-	virtual void HandleVariableRenamed(UBlueprint* InBlueprint, UClass* InVariableClass, UEdGraph* InGraph, const FName& InOldVarName, const FName& InNewVarName) override;
-	virtual void HandleFunctionRenamed(UBlueprint* InBlueprint, UClass* InFunctionClass, UEdGraph* InGraph, const FName& InOldFuncName, const FName& InNewFuncName) override;
-	virtual void ReplaceReferences(UBlueprint* InBlueprint, UBlueprint* InReplacementBlueprint, const FMemberReference& InSource, const FMemberReference& InReplacement) override;
-	virtual bool ReferencesVariable(const FName& InVarName, const UStruct* InScope) const override;
-	virtual bool ReferencesFunction(const FName& InFunctionName, const UStruct* InScope) const override;
+	UE_API virtual bool CanJumpToDefinition() const override;
+	UE_API virtual void JumpToDefinition() const override;
+	UE_API virtual void HandleVariableRenamed(UBlueprint* InBlueprint, UClass* InVariableClass, UEdGraph* InGraph, const FName& InOldVarName, const FName& InNewVarName) override;
+	UE_API virtual void HandleFunctionRenamed(UBlueprint* InBlueprint, UClass* InFunctionClass, UEdGraph* InGraph, const FName& InOldFuncName, const FName& InNewFuncName) override;
+	UE_API virtual void ReplaceReferences(UBlueprint* InBlueprint, UBlueprint* InReplacementBlueprint, const FMemberReference& InSource, const FMemberReference& InReplacement) override;
+	UE_API virtual bool ReferencesVariable(const FName& InVarName, const UStruct* InScope) const override;
+	UE_API virtual bool ReferencesFunction(const FName& InFunctionName, const UStruct* InScope) const override;
 	// End of UK2Node interface
 
 	// UAnimGraphNode_Base interface
@@ -288,19 +290,19 @@ public:
 	virtual bool IsPoseWatchable() const { return true; }
 
 	// Gets the menu category this node belongs in
-	virtual FString GetNodeCategory() const;
+	UE_API virtual FString GetNodeCategory() const;
 
 	// Is this node a sink that has no pose outputs?
 	virtual bool IsSinkNode() const { return false; }
 
 	// Create any output pins necessary for this node
-	virtual void CreateOutputPins();
+	UE_API virtual void CreateOutputPins();
 
 	// customize pin data based on the input
 	virtual void CustomizePinData(UEdGraphPin* Pin, FName SourcePropertyName, int32 ArrayIndex) const {}
 
 	// Gives each visual node a chance to do final validation before it's node is harvested for use at runtime
-	virtual void ValidateAnimNodeDuringCompilation(USkeleton* ForSkeleton, FCompilerResultsLog& MessageLog);
+	UE_API virtual void ValidateAnimNodeDuringCompilation(USkeleton* ForSkeleton, FCompilerResultsLog& MessageLog);
 
 	// Gives each visual node a chance to validate that they are still valid in the context of the compiled class, giving a last shot at error or warning generation after primary compilation is finished
 	virtual void ValidateAnimNodePostCompile(FCompilerResultsLog& MessageLog, UAnimBlueprintGeneratedClass* CompiledClass, int32 CompiledNodeIndex) {}
@@ -312,7 +314,7 @@ public:
 	virtual void BakeDataDuringCompilation(FCompilerResultsLog& MessageLog) {}
 
 	// Give the node a chance to change the display name of a pin
-	virtual void PostProcessPinName(const UEdGraphPin* Pin, FString& DisplayName) const;
+	UE_API virtual void PostProcessPinName(const UEdGraphPin* Pin, FString& DisplayName) const;
 
 	/** Get the animation blueprint to which this node belongs */
 	UAnimBlueprint* GetAnimBlueprint() const { return CastChecked<UAnimBlueprint>(GetBlueprint()); }
@@ -324,7 +326,7 @@ public:
 	virtual void ReplaceReferredAnimations(const TMap<UAnimationAsset*, UAnimationAsset*>& AnimAssetReplacementMap) {};
 
 	// Helper function for GetAllAnimationSequencesReferred
-	void HandleAnimReferenceCollection(UAnimationAsset* AnimAsset, TArray<UAnimationAsset*>& AnimationAssets) const;
+	UE_API void HandleAnimReferenceCollection(UAnimationAsset* AnimAsset, TArray<UAnimationAsset*>& AnimationAssets) const;
 
 	// Helper function for ReplaceReferredAnimations	
 	template<class AssetType>
@@ -339,16 +341,16 @@ public:
 	 * @param	InModeTools		The mode tools. Use this to push the editor mode if required.
 	 * @param	InRuntimeNode	The runtime node to go with this skeletal control. This may be NULL in some cases when bInIsSelected is false.
 	 */
-	virtual void OnNodeSelected(bool bInIsSelected, class FEditorModeTools& InModeTools, struct FAnimNode_Base* InRuntimeNode);
+	UE_API virtual void OnNodeSelected(bool bInIsSelected, class FEditorModeTools& InModeTools, struct FAnimNode_Base* InRuntimeNode);
 
 	/** Pose Watch change notification callback. Should be called every time a pose watch on this node is created or destroyed. */
-	virtual void OnPoseWatchChanged(const bool IsPoseWatchEnabled, TObjectPtr<class UPoseWatch> InPoseWatch, FEditorModeTools& InModeTools, FAnimNode_Base* InRuntimeNode);
+	UE_API virtual void OnPoseWatchChanged(const bool IsPoseWatchEnabled, TObjectPtr<class UPoseWatch> InPoseWatch, FEditorModeTools& InModeTools, FAnimNode_Base* InRuntimeNode);
 
 	/**
 	 * Override this function to push an editor mode when this node is selected
 	 * @return the editor mode to use when this node is selected
 	 */
-	virtual FEditorModeID GetEditorMode() const;
+	UE_API virtual FEditorModeID GetEditorMode() const;
 
 	// Draw function for supporting visualization
 	virtual void Draw(FPrimitiveDrawInterface* PDI, USkeletalMeshComponent * PreviewSkelMeshComp) const {}
@@ -357,7 +359,7 @@ public:
 	 *	Default implementation calls the basic draw function for selected nodes and does nothing for pose watched nodes. Nodes
 	 *	that should render something when a pose watch is enabled but they are not selected should override this function.
 	 */
-	virtual void Draw(FPrimitiveDrawInterface* PDI, USkeletalMeshComponent* PreviewSkelMeshComp, const bool bIsSelected, const bool bIsPoseWatchEnabled) const;
+	UE_API virtual void Draw(FPrimitiveDrawInterface* PDI, USkeletalMeshComponent* PreviewSkelMeshComp, const bool bIsSelected, const bool bIsPoseWatchEnabled) const;
 
 	// Canvas draw function to draw to viewport
 	virtual void DrawCanvas(FViewport& InViewport, FSceneView& View, FCanvas& Canvas, USkeletalMeshComponent * PreviewSkelMeshComp) const {}
@@ -386,7 +388,7 @@ public:
 	virtual void CustomizeDetails(IDetailLayoutBuilder& DetailBuilder){ }
 
 	/** Try to find the preview node instance for this anim graph node */
-	FAnimNode_Base* FindDebugAnimNode(USkeletalMeshComponent * PreviewSkelMeshComp) const;
+	UE_API FAnimNode_Base* FindDebugAnimNode(USkeletalMeshComponent * PreviewSkelMeshComp) const;
 
 	template<typename NodeType>
 	NodeType* GetActiveInstanceNode(UObject* AnimInstanceObject) const
@@ -408,7 +410,7 @@ public:
 	 *	Returns whether this node supports the supplied asset class
 	 *	@param	bPrimaryAssetHandler	Is this the 'primary' handler for this asset (the node that should be created when asset is dropped)
 	 */
-	virtual EAnimAssetHandlerType SupportsAssetClass(const UClass* AssetClass) const;
+	UE_API virtual EAnimAssetHandlerType SupportsAssetClass(const UClass* AssetClass) const;
 
 	// Event that observers can bind to so that they are notified about changes
 	// made to this node through the property system
@@ -425,7 +427,7 @@ public:
 	 * @param	InPinDirection	The direction of the pin we are looking for. If this is EGPD_MAX, all directions are considered
 	 * @return true if the pin is present and connected
 	 */
-	bool IsPinExposedAndLinked(const FString& InPinName, const EEdGraphPinDirection Direction = EGPD_MAX) const;
+	UE_API bool IsPinExposedAndLinked(const FString& InPinName, const EEdGraphPinDirection Direction = EGPD_MAX) const;
 
 	/**
 	 * Helper function to check whether a pin is valid and bound via property access
@@ -433,7 +435,7 @@ public:
 	 * @param	InPinDirection	The direction of the pin we are looking for. If this is EGPD_MAX, all directions are considered
 	 * @return true if the pin is present and bound
 	 */
-	bool IsPinExposedAndBound(const FString& InPinName, const EEdGraphPinDirection InDirection = EGPD_MAX) const;
+	UE_API bool IsPinExposedAndBound(const FString& InPinName, const EEdGraphPinDirection InDirection = EGPD_MAX) const;
 
 	/**
 	 * Helper function to check whether a pin is not linked, not bound via property access and still has its default value
@@ -441,7 +443,7 @@ public:
 	 * @param	InPinDirection	The direction of the pin we are looking for. If this is EGPD_MAX, all directions are considered
 	 * @return true if the pin is unlinked, unbound and still has its default value
 	 */
-	bool IsPinUnlinkedUnboundAndUnset(const FString& InPinName, const EEdGraphPinDirection InDirection) const;
+	UE_API bool IsPinUnlinkedUnboundAndUnset(const FString& InPinName, const EEdGraphPinDirection InDirection) const;
 
 	// Event that is broadcast to inform observers that the node title has changed
 	// The default SAnimationGraphNode uses this to invalidate cached node title text
@@ -465,7 +467,7 @@ public:
 	virtual const UAnimGraphNode_Base* GetProxyNodeForAttributes() const { return this; }
 
 	// Sets the visibility of the specified pin, reconstructs the node if it changes
-	void SetPinVisibility(bool bInVisible, int32 InOptionalPinIndex);
+	UE_API void SetPinVisibility(bool bInVisible, int32 InOptionalPinIndex);
 
 	// Arguments used to construct a property binding widget
 	struct FAnimPropertyBindingWidgetArgs
@@ -515,20 +517,20 @@ public:
 	};
 	
 	// Make a property binding widget to edit the bindings of the passed-in nodes
-	static TSharedRef<SWidget> MakePropertyBindingWidget(const FAnimPropertyBindingWidgetArgs& InArgs);
+	static UE_API TSharedRef<SWidget> MakePropertyBindingWidget(const FAnimPropertyBindingWidgetArgs& InArgs);
 
 	// Get the property corresponding to a pin. For array element pins returns the outer array property. Returns null if a property cannot be found.
-	FProperty* GetPinProperty(const UEdGraphPin* InPin) const;
-	virtual FProperty* GetPinProperty(FName InPinName) const;
+	UE_API FProperty* GetPinProperty(const UEdGraphPin* InPin) const;
+	UE_API virtual FProperty* GetPinProperty(FName InPinName) const;
 	
 	// Check whether the named pin is bindable
-	virtual bool IsPinBindable(const UEdGraphPin* InPin) const;
+	UE_API virtual bool IsPinBindable(const UEdGraphPin* InPin) const;
 
 	// Get the tag for this node, if any
 	FName GetTag() const { return Tag; }
 
 	// Set the tag for this node
-	void SetTag(FName InTag);
+	UE_API void SetTag(FName InTag);
 
 	// Get the currently-debugged runtime anim node (in the anim BP debugger that this node is currently being edited in)
 	// @return nullptr if the node cannot be found
@@ -541,20 +543,20 @@ public:
 	// Refreshes the debugged component post-edit.
 	// This is required to see changes as the component may be either an editor-only component that is not ticking,
 	// or in a paused PIE world
-	void PostEditRefreshDebuggedComponent();
+	UE_API void PostEditRefreshDebuggedComponent();
 
 	// Gets editor information for all the bound anim node functions (category metadata string, member variable's name)
 	// Used by SAnimGraphNode to display all bound functions of an anim node.
-	virtual void GetBoundFunctionsInfo(TArray<TPair<FName, FName>> & InOutBindingsInfo);
+	UE_API virtual void GetBoundFunctionsInfo(TArray<TPair<FName, FName>> & InOutBindingsInfo);
 	
 	// Check if a specified function reference appears to be valid by inspecting only the validity of the name and guid
-	static bool IsPotentiallyBoundFunction(const FMemberReference& FunctionReference);
+	static UE_API bool IsPotentiallyBoundFunction(const FMemberReference& FunctionReference);
 
 	// Check whether the specified property is bound via PropertyBindings
-	virtual bool HasBinding(FName InPropertyName) const;
+	UE_API virtual bool HasBinding(FName InPropertyName) const;
 
 	// Update the property binding name
-	bool RedirectBinding(FName OldPropertyBindingName, FName NewPropertyBindingName);
+	UE_API bool RedirectBinding(FName OldPropertyBindingName, FName NewPropertyBindingName);
 	
 	// Get the bindings for this node
 	const UAnimGraphNodeBinding* GetBinding() const { return Binding; }
@@ -563,16 +565,16 @@ public:
 	UAnimGraphNodeBinding* GetMutableBinding() { return Binding; }
 
 	// Remove any bindings for the specified name
-	void RemoveBindings(FName InBindingName);
+	UE_API void RemoveBindings(FName InBindingName);
 
 	// Gets the animation FNode type represented by this ed graph node
-	UScriptStruct* GetFNodeType() const;
+	UE_API UScriptStruct* GetFNodeType() const;
 
 	// Gets the animation FNode property represented by this ed graph node
-	FStructProperty* GetFNodeProperty() const;
+	UE_API FStructProperty* GetFNodeProperty() const;
 
 	// Get the runtime anim node that we template
-	FAnimNode_Base* GetFNode();
+	UE_API FAnimNode_Base* GetFNode();
 
 protected:
 	friend class FAnimBlueprintCompilerContext;
@@ -591,29 +593,29 @@ protected:
 	// This will be called when a pose link is found, and can be called with PoseProperty being either of:
 	//  - an array property (ArrayIndex >= 0)
 	//  - a single pose property (ArrayIndex == INDEX_NONE)
-	virtual void CreatePinsForPoseLink(FProperty* PoseProperty, int32 ArrayIndex);
+	UE_API virtual void CreatePinsForPoseLink(FProperty* PoseProperty, int32 ArrayIndex);
 
 	//
-	virtual FPoseLinkMappingRecord GetLinkIDLocation(const UScriptStruct* NodeType, UEdGraphPin* InputLinkPin);
+	UE_API virtual FPoseLinkMappingRecord GetLinkIDLocation(const UScriptStruct* NodeType, UEdGraphPin* InputLinkPin);
 
 	/** Get the property (and possibly array index) associated with the supplied pin */
-	virtual void GetPinAssociatedProperty(const UScriptStruct* NodeType, const UEdGraphPin* InputPin, FProperty*& OutProperty, int32& OutIndex) const;
+	UE_API virtual void GetPinAssociatedProperty(const UScriptStruct* NodeType, const UEdGraphPin* InputPin, FProperty*& OutProperty, int32& OutIndex) const;
 
 	// Process this node's data during compilation
-	void ProcessDuringCompilation(IAnimBlueprintCompilationContext& InCompilationContext, IAnimBlueprintGeneratedClassCompiledData& OutCompiledData);
+	UE_API void ProcessDuringCompilation(IAnimBlueprintCompilationContext& InCompilationContext, IAnimBlueprintGeneratedClassCompiledData& OutCompiledData);
 
 	// Process this node's data during compilation (override point)
 	virtual void OnProcessDuringCompilation(IAnimBlueprintCompilationContext& InCompilationContext, IAnimBlueprintGeneratedClassCompiledData& OutCompiledData) {}
 
 	// Copy this node's data during the last phase of compilation where term defaults are copied to the new CDO
-	void CopyTermDefaultsToDefaultObject(IAnimBlueprintCopyTermDefaultsContext& InCompilationContext, IAnimBlueprintNodeCopyTermDefaultsContext& InPerNodeContext, IAnimBlueprintGeneratedClassCompiledData& OutCompiledData);
+	UE_API void CopyTermDefaultsToDefaultObject(IAnimBlueprintCopyTermDefaultsContext& InCompilationContext, IAnimBlueprintNodeCopyTermDefaultsContext& InPerNodeContext, IAnimBlueprintGeneratedClassCompiledData& OutCompiledData);
 
 	// Override point for CopyTermDefaultsToDefaultObject
 	// Copy this node's data during the last phase of compilation where term defaults are copied to the new CDO
 	virtual void OnCopyTermDefaultsToDefaultObject(IAnimBlueprintCopyTermDefaultsContext& InCompilationContext, IAnimBlueprintNodeCopyTermDefaultsContext& InPerNodeContext, IAnimBlueprintGeneratedClassCompiledData& OutCompiledData) {}
 
 	// Called to override the assets held on a runtime anim node. Implements per-node logic for child anim blueprints.
-	void OverrideAssets(IAnimBlueprintNodeOverrideAssetsContext& InContext) const;
+	UE_API void OverrideAssets(IAnimBlueprintNodeOverrideAssetsContext& InContext) const;
 
 	// Override point for OverrideAssets
 	virtual void OnOverrideAssets(IAnimBlueprintNodeOverrideAssetsContext& InContext) const {}
@@ -622,7 +624,7 @@ protected:
 	virtual bool ShouldCreateStructEvalHandlers() const { return true; }
 	
 	// Allocates or reallocates pins
-	void InternalPinCreation(TArray<UEdGraphPin*>* OldPins);
+	UE_API void InternalPinCreation(TArray<UEdGraphPin*>* OldPins);
 
 	// Override point to create custom pins
 	// @param	OldPins		In the case of reconstruction, the OldPins array contains all the pins that the node had prior, otherwise if the
@@ -635,7 +637,7 @@ protected:
 	// @param	OutPinProperty	The property that the binding represents
 	// @param	OutOptionalPinIndex	The optional pin index (index into the ShowPinForProperties array) 
 	// @return false if the pin cannot be bound
-	virtual bool GetPinBindingInfo(FName InPinName, FName& OutBindingName, FProperty*& OutPinProperty, int32& OutOptionalPinIndex) const;
+	UE_API virtual bool GetPinBindingInfo(FName InPinName, FName& OutBindingName, FProperty*& OutPinProperty, int32& OutOptionalPinIndex) const;
 
 	FOnNodePropertyChangedEvent PropertyChangeEvent;
 
@@ -644,16 +646,16 @@ protected:
 	FPinVisibilityChangedEvent PinVisibilityChangedEvent;
 
 	// Helper function used to refresh the type of a binding
-	void RecalculateBindingType(FAnimGraphNodePropertyBinding& InBinding);
+	UE_API void RecalculateBindingType(FAnimGraphNodePropertyBinding& InBinding);
 	
 	/** @return the current object being debugged from the blueprint for this node. Can be nullptr. */
-	UObject* GetObjectBeingDebugged() const;
+	UE_API UObject* GetObjectBeingDebugged() const;
 
 	/** Helper function used to validate anim node function references */
-	void ValidateFunctionRef(FName InPropertyName, const FMemberReference& InRef, const FText& InFunctionName, FCompilerResultsLog& MessageLog);
+	UE_API void ValidateFunctionRef(FName InPropertyName, const FMemberReference& InRef, const FText& InFunctionName, FCompilerResultsLog& MessageLog);
 	
 	// Create the bindings subobject if required
-	void EnsureBindingsArePresent();
+	UE_API void EnsureBindingsArePresent();
 
 protected:
 	// Old shown pins. Needs to be a member variable to track pin visibility changes between Pre and PostEditChange 
@@ -690,3 +692,5 @@ template<class TNodeType> TNodeType* UAnimGraphNode_Base::GetDebuggedAnimNode()
 
 	return nullptr;
 }
+
+#undef UE_API
diff --git a/Engine/Source/Editor/AnimGraph/Public/AnimGraphNode_BlendListBase.h b/Engine/Source/Editor/AnimGraph/Public/AnimGraphNode_BlendListBase.h
--- a/Engine/Source/Editor/AnimGraph/Public/AnimGraphNode_BlendListBase.h
+++ b/Engine/Source/Editor/AnimGraph/Public/AnimGraphNode_BlendListBase.h
@@ -7,27 +7,31 @@
 #include "AnimGraphNode_Base.h"
 #include "AnimGraphNode_BlendListBase.generated.h"
 
-UCLASS(Abstract)
-class ANIMGRAPH_API UAnimGraphNode_BlendListBase : public UAnimGraphNode_Base
+#define UE_API ANIMGRAPH_API
+
+UCLASS(MinimalAPI, Abstract)
+class UAnimGraphNode_BlendListBase : public UAnimGraphNode_Base
 {
 	GENERATED_UCLASS_BODY()
 
-	virtual void Serialize(FArchive& Ar) override;
+	UE_API virtual void Serialize(FArchive& Ar) override;
 
 	// UEdGraphNode interface
-	virtual FLinearColor GetNodeTitleColor() const override;
-	virtual void PostEditChangeProperty(struct FPropertyChangedEvent& PropertyChangedEvent) override;
-	virtual void ReallocatePinsDuringReconstruction(TArray<UEdGraphPin*>& OldPins) override;
+	UE_API virtual FLinearColor GetNodeTitleColor() const override;
+	UE_API virtual void PostEditChangeProperty(struct FPropertyChangedEvent& PropertyChangedEvent) override;
+	UE_API virtual void ReallocatePinsDuringReconstruction(TArray<UEdGraphPin*>& OldPins) override;
 	// End of UEdGraphNode interface
 
 	// UAnimGraphNode_Base interface
-	virtual FString GetNodeCategory() const override;
-	virtual void GetOutputLinkAttributes(FNodeAttributeArray& OutAttributes) const override;
+	UE_API virtual FString GetNodeCategory() const override;
+	UE_API virtual void GetOutputLinkAttributes(FNodeAttributeArray& OutAttributes) const override;
 	// End of UAnimGraphNode_Base interface
 
 protected:
 	// removes removed pins and adjusts array indices of remained pins
-	void RemovePinsFromOldPins(TArray<UEdGraphPin*>& OldPins, int32 RemovedArrayIndex);
+	UE_API void RemovePinsFromOldPins(TArray<UEdGraphPin*>& OldPins, int32 RemovedArrayIndex);
 
 	int32 RemovedPinArrayIndex;
 };
+
+#undef UE_API
diff --git a/Engine/Source/Editor/AnimGraph/Public/AnimGraphNode_BlendSpaceGraph.h b/Engine/Source/Editor/AnimGraph/Public/AnimGraphNode_BlendSpaceGraph.h
--- a/Engine/Source/Editor/AnimGraph/Public/AnimGraphNode_BlendSpaceGraph.h
+++ b/Engine/Source/Editor/AnimGraph/Public/AnimGraphNode_BlendSpaceGraph.h
@@ -7,11 +7,13 @@
 #include "AnimNodes/AnimNode_BlendSpaceGraph.h"
 #include "AnimGraphNode_BlendSpaceGraph.generated.h"
 
+#define UE_API ANIMGRAPH_API
+
 class FBlueprintActionDatabaseRegistrar;
 class FCompilerResultsLog;
 
-UCLASS()
-class ANIMGRAPH_API UAnimGraphNode_BlendSpaceGraph : public UAnimGraphNode_BlendSpaceGraphBase
+UCLASS(MinimalAPI)
+class UAnimGraphNode_BlendSpaceGraph : public UAnimGraphNode_BlendSpaceGraphBase
 {
 	GENERATED_BODY()
 
@@ -20,11 +22,13 @@ private:
 	FAnimNode_BlendSpaceGraph Node;
 
 	// UEdGraphNode interface
-	virtual FText GetNodeTitle(ENodeTitleType::Type TitleType) const override;
+	UE_API virtual FText GetNodeTitle(ENodeTitleType::Type TitleType) const override;
 	
 	// UK2Node interface
-	virtual void GetMenuActions(FBlueprintActionDatabaseRegistrar& ActionRegistrar) const override;
+	UE_API virtual void GetMenuActions(FBlueprintActionDatabaseRegistrar& ActionRegistrar) const override;
 
 	// UAnimGraphNode_Base interface
-	virtual void BakeDataDuringCompilation(FCompilerResultsLog& MessageLog) override;
+	UE_API virtual void BakeDataDuringCompilation(FCompilerResultsLog& MessageLog) override;
 };
+
+#undef UE_API
