Updated headers using UnrealCodeFixup to make sure dllstorage is on methods/staticvars instead of types. Part 4/n

diff --git a/Engine/Source/Editor/AnimGraph/Public/AnimGraphNode_BlendSpaceGraphBase.h b/Engine/Source/Editor/AnimGraph/Public/AnimGraphNode_BlendSpaceGraphBase.h
--- a/Engine/Source/Editor/AnimGraph/Public/AnimGraphNode_BlendSpaceGraphBase.h
+++ b/Engine/Source/Editor/AnimGraph/Public/AnimGraphNode_BlendSpaceGraphBase.h
@@ -8,6 +8,8 @@
 #include "Containers/ArrayView.h"
 #include "AnimGraphNode_BlendSpaceGraphBase.generated.h"
 
+#define UE_API ANIMGRAPH_API
+
 class FBlueprintActionDatabaseRegistrar;
 class IAnimBlueprintCopyTermDefaultsContext;
 class IAnimBlueprintNodeCopyTermDefaultsContext;
@@ -15,13 +17,13 @@ class IAnimBlueprintGeneratedClassCompiledData;
 class UBlendSpaceGraph;
 class UAnimationBlendSpaceSampleGraph;
 
-UCLASS(Abstract)
-class ANIMGRAPH_API UAnimGraphNode_BlendSpaceGraphBase : public UAnimGraphNode_Base
+UCLASS(MinimalAPI, Abstract)
+class UAnimGraphNode_BlendSpaceGraphBase : public UAnimGraphNode_Base
 {
 	GENERATED_BODY()
 
 public:
-	UAnimGraphNode_BlendSpaceGraphBase();
+	UE_API UAnimGraphNode_BlendSpaceGraphBase();
 
 	// Access the graphs for each sample
 	TArrayView<UEdGraph* const> GetGraphs() const { return Graphs; }
@@ -30,38 +32,38 @@ public:
 	UBlendSpaceGraph* GetBlendSpaceGraph() const { return BlendSpaceGraph; }
 
 	// Adds a new graph to the internal array
-	UAnimationBlendSpaceSampleGraph* AddGraph(FName InSampleName, UAnimSequence* InSequence);
+	UE_API UAnimationBlendSpaceSampleGraph* AddGraph(FName InSampleName, UAnimSequence* InSequence);
 
 	/** Returns the sample index associated with the graph, or -1 if not found */
-	int32 GetSampleIndex(const UEdGraph* Graph) const;
+	UE_API int32 GetSampleIndex(const UEdGraph* Graph) const;
 
 	// Removes the graph at the specified index
-	void RemoveGraph(int32 InSampleIndex);
+	UE_API void RemoveGraph(int32 InSampleIndex);
 
 	// Replaces the graph at the specified index
-	void ReplaceGraph(int32 InSampleIndex, UAnimSequence* InSequence);
+	UE_API void ReplaceGraph(int32 InSampleIndex, UAnimSequence* InSequence);
 
 	// Setup this node from the specified asset
-	void SetupFromAsset(const FAssetData& InAssetData, bool bInIsTemplateNode);
+	UE_API void SetupFromAsset(const FAssetData& InAssetData, bool bInIsTemplateNode);
 
 	// UEdGraphNode interface
-	virtual void PostPlacedNewNode() override;
+	UE_API virtual void PostPlacedNewNode() override;
 
 	// @return the sync group name assigned to this node
-	FName GetSyncGroupName() const;
+	UE_API FName GetSyncGroupName() const;
 
 	// Set the sync group name assigned to this node
-	void SetSyncGroupName(FName InName);
+	UE_API void SetSyncGroupName(FName InName);
 
 protected:
 	// Get the name of the blendspace graph
-	FString GetBlendSpaceGraphName() const;
+	UE_API FString GetBlendSpaceGraphName() const;
 
 	// Get the name of the blendspace
-	FString GetBlendSpaceName() const;
+	UE_API FString GetBlendSpaceName() const;
 
 	// Setup this node from the specified class
-	void SetupFromClass(TSubclassOf<UBlendSpace> InBlendSpaceClass, bool bInIsTemplateNode);
+	UE_API void SetupFromClass(TSubclassOf<UBlendSpace> InBlendSpaceClass, bool bInIsTemplateNode);
 
 	// Internal blendspace
 	UPROPERTY()
@@ -84,34 +86,36 @@ protected:
 
 protected:
 	// UEdGraphNode interface
-	virtual FText GetMenuCategory() const override;
-	virtual FLinearColor GetNodeTitleColor() const override;
-	virtual FSlateIcon GetIconAndTint(FLinearColor& OutColor) const override;
-	virtual FText GetTooltipText() const override;
-	virtual UObject* GetJumpTargetForDoubleClick() const override;
-	virtual void JumpToDefinition() const override;
-	virtual TArray<UEdGraph*> GetSubGraphs() const override;
-	virtual void DestroyNode() override;
-	virtual void OnRenameNode(const FString& NewName) override;
-	virtual TSharedPtr<INameValidatorInterface> MakeNameValidator() const override;
-	virtual void PostPasteNode() override;
-	virtual void CustomizePinData(UEdGraphPin* Pin, FName SourcePropertyName, int32 ArrayIndex) const override;
-	virtual void PostProcessPinName(const UEdGraphPin* Pin, FString& DisplayName) const override;
+	UE_API virtual FText GetMenuCategory() const override;
+	UE_API virtual FLinearColor GetNodeTitleColor() const override;
+	UE_API virtual FSlateIcon GetIconAndTint(FLinearColor& OutColor) const override;
+	UE_API virtual FText GetTooltipText() const override;
+	UE_API virtual UObject* GetJumpTargetForDoubleClick() const override;
+	UE_API virtual void JumpToDefinition() const override;
+	UE_API virtual TArray<UEdGraph*> GetSubGraphs() const override;
+	UE_API virtual void DestroyNode() override;
+	UE_API virtual void OnRenameNode(const FString& NewName) override;
+	UE_API virtual TSharedPtr<INameValidatorInterface> MakeNameValidator() const override;
+	UE_API virtual void PostPasteNode() override;
+	UE_API virtual void CustomizePinData(UEdGraphPin* Pin, FName SourcePropertyName, int32 ArrayIndex) const override;
+	UE_API virtual void PostProcessPinName(const UEdGraphPin* Pin, FString& DisplayName) const override;
 
 	// UAnimGraphNode_Base interface
-	virtual void OnProcessDuringCompilation(IAnimBlueprintCompilationContext& InCompilationContext, IAnimBlueprintGeneratedClassCompiledData& OutCompiledData) override;
-	virtual void OnCopyTermDefaultsToDefaultObject(IAnimBlueprintCopyTermDefaultsContext& InCompilationContext, IAnimBlueprintNodeCopyTermDefaultsContext& InPerNodeContext, IAnimBlueprintGeneratedClassCompiledData& OutCompiledData) override;
-	virtual void CustomizeDetails(IDetailLayoutBuilder& InDetailBuilder) override;
-	virtual void GetInputLinkAttributes(FNodeAttributeArray& OutAttributes) const override;
-	virtual void GetRequiredExtensions(TArray<TSubclassOf<UAnimBlueprintExtension>>& OutExtensions) const override;
+	UE_API virtual void OnProcessDuringCompilation(IAnimBlueprintCompilationContext& InCompilationContext, IAnimBlueprintGeneratedClassCompiledData& OutCompiledData) override;
+	UE_API virtual void OnCopyTermDefaultsToDefaultObject(IAnimBlueprintCopyTermDefaultsContext& InCompilationContext, IAnimBlueprintNodeCopyTermDefaultsContext& InPerNodeContext, IAnimBlueprintGeneratedClassCompiledData& OutCompiledData) override;
+	UE_API virtual void CustomizeDetails(IDetailLayoutBuilder& InDetailBuilder) override;
+	UE_API virtual void GetInputLinkAttributes(FNodeAttributeArray& OutAttributes) const override;
+	UE_API virtual void GetRequiredExtensions(TArray<TSubclassOf<UAnimBlueprintExtension>>& OutExtensions) const override;
 
 	// UK2Node interface
-	virtual void PreloadRequiredAssets() override;
-	virtual bool IsActionFilteredOut(class FBlueprintActionFilter const& Filter) override;
+	UE_API virtual void PreloadRequiredAssets() override;
+	UE_API virtual bool IsActionFilteredOut(class FBlueprintActionFilter const& Filter) override;
 
 	// Helper function for compilation
-	UAnimGraphNode_Base* ExpandGraphAndProcessNodes(UEdGraph* SourceGraph, UAnimGraphNode_Base* SourceRootNode, IAnimBlueprintCompilationContext& InCompilationContext, IAnimBlueprintGeneratedClassCompiledData& OutCompiledData);
+	UE_API UAnimGraphNode_Base* ExpandGraphAndProcessNodes(UEdGraph* SourceGraph, UAnimGraphNode_Base* SourceRootNode, IAnimBlueprintCompilationContext& InCompilationContext, IAnimBlueprintGeneratedClassCompiledData& OutCompiledData);
 
 	// Helper function for AddGraph/ReplaceGraph - builds the new graph but doesn't add it to Graphs array.
-	UAnimationBlendSpaceSampleGraph* AddGraphInternal(FName InSampleName, UAnimSequence* InSequence);
+	UE_API UAnimationBlendSpaceSampleGraph* AddGraphInternal(FName InSampleName, UAnimSequence* InSequence);
 };
+
+#undef UE_API
diff --git a/Engine/Source/Editor/AnimGraph/Public/AnimGraphNode_BoneDrivenController.h b/Engine/Source/Editor/AnimGraph/Public/AnimGraphNode_BoneDrivenController.h
--- a/Engine/Source/Editor/AnimGraph/Public/AnimGraphNode_BoneDrivenController.h
+++ b/Engine/Source/Editor/AnimGraph/Public/AnimGraphNode_BoneDrivenController.h
@@ -10,6 +10,8 @@
 #include "BoneControllers/AnimNode_BoneDrivenController.h"
 #include "AnimGraphNode_BoneDrivenController.generated.h"
 
+#define UE_API ANIMGRAPH_API
+
 class FCompilerResultsLog;
 class FPrimitiveDrawInterface;
 class IDetailCategoryBuilder;
@@ -20,8 +22,8 @@ class USkeletalMeshComponent;
 /**
  * This is the 'source version' of a bone driven controller, which maps part of the state from one bone to another (e.g., 2 * source.x -> target.z)
  */
-UCLASS()
-class ANIMGRAPH_API UAnimGraphNode_BoneDrivenController : public UAnimGraphNode_SkeletalControlBase
+UCLASS(MinimalAPI)
+class UAnimGraphNode_BoneDrivenController : public UAnimGraphNode_SkeletalControlBase
 {
 	GENERATED_UCLASS_BODY()
 
@@ -30,39 +32,41 @@ class ANIMGRAPH_API UAnimGraphNode_BoneDrivenController : public UAnimGraphNode_
 
 public:
 	// UObject interface
-	virtual void Serialize(FArchive& Ar) override;
+	UE_API virtual void Serialize(FArchive& Ar) override;
 	// End of UObject interface
 
 	// UEdGraphNode interface
-	virtual FText GetNodeTitle(ENodeTitleType::Type TitleType) const override;
-	virtual FText GetTooltipText() const override;
+	UE_API virtual FText GetNodeTitle(ENodeTitleType::Type TitleType) const override;
+	UE_API virtual FText GetTooltipText() const override;
 	// End of UEdGraphNode interface
 
 	// UAnimGraphNode_Base interface
-	virtual void ValidateAnimNodeDuringCompilation(USkeleton* ForSkeleton, FCompilerResultsLog& MessageLog) override;
-	virtual void CustomizeDetails(IDetailLayoutBuilder& DetailBuilder) override;
+	UE_API virtual void ValidateAnimNodeDuringCompilation(USkeleton* ForSkeleton, FCompilerResultsLog& MessageLog) override;
+	UE_API virtual void CustomizeDetails(IDetailLayoutBuilder& DetailBuilder) override;
 	// End of UAnimGraphNode_Base interface
 
 	// UAnimGraphNode_SkeletalControlBase interface
-	virtual void Draw(FPrimitiveDrawInterface* PDI, USkeletalMeshComponent* SkelMeshComp) const override;
+	UE_API virtual void Draw(FPrimitiveDrawInterface* PDI, USkeletalMeshComponent* SkelMeshComp) const override;
 	// End of UAnimGraphNode_SkeletalControlBase interface
 
 protected:
 
 	// UAnimGraphNode_SkeletalControlBase protected interface
-	virtual FText GetControllerDescription() const override;
+	UE_API virtual FText GetControllerDescription() const override;
 	virtual const FAnimNode_SkeletalControlBase* GetNode() const override { return &Node; }
 	// End of UAnimGraphNode_SkeletalControlBase protected interface
 
 	// Should non-curve mapping values be shown (multiplier, range)?
-	static EVisibility AreNonCurveMappingValuesVisible(IDetailLayoutBuilder* DetailLayoutBuilder);
-	static EVisibility AreRemappingValuesVisible(IDetailLayoutBuilder* DetailLayoutBuilder);
+	static UE_API EVisibility AreNonCurveMappingValuesVisible(IDetailLayoutBuilder* DetailLayoutBuilder);
+	static UE_API EVisibility AreRemappingValuesVisible(IDetailLayoutBuilder* DetailLayoutBuilder);
 
 	// Should destination bone or morph target properties be visible
-	static EVisibility AreTargetBonePropertiesVisible(IDetailLayoutBuilder* DetailLayoutBuilder);
-	static EVisibility AreTargetCurvePropertiesVisible(IDetailLayoutBuilder* DetailLayoutBuilder);
+	static UE_API EVisibility AreTargetBonePropertiesVisible(IDetailLayoutBuilder* DetailLayoutBuilder);
+	static UE_API EVisibility AreTargetCurvePropertiesVisible(IDetailLayoutBuilder* DetailLayoutBuilder);
 
-	static void AddTripletPropertyRow(const FText& Name, const FText& Tooltip, IDetailCategoryBuilder& Category, TSharedRef<IPropertyHandle> PropertyHandle, const FName XPropertyName, const FName YPropertyName, const FName ZPropertyName, TAttribute<EVisibility> VisibilityAttribute);
-	static void AddRangePropertyRow(const FText& Name, const FText& Tooltip, IDetailCategoryBuilder& Category, TSharedRef<IPropertyHandle> PropertyHandle, const FName MinPropertyName, const FName MaxPropertyName, TAttribute<EVisibility> VisibilityAttribute);
-	static FText ComponentTypeToText(EComponentType::Type Component);
+	static UE_API void AddTripletPropertyRow(const FText& Name, const FText& Tooltip, IDetailCategoryBuilder& Category, TSharedRef<IPropertyHandle> PropertyHandle, const FName XPropertyName, const FName YPropertyName, const FName ZPropertyName, TAttribute<EVisibility> VisibilityAttribute);
+	static UE_API void AddRangePropertyRow(const FText& Name, const FText& Tooltip, IDetailCategoryBuilder& Category, TSharedRef<IPropertyHandle> PropertyHandle, const FName MinPropertyName, const FName MaxPropertyName, TAttribute<EVisibility> VisibilityAttribute);
+	static UE_API FText ComponentTypeToText(EComponentType::Type Component);
 };
+
+#undef UE_API
diff --git a/Engine/Source/Editor/AnimGraph/Public/AnimGraphNode_Constraint.h b/Engine/Source/Editor/AnimGraph/Public/AnimGraphNode_Constraint.h
--- a/Engine/Source/Editor/AnimGraph/Public/AnimGraphNode_Constraint.h
+++ b/Engine/Source/Editor/AnimGraph/Public/AnimGraphNode_Constraint.h
@@ -7,8 +7,10 @@
 #include "EdGraph/EdGraphNodeUtils.h" // for FNodeTitleTextTable
 #include "AnimGraphNode_Constraint.generated.h"
 
-UCLASS()
-class ANIMGRAPH_API UAnimGraphNode_Constraint : public UAnimGraphNode_SkeletalControlBase
+#define UE_API ANIMGRAPH_API
+
+UCLASS(MinimalAPI)
+class UAnimGraphNode_Constraint : public UAnimGraphNode_SkeletalControlBase
 {
 	GENERATED_UCLASS_BODY()
 
@@ -17,23 +19,23 @@ class ANIMGRAPH_API UAnimGraphNode_Constraint : public UAnimGraphNode_SkeletalCo
 
 public:
 	// UEdGraphNode interface
-	virtual FText GetNodeTitle(ENodeTitleType::Type TitleType) const override;
-	virtual FText GetTooltipText() const override;
+	UE_API virtual FText GetNodeTitle(ENodeTitleType::Type TitleType) const override;
+	UE_API virtual FText GetTooltipText() const override;
 	// End of UEdGraphNode interface
 
 	// UObject interface
-	virtual void PostEditChangeProperty(struct FPropertyChangedEvent& PropertyChangedEvent) override;
+	UE_API virtual void PostEditChangeProperty(struct FPropertyChangedEvent& PropertyChangedEvent) override;
 	// End of UObject interface
 
 protected:
 	// UAnimGraphNode_Base interface
-	virtual void ValidateAnimNodeDuringCompilation(USkeleton* ForSkeleton, FCompilerResultsLog& MessageLog) override;
-	virtual void PostProcessPinName(const UEdGraphPin* Pin, FString& DisplayName) const override;
+	UE_API virtual void ValidateAnimNodeDuringCompilation(USkeleton* ForSkeleton, FCompilerResultsLog& MessageLog) override;
+	UE_API virtual void PostProcessPinName(const UEdGraphPin* Pin, FString& DisplayName) const override;
 	// End of UAnimGraphNode_Base interface
 
 	// UAnimGraphNode_SkeletalControlBase interface
-	virtual void Draw(FPrimitiveDrawInterface* PDI, USkeletalMeshComponent* SkelMeshComp) const override;
-	virtual FText GetControllerDescription() const override;
+	UE_API virtual void Draw(FPrimitiveDrawInterface* PDI, USkeletalMeshComponent* SkelMeshComp) const override;
+	UE_API virtual FText GetControllerDescription() const override;
 	virtual const FAnimNode_SkeletalControlBase* GetNode() const override { return &Node; }
 	// End of UAnimGraphNode_SkeletalControlBase interface
 
@@ -42,3 +44,4 @@ private:
 	FNodeTitleTextTable CachedNodeTitles;
 };
 
+#undef UE_API
diff --git a/Engine/Source/Editor/AnimGraph/Public/AnimGraphNode_CustomProperty.h b/Engine/Source/Editor/AnimGraph/Public/AnimGraphNode_CustomProperty.h
--- a/Engine/Source/Editor/AnimGraph/Public/AnimGraphNode_CustomProperty.h
+++ b/Engine/Source/Editor/AnimGraph/Public/AnimGraphNode_CustomProperty.h
@@ -11,60 +11,62 @@
 
 #include "AnimGraphNode_CustomProperty.generated.h"
 
+#define UE_API ANIMGRAPH_API
+
 class FCompilerResultsLog;
 class IDetailLayoutBuilder;
 class IPropertyHandle;
 
-UCLASS(Abstract)
-class ANIMGRAPH_API UAnimGraphNode_CustomProperty : public UAnimGraphNode_Base, public IClassVariableCreator
+UCLASS(MinimalAPI, Abstract)
+class UAnimGraphNode_CustomProperty : public UAnimGraphNode_Base, public IClassVariableCreator
 {
 	GENERATED_BODY()
 
 public:
 
 	// UObject interface
-	virtual void Serialize(FArchive& Ar) override;
-	virtual void PostEditChangeProperty(FPropertyChangedEvent& PropertyChangedEvent) override;
-	virtual void PostDuplicate(bool bDuplicateForPIE) override;
+	UE_API virtual void Serialize(FArchive& Ar) override;
+	UE_API virtual void PostEditChangeProperty(FPropertyChangedEvent& PropertyChangedEvent) override;
+	UE_API virtual void PostDuplicate(bool bDuplicateForPIE) override;
 	
 	// IClassVariableCreator interface
-	virtual void CreateClassVariablesFromBlueprint(IAnimBlueprintVariableCreationContext& InCreationContext) override;
+	UE_API virtual void CreateClassVariablesFromBlueprint(IAnimBlueprintVariableCreationContext& InCreationContext) override;
 
 	//~ Begin UEdGraphNode Interface.
-	virtual void ValidateAnimNodeDuringCompilation(USkeleton* ForSkeleton, FCompilerResultsLog& MessageLog) override;
-	virtual UObject* GetJumpTargetForDoubleClick() const override;
-	virtual bool HasExternalDependencies(TArray<class UStruct*>* OptionalOutput /*= NULL*/) const override;
-	virtual void PinConnectionListChanged(UEdGraphPin* Pin) override;
+	UE_API virtual void ValidateAnimNodeDuringCompilation(USkeleton* ForSkeleton, FCompilerResultsLog& MessageLog) override;
+	UE_API virtual UObject* GetJumpTargetForDoubleClick() const override;
+	UE_API virtual bool HasExternalDependencies(TArray<class UStruct*>* OptionalOutput /*= NULL*/) const override;
+	UE_API virtual void PinConnectionListChanged(UEdGraphPin* Pin) override;
 	//~ End UEdGraphNode Interface.
 
 	// UAnimGraphNode_Base interface
-	virtual void CustomizeDetails(IDetailLayoutBuilder& DetailBuilder) override;
-	virtual void OnProcessDuringCompilation(IAnimBlueprintCompilationContext& InCompilationContext, IAnimBlueprintGeneratedClassCompiledData& OutCompiledData) override;
-	virtual void OnCopyTermDefaultsToDefaultObject(IAnimBlueprintCopyTermDefaultsContext& InCompilationContext, IAnimBlueprintNodeCopyTermDefaultsContext& InPerNodeContext, IAnimBlueprintGeneratedClassCompiledData& OutCompiledData) override;
-	virtual void CreateCustomPins(TArray<UEdGraphPin*>* OldPins) override;
-	virtual FProperty* GetPinProperty(FName InPinName) const override;
-	virtual bool IsPinBindable(const UEdGraphPin* InPin) const override;
-	virtual bool GetPinBindingInfo(FName InPinName, FName& OutBindingName, FProperty*& OutPinProperty, int32& OutOptionalPinIndex) const override;
-	virtual bool HasBinding(FName InBindingName) const override;
+	UE_API virtual void CustomizeDetails(IDetailLayoutBuilder& DetailBuilder) override;
+	UE_API virtual void OnProcessDuringCompilation(IAnimBlueprintCompilationContext& InCompilationContext, IAnimBlueprintGeneratedClassCompiledData& OutCompiledData) override;
+	UE_API virtual void OnCopyTermDefaultsToDefaultObject(IAnimBlueprintCopyTermDefaultsContext& InCompilationContext, IAnimBlueprintNodeCopyTermDefaultsContext& InPerNodeContext, IAnimBlueprintGeneratedClassCompiledData& OutCompiledData) override;
+	UE_API virtual void CreateCustomPins(TArray<UEdGraphPin*>* OldPins) override;
+	UE_API virtual FProperty* GetPinProperty(FName InPinName) const override;
+	UE_API virtual bool IsPinBindable(const UEdGraphPin* InPin) const override;
+	UE_API virtual bool GetPinBindingInfo(FName InPinName, FName& OutBindingName, FProperty*& OutPinProperty, int32& OutOptionalPinIndex) const override;
+	UE_API virtual bool HasBinding(FName InBindingName) const override;
 	
 	// Gets the property on InOwnerInstanceClass that corresponds to InInputPin
-	void GetInstancePinProperty(const IAnimBlueprintCompilationContext& InCompilationContext, UEdGraphPin* InInputPin, FProperty*& OutProperty);
+	UE_API void GetInstancePinProperty(const IAnimBlueprintCompilationContext& InCompilationContext, UEdGraphPin* InInputPin, FProperty*& OutProperty);
 	// Gets the unique name for the property linked to a given pin
-	FString GetPinTargetVariableName(const UEdGraphPin* InPin) const;
+	UE_API FString GetPinTargetVariableName(const UEdGraphPin* InPin) const;
 	// Gets the unique name for the property linked to a given pin name
-	FString GetPinTargetVariableName(FName InPinName) const;
+	UE_API FString GetPinTargetVariableName(FName InPinName) const;
 	// Gets Target Class this properties to link
-	UClass* GetTargetClass() const;
+	UE_API UClass* GetTargetClass() const;
 	// Add Source and Target Properties - Check FAnimNode_CustomProperty
-	void AddSourceTargetProperties(const FName& InSourcePropertyName, const FName& InTargetPropertyName);
+	UE_API void AddSourceTargetProperties(const FName& InSourcePropertyName, const FName& InTargetPropertyName);
 	// Helper used to get the skeleton class we are targeting
-	virtual UClass* GetTargetSkeletonClass() const;
+	UE_API virtual UClass* GetTargetSkeletonClass() const;
 
 	// ----- UI CALLBACKS ----- //
 	// User changed the instance class etc.
-	void OnStructuralPropertyChanged(IDetailLayoutBuilder* DetailBuilder);
+	UE_API void OnStructuralPropertyChanged(IDetailLayoutBuilder* DetailBuilder);
 	// User changed the instance class
-	void OnInstanceClassChanged(IDetailLayoutBuilder* DetailBuilder);
+	UE_API void OnInstanceClassChanged(IDetailLayoutBuilder* DetailBuilder);
 protected:
 	friend struct FCustomPropertyOptionalPinManager;
 	friend class SAnimationGraphNode;
@@ -86,7 +88,7 @@ protected:
 	TArray<FOptionalPinFromProperty> CustomPinProperties;
 	
 	// Gets a property's type as FText (for UI)
-	FText GetPropertyTypeText(FProperty* Property);
+	UE_API FText GetPropertyTypeText(FProperty* Property);
 
 	// internal node accessor
 	virtual FAnimNode_CustomProperty* GetCustomPropertyNode() PURE_VIRTUAL(UAnimGraphNode_CustomProperty::GetCustomPropertyNode, return nullptr;);
@@ -99,8 +101,10 @@ protected:
 	virtual bool NeedsToSpecifyValidTargetClass() const { return true; }
 	
 	// Sets the visibility of the specified pin, reconstructs the node if it changes
-	void SetCustomPinVisibility(bool bInVisible, int32 InOptionalPinIndex);
+	UE_API void SetCustomPinVisibility(bool bInVisible, int32 InOptionalPinIndex);
 
 	// Helper function for GetPinTargetVariableName
-	FString GetPinTargetVariableNameBase(FName InPinName) const;
+	UE_API FString GetPinTargetVariableNameBase(FName InPinName) const;
 };
+
+#undef UE_API
diff --git a/Engine/Source/Editor/AnimGraph/Public/AnimGraphNode_LinkedInputPose.h b/Engine/Source/Editor/AnimGraph/Public/AnimGraphNode_LinkedInputPose.h
--- a/Engine/Source/Editor/AnimGraph/Public/AnimGraphNode_LinkedInputPose.h
+++ b/Engine/Source/Editor/AnimGraph/Public/AnimGraphNode_LinkedInputPose.h
@@ -11,6 +11,8 @@
 
 #include "AnimGraphNode_LinkedInputPose.generated.h"
 
+#define UE_API ANIMGRAPH_API
+
 class SEditableTextBox;
 
 /** Required info for reconstructing a manually specified pin */
@@ -40,8 +42,8 @@ struct FAnimBlueprintFunctionPinInfo
 	FEdGraphPinType Type;
 };
 
-UCLASS()
-class ANIMGRAPH_API UAnimGraphNode_LinkedInputPose : public UAnimGraphNode_Base, public IClassVariableCreator
+UCLASS(MinimalAPI)
+class UAnimGraphNode_LinkedInputPose : public UAnimGraphNode_Base, public IClassVariableCreator
 {
 	GENERATED_BODY()
 
@@ -60,83 +62,85 @@ public:
 	UPROPERTY()
 	int32 InputPoseIndex;
 
-	UAnimGraphNode_LinkedInputPose();
+	UE_API UAnimGraphNode_LinkedInputPose();
 
 	/** UObject interface */
-	virtual void PostEditChangeProperty(FPropertyChangedEvent& PropertyChangedEvent) override;
+	UE_API virtual void PostEditChangeProperty(FPropertyChangedEvent& PropertyChangedEvent) override;
 
 	/** IClassVariableCreator interface */
-	virtual void CreateClassVariablesFromBlueprint(IAnimBlueprintVariableCreationContext& InCreationContext) override;
+	UE_API virtual void CreateClassVariablesFromBlueprint(IAnimBlueprintVariableCreationContext& InCreationContext) override;
 
 	/** UEdGraphNode interface */
-	virtual FLinearColor GetNodeTitleColor() const override;
-	virtual FText GetTooltipText() const override;
-	virtual FText GetNodeTitle(ENodeTitleType::Type TitleType) const override;
-	virtual FText GetMenuCategory() const override;
-	virtual bool CanUserDeleteNode() const override;
-	virtual bool CanDuplicateNode() const override;
-	virtual void AllocateDefaultPins() override;
-	virtual void ReallocatePinsDuringReconstruction(TArray<UEdGraphPin*>& OldPins) override;
-	virtual void PostPlacedNewNode() override;
-	virtual bool IsCompatibleWithGraph(UEdGraph const* Graph) const override;
-	virtual bool HasExternalDependencies(TArray<UStruct*>* OptionalOutput) const override;
+	UE_API virtual FLinearColor GetNodeTitleColor() const override;
+	UE_API virtual FText GetTooltipText() const override;
+	UE_API virtual FText GetNodeTitle(ENodeTitleType::Type TitleType) const override;
+	UE_API virtual FText GetMenuCategory() const override;
+	UE_API virtual bool CanUserDeleteNode() const override;
+	UE_API virtual bool CanDuplicateNode() const override;
+	UE_API virtual void AllocateDefaultPins() override;
+	UE_API virtual void ReallocatePinsDuringReconstruction(TArray<UEdGraphPin*>& OldPins) override;
+	UE_API virtual void PostPlacedNewNode() override;
+	UE_API virtual bool IsCompatibleWithGraph(UEdGraph const* Graph) const override;
+	UE_API virtual bool HasExternalDependencies(TArray<UStruct*>* OptionalOutput) const override;
 
 	/** UK2Node interface */
-	virtual void ExpandNode(class FKismetCompilerContext& InCompilerContext, UEdGraph* InSourceGraph) override;
+	UE_API virtual void ExpandNode(class FKismetCompilerContext& InCompilerContext, UEdGraph* InSourceGraph) override;
 
 	/** UAnimGraphNode_Base interface */
-	virtual void CustomizeDetails(IDetailLayoutBuilder& DetailBuilder) override;
-	virtual void GetOutputLinkAttributes(FNodeAttributeArray& OutAttributes) const override;
+	UE_API virtual void CustomizeDetails(IDetailLayoutBuilder& DetailBuilder) override;
+	UE_API virtual void GetOutputLinkAttributes(FNodeAttributeArray& OutAttributes) const override;
 	virtual bool ShouldShowAttributesOnPins() const override { return false; }
-	virtual void OnCopyTermDefaultsToDefaultObject(IAnimBlueprintCopyTermDefaultsContext& InCompilationContext, IAnimBlueprintNodeCopyTermDefaultsContext& InPerNodeContext, IAnimBlueprintGeneratedClassCompiledData& OutCompiledData) override;
-	virtual void GetRequiredExtensions(TArray<TSubclassOf<UAnimBlueprintExtension>>& OutExtensions) const override;
+	UE_API virtual void OnCopyTermDefaultsToDefaultObject(IAnimBlueprintCopyTermDefaultsContext& InCompilationContext, IAnimBlueprintNodeCopyTermDefaultsContext& InPerNodeContext, IAnimBlueprintGeneratedClassCompiledData& OutCompiledData) override;
+	UE_API virtual void GetRequiredExtensions(TArray<TSubclassOf<UAnimBlueprintExtension>>& OutExtensions) const override;
 
 	/** Make a name widget for this linked input pose node */
-	TSharedRef<SWidget> MakeNameWidget(IDetailLayoutBuilder& DetailBuilder);
+	UE_API TSharedRef<SWidget> MakeNameWidget(IDetailLayoutBuilder& DetailBuilder);
 
 	/** @return whether this node is editable. Non-editable nodes are implemented from function interfaces */
 	bool IsEditable() const { return CanUserDeleteNode(); }
 
 	/** @return the number of parameter inputs this node provides. This is either determined via the Inputs array or via the FunctionReference. */
-	int32 GetNumInputs() const;
+	UE_API int32 GetNumInputs() const;
 
 	/** Promotes the node from being a part of an interface override to a full function that allows for parameter and result pin additions */
-	void PromoteFromInterfaceOverride();
+	UE_API void PromoteFromInterfaceOverride();
 
 	/** Conform input pose name according to function */
-	void ConformInputPoseName();
+	UE_API void ConformInputPoseName();
 
 	/** Validate pose index against the function reference (used to determine whether we should exist or not) */
-	bool ValidateAgainstFunctionReference() const;
+	UE_API bool ValidateAgainstFunctionReference() const;
 
 	/** Helper function for iterating stub function parameters */
-	void IterateFunctionParameters(TFunctionRef<void(const FName&, const FEdGraphPinType&)> InFunc) const;
+	UE_API void IterateFunctionParameters(TFunctionRef<void(const FName&, const FEdGraphPinType&)> InFunc) const;
 
 private:
 	friend class FAnimGraphDetails;
 	friend class UAnimationGraph;
 	
 	// Helper function for common code in AllocateDefaultPins and ReallocatePinsDuringReconstruction
-	void AllocatePinsInternal();
+	UE_API void AllocatePinsInternal();
 
 	/** Create pins from the user-defined Inputs array */
-	void CreateUserDefinedPins();
+	UE_API void CreateUserDefinedPins();
 
 	/** Create pins from the stub function FunctionReference */
-	void CreatePinsFromStubFunction(const UFunction* Function);
+	UE_API void CreatePinsFromStubFunction(const UFunction* Function);
 
 	friend class UAnimBlueprintExtension_LinkedInputPose;
 	// Called pre-compilation to evaluate is this input pose is used
-	void AnalyzeLinks(TArrayView<UAnimGraphNode_Base*> InAnimNodes);
+	UE_API void AnalyzeLinks(TArrayView<UAnimGraphNode_Base*> InAnimNodes);
 
 	/** Reconstruct any layer nodes in this BP post-edit */
-	static void ReconstructLayerNodes(UBlueprint* InBlueprint);
+	static UE_API void ReconstructLayerNodes(UBlueprint* InBlueprint);
 	
 private:
 	/** UI helper functions */
-	void HandleInputPoseArrayChanged();
-	void HandleInputPinArrayChanged();
+	UE_API void HandleInputPoseArrayChanged();
+	UE_API void HandleInputPinArrayChanged();
 };
 
 UE_DEPRECATED(4.24, "UAnimGraphNode_SubInput has been renamed to UAnimGraphNode_LinkedInputPose")
-typedef UAnimGraphNode_LinkedInputPose UAnimGraphNode_SubInput;
\ No newline at end of file
+typedef UAnimGraphNode_LinkedInputPose UAnimGraphNode_SubInput;
+
+#undef UE_API
diff --git a/Engine/Source/Editor/AnimGraph/Public/AnimGraphNode_ModifyBone.h b/Engine/Source/Editor/AnimGraph/Public/AnimGraphNode_ModifyBone.h
--- a/Engine/Source/Editor/AnimGraph/Public/AnimGraphNode_ModifyBone.h
+++ b/Engine/Source/Editor/AnimGraph/Public/AnimGraphNode_ModifyBone.h
@@ -9,10 +9,12 @@
 #include "AnimGraphNode_SkeletalControlBase.h"
 #include "AnimGraphNode_ModifyBone.generated.h"
 
+#define UE_API ANIMGRAPH_API
+
 class FCompilerResultsLog;
 
-UCLASS(meta=(Keywords = "Modify Transform"))
-class ANIMGRAPH_API UAnimGraphNode_ModifyBone : public UAnimGraphNode_SkeletalControlBase
+UCLASS(MinimalAPI, meta=(Keywords = "Modify Transform"))
+class UAnimGraphNode_ModifyBone : public UAnimGraphNode_SkeletalControlBase
 {
 	GENERATED_UCLASS_BODY()
 
@@ -21,20 +23,20 @@ class ANIMGRAPH_API UAnimGraphNode_ModifyBone : public UAnimGraphNode_SkeletalCo
 
 public:
 	// UEdGraphNode interface
-	virtual FText GetNodeTitle(ENodeTitleType::Type TitleType) const override;
-	virtual FText GetTooltipText() const override;
+	UE_API virtual FText GetNodeTitle(ENodeTitleType::Type TitleType) const override;
+	UE_API virtual FText GetTooltipText() const override;
 	// End of UEdGraphNode interface
 
 protected:
 	// UAnimGraphNode_Base interface
-	virtual void ValidateAnimNodeDuringCompilation(USkeleton* ForSkeleton, FCompilerResultsLog& MessageLog) override;
-	virtual FEditorModeID GetEditorMode() const override;
-	virtual void CopyNodeDataToPreviewNode(FAnimNode_Base* InPreviewNode) override;
-	virtual void CopyPinDefaultsToNodeData(UEdGraphPin* InPin) override;
+	UE_API virtual void ValidateAnimNodeDuringCompilation(USkeleton* ForSkeleton, FCompilerResultsLog& MessageLog) override;
+	UE_API virtual FEditorModeID GetEditorMode() const override;
+	UE_API virtual void CopyNodeDataToPreviewNode(FAnimNode_Base* InPreviewNode) override;
+	UE_API virtual void CopyPinDefaultsToNodeData(UEdGraphPin* InPin) override;
 	// End of UAnimGraphNode_Base interface
 
 	// UAnimGraphNode_SkeletalControlBase interface
-	virtual FText GetControllerDescription() const override;
+	UE_API virtual FText GetControllerDescription() const override;
 	virtual const FAnimNode_SkeletalControlBase* GetNode() const override { return &Node; }
 	// End of UAnimGraphNode_SkeletalControlBase interface
 
@@ -43,3 +45,4 @@ private:
 	FNodeTitleTextTable CachedNodeTitles;
 };
 
+#undef UE_API
diff --git a/Engine/Source/Editor/AnimGraph/Public/AnimGraphNode_ObserveBone.h b/Engine/Source/Editor/AnimGraph/Public/AnimGraphNode_ObserveBone.h
--- a/Engine/Source/Editor/AnimGraph/Public/AnimGraphNode_ObserveBone.h
+++ b/Engine/Source/Editor/AnimGraph/Public/AnimGraphNode_ObserveBone.h
@@ -8,13 +8,15 @@
 #include "BoneControllers/AnimNode_ObserveBone.h"
 #include "AnimGraphNode_ObserveBone.generated.h"
 
+#define UE_API ANIMGRAPH_API
+
 class FCompilerResultsLog;
 class IDetailLayoutBuilder;
 class SGraphNode;
 
 // This allows you to observe the state of a bone at a particular point in the graph, showing it in any space and optionally relative to the reference pose
-UCLASS()
-class ANIMGRAPH_API UAnimGraphNode_ObserveBone : public UAnimGraphNode_SkeletalControlBase
+UCLASS(MinimalAPI)
+class UAnimGraphNode_ObserveBone : public UAnimGraphNode_SkeletalControlBase
 {
 	GENERATED_UCLASS_BODY()
 
@@ -23,23 +25,24 @@ class ANIMGRAPH_API UAnimGraphNode_ObserveBone : public UAnimGraphNode_SkeletalC
 
 public:
 	// UEdGraphNode interface
-	virtual FText GetNodeTitle(ENodeTitleType::Type TitleType) const override;
-	virtual FText GetTooltipText() const override;
-	virtual TSharedPtr<SGraphNode> CreateVisualWidget() override;
-	virtual FLinearColor GetNodeTitleColor() const override;
+	UE_API virtual FText GetNodeTitle(ENodeTitleType::Type TitleType) const override;
+	UE_API virtual FText GetTooltipText() const override;
+	UE_API virtual TSharedPtr<SGraphNode> CreateVisualWidget() override;
+	UE_API virtual FLinearColor GetNodeTitleColor() const override;
 	// End of UEdGraphNode interface
 
 protected:
 	// UAnimGraphNode_Base interface
-	virtual void ValidateAnimNodeDuringCompilation(USkeleton* ForSkeleton, FCompilerResultsLog& MessageLog) override;
-	virtual void CustomizePinData(UEdGraphPin* Pin, FName SourcePropertyName, int32 ArrayIndex) const override;
-	virtual void CustomizeDetails(IDetailLayoutBuilder& DetailBuilder) override;
-	virtual FEditorModeID GetEditorMode() const override;
+	UE_API virtual void ValidateAnimNodeDuringCompilation(USkeleton* ForSkeleton, FCompilerResultsLog& MessageLog) override;
+	UE_API virtual void CustomizePinData(UEdGraphPin* Pin, FName SourcePropertyName, int32 ArrayIndex) const override;
+	UE_API virtual void CustomizeDetails(IDetailLayoutBuilder& DetailBuilder) override;
+	UE_API virtual FEditorModeID GetEditorMode() const override;
 	// End of UAnimGraphNode_Base interface
 
 	// UAnimGraphNode_SkeletalControlBase interface
-	virtual FText GetControllerDescription() const override;
+	UE_API virtual FText GetControllerDescription() const override;
 	virtual const FAnimNode_SkeletalControlBase* GetNode() const override { return &Node; }
 	// End of UAnimGraphNode_SkeletalControlBase interface
 };
 
+#undef UE_API
diff --git a/Engine/Source/Editor/AnimGraph/Public/AnimGraphNode_PoseDriver.h b/Engine/Source/Editor/AnimGraph/Public/AnimGraphNode_PoseDriver.h
--- a/Engine/Source/Editor/AnimGraph/Public/AnimGraphNode_PoseDriver.h
+++ b/Engine/Source/Editor/AnimGraph/Public/AnimGraphNode_PoseDriver.h
@@ -9,10 +9,12 @@
 #include "EdGraph/EdGraphNodeUtils.h"
 #include "AnimGraphNode_PoseDriver.generated.h"
 
+#define UE_API ANIMGRAPH_API
+
 class FCompilerResultsLog;
 
-UCLASS(BlueprintType)
-class ANIMGRAPH_API UAnimGraphNode_PoseDriver : public UAnimGraphNode_PoseHandler
+UCLASS(MinimalAPI, BlueprintType)
+class UAnimGraphNode_PoseDriver : public UAnimGraphNode_PoseHandler
 {
 	GENERATED_UCLASS_BODY()
 
@@ -39,85 +41,85 @@ class ANIMGRAPH_API UAnimGraphNode_PoseDriver : public UAnimGraphNode_PoseHandle
 public:
 
 	/** Get the current preview node instance */
-	FAnimNode_PoseDriver* GetPreviewPoseDriverNode() const;
+	UE_API FAnimNode_PoseDriver* GetPreviewPoseDriverNode() const;
 
 	/** Util to replace current contents of PoseTargets with info from assigned PoseAsset */
 	UFUNCTION(BlueprintCallable, Category=PoseDriver)
-	void CopyTargetsFromPoseAsset();
+	UE_API void CopyTargetsFromPoseAsset();
 
 	/** Sets the pose-driver its source bones by name */
 	UFUNCTION(BlueprintCallable, Category=PoseDriver)
-	void SetSourceBones(const TArray<FName>& BoneNames);
+	UE_API void SetSourceBones(const TArray<FName>& BoneNames);
 
 	/** Returns the pose-driver its source bones by name */
 	UFUNCTION(BlueprintPure, Category=PoseDriver)
-	void GetSourceBoneNames(TArray<FName>& BoneNames);
+	UE_API void GetSourceBoneNames(TArray<FName>& BoneNames);
  
 	/** Set the pose-driver its driven bones by name */
 	UFUNCTION(BlueprintCallable, Category=PoseDriver)
-	void SetDrivingBones(const TArray<FName>& BoneNames);
+	UE_API void SetDrivingBones(const TArray<FName>& BoneNames);
 
 	/** Returns the pose-driver its driven bones by name */
 	UFUNCTION(BlueprintPure, Category=PoseDriver)
-	void GetDrivingBoneNames(TArray<FName>& BoneNames);
+	UE_API void GetDrivingBoneNames(TArray<FName>& BoneNames);
 
 	/**  */
 	UFUNCTION(BlueprintCallable, Category=PoseDriver)
-	void SetRBFParameters(FRBFParams Parameters);
+	UE_API void SetRBFParameters(FRBFParams Parameters);
 
 	/**  */
 	UFUNCTION(BlueprintPure, Category=PoseDriver)
-	FRBFParams& GetRBFParameters();
+	UE_API FRBFParams& GetRBFParameters();
 	
 	/**  */
 	UFUNCTION(BlueprintCallable, Category=PoseDriver)
-	void SetPoseDriverSource(EPoseDriverSource DriverSource);
+	UE_API void SetPoseDriverSource(EPoseDriverSource DriverSource);
 
 	/**  */
 	UFUNCTION(BlueprintPure, Category=PoseDriver)
-	EPoseDriverSource& GetPoseDriverSource();
+	UE_API EPoseDriverSource& GetPoseDriverSource();
 	
 	/**  */
 	UFUNCTION(BlueprintCallable, Category=PoseDriver)
-	void SetPoseDriverOutput(EPoseDriverOutput DriverOutput);
+	UE_API void SetPoseDriverOutput(EPoseDriverOutput DriverOutput);
 
 	/**  */
 	UFUNCTION(BlueprintPure, Category=PoseDriver)
-	EPoseDriverOutput& GetPoseDriverOutput();
+	UE_API EPoseDriverOutput& GetPoseDriverOutput();
 	
 	/** Automatically modify TargetScale for each PoseTarget, based on distance to nearest neighbor */
-	void AutoSetTargetScales(float& OutMaxDistance);
+	UE_API void AutoSetTargetScales(float& OutMaxDistance);
 
 	/** Adds a new target, reallocating transforms array appropriately */
-	void AddNewTarget();
+	UE_API void AddNewTarget();
 
 	/** Reallocates transforms arrays as necessary to accommodate source bones */
-	void ReserveTargetTransforms();
+	UE_API void ReserveTargetTransforms();
 
 	/** Return the color for a given weight. Used for Details and EditMode */
-	FLinearColor GetColorFromWeight(float InWeight);
+	UE_API FLinearColor GetColorFromWeight(float InWeight);
 
 	/** Used to refer back to preview instance in anim tools */
 	UPROPERTY(Transient)
 	TObjectPtr<USkeletalMeshComponent> LastPreviewComponent;
 
 	// Begin UObject Interface.
-	virtual void PostLoad() override;
-	virtual void Serialize(FArchive& Ar) override;
+	UE_API virtual void PostLoad() override;
+	UE_API virtual void Serialize(FArchive& Ar) override;
 	// End UObject Interface.
 
 	// UEdGraphNode interface
-	virtual FText GetNodeTitle(ENodeTitleType::Type TitleType) const override;
-	virtual FText GetTooltipText() const override;
-	virtual FText GetMenuCategory() const override;
-	virtual FLinearColor GetNodeBodyTintColor() const override;
+	UE_API virtual FText GetNodeTitle(ENodeTitleType::Type TitleType) const override;
+	UE_API virtual FText GetTooltipText() const override;
+	UE_API virtual FText GetMenuCategory() const override;
+	UE_API virtual FLinearColor GetNodeBodyTintColor() const override;
 	// End of UEdGraphNode interface
 
 	// UAnimGraphNode_Base interface
-	virtual void ValidateAnimNodeDuringCompilation(USkeleton* ForSkeleton, FCompilerResultsLog& MessageLog) override;
-	virtual FEditorModeID GetEditorMode() const override;
-	virtual EAnimAssetHandlerType SupportsAssetClass(const UClass* AssetClass) const override;
-	virtual void CopyNodeDataToPreviewNode(FAnimNode_Base* InPreviewNode) override;
+	UE_API virtual void ValidateAnimNodeDuringCompilation(USkeleton* ForSkeleton, FCompilerResultsLog& MessageLog) override;
+	UE_API virtual FEditorModeID GetEditorMode() const override;
+	UE_API virtual EAnimAssetHandlerType SupportsAssetClass(const UClass* AssetClass) const override;
+	UE_API virtual void CopyNodeDataToPreviewNode(FAnimNode_Base* InPreviewNode) override;
 	// End of UAnimGraphNode_Base interface
 
 protected:
@@ -130,3 +132,5 @@ protected:
 	/** Constructing FText strings can be costly, so we cache the node's title */
 	FNodeTitleTextTable CachedNodeTitles;
 };
+
+#undef UE_API
diff --git a/Engine/Source/Editor/AnimGraph/Public/AnimGraphNode_Root.h b/Engine/Source/Editor/AnimGraph/Public/AnimGraphNode_Root.h
--- a/Engine/Source/Editor/AnimGraph/Public/AnimGraphNode_Root.h
+++ b/Engine/Source/Editor/AnimGraph/Public/AnimGraphNode_Root.h
@@ -26,10 +26,10 @@ class UAnimGraphNode_Root : public UAnimGraphNode_Base
 	ANIMGRAPH_API virtual FLinearColor GetNodeTitleColor() const override;
 	ANIMGRAPH_API virtual FText GetTooltipText() const override;
 	ANIMGRAPH_API virtual FText GetNodeTitle(ENodeTitleType::Type TitleType) const override;
-	ANIMGRAPH_API virtual bool CanUserDeleteNode() const override { return false; }
-	ANIMGRAPH_API virtual bool CanDuplicateNode() const override { return false; }
+	virtual bool CanUserDeleteNode() const override { return false; }
+	virtual bool CanDuplicateNode() const override { return false; }
 	ANIMGRAPH_API virtual void GetMenuActions(FBlueprintActionDatabaseRegistrar& ActionRegistrar) const override;
-	ANIMGRAPH_API virtual bool IsNodeRootSet() const override { return true; }
+	virtual bool IsNodeRootSet() const override { return true; }
 	//~ End UEdGraphNode Interface.
 
 	//~ Begin UAnimGraphNode_Base Interface
diff --git a/Engine/Source/Editor/AnimGraph/Public/AnimGraphNode_RotationOffsetBlendSpaceGraph.h b/Engine/Source/Editor/AnimGraph/Public/AnimGraphNode_RotationOffsetBlendSpaceGraph.h
--- a/Engine/Source/Editor/AnimGraph/Public/AnimGraphNode_RotationOffsetBlendSpaceGraph.h
+++ b/Engine/Source/Editor/AnimGraph/Public/AnimGraphNode_RotationOffsetBlendSpaceGraph.h
@@ -7,11 +7,13 @@
 #include "AnimNodes/AnimNode_RotationOffsetBlendSpaceGraph.h"
 #include "AnimGraphNode_RotationOffsetBlendSpaceGraph.generated.h"
 
+#define UE_API ANIMGRAPH_API
+
 class FBlueprintActionDatabaseRegistrar;
 class FCompilerResultsLog;
 
-UCLASS()
-class ANIMGRAPH_API UAnimGraphNode_RotationOffsetBlendSpaceGraph : public UAnimGraphNode_BlendSpaceGraphBase
+UCLASS(MinimalAPI)
+class UAnimGraphNode_RotationOffsetBlendSpaceGraph : public UAnimGraphNode_BlendSpaceGraphBase
 {
 	GENERATED_BODY()
 
@@ -20,16 +22,18 @@ private:
 	FAnimNode_RotationOffsetBlendSpaceGraph Node;
 
 	// UObject interface
-	virtual void PostEditChangeProperty(FPropertyChangedEvent& PropertyChangedEvent) override;
+	UE_API virtual void PostEditChangeProperty(FPropertyChangedEvent& PropertyChangedEvent) override;
 
 	// UEdGraphNode interface
-	virtual FText GetNodeTitle(ENodeTitleType::Type TitleType) const override;
+	UE_API virtual FText GetNodeTitle(ENodeTitleType::Type TitleType) const override;
 	
 	// UK2Node interface
-	virtual void GetMenuActions(FBlueprintActionDatabaseRegistrar& ActionRegistrar) const override;
-	virtual void CustomizePinData(UEdGraphPin* Pin, FName SourcePropertyName, int32 ArrayIndex) const override;
+	UE_API virtual void GetMenuActions(FBlueprintActionDatabaseRegistrar& ActionRegistrar) const override;
+	UE_API virtual void CustomizePinData(UEdGraphPin* Pin, FName SourcePropertyName, int32 ArrayIndex) const override;
 
 	// UAnimGraphNode_Base interface
-	virtual void BakeDataDuringCompilation(FCompilerResultsLog& MessageLog) override;
-	virtual void CustomizeDetails(IDetailLayoutBuilder& DetailBuilder) override;
+	UE_API virtual void BakeDataDuringCompilation(FCompilerResultsLog& MessageLog) override;
+	UE_API virtual void CustomizeDetails(IDetailLayoutBuilder& DetailBuilder) override;
 };
+
+#undef UE_API
diff --git a/Engine/Source/Editor/AnimGraph/Public/AnimGraphNode_SkeletalControlBase.h b/Engine/Source/Editor/AnimGraph/Public/AnimGraphNode_SkeletalControlBase.h
--- a/Engine/Source/Editor/AnimGraph/Public/AnimGraphNode_SkeletalControlBase.h
+++ b/Engine/Source/Editor/AnimGraph/Public/AnimGraphNode_SkeletalControlBase.h
@@ -8,6 +8,8 @@
 #include "AnimGraphNode_Base.h"
 #include "AnimGraphNode_SkeletalControlBase.generated.h"
 
+#define UE_API ANIMGRAPH_API
+
 class FCompilerResultsLog;
 class USkeletalMeshComponent;
 struct FAnimNode_SkeletalControlBase;
@@ -19,26 +21,26 @@ struct HActor;
  *
  * Concrete subclasses should contain a member struct derived from FAnimNode_SkeletalControlBase
  */
-UCLASS(Abstract)
-class ANIMGRAPH_API UAnimGraphNode_SkeletalControlBase : public UAnimGraphNode_Base
+UCLASS(MinimalAPI, Abstract)
+class UAnimGraphNode_SkeletalControlBase : public UAnimGraphNode_Base
 {
 	GENERATED_UCLASS_BODY()
 
 public:
 	// UEdGraphNode interface
-	virtual FLinearColor GetNodeTitleColor() const override;
-	virtual FText GetTooltipText() const override;
-	virtual void PostEditChangeProperty(struct FPropertyChangedEvent& PropertyChangedEvent) override;
-	virtual bool ShowVisualWarning() const override;
-	virtual FText GetVisualWarningTooltipText() const override;
+	UE_API virtual FLinearColor GetNodeTitleColor() const override;
+	UE_API virtual FText GetTooltipText() const override;
+	UE_API virtual void PostEditChangeProperty(struct FPropertyChangedEvent& PropertyChangedEvent) override;
+	UE_API virtual bool ShowVisualWarning() const override;
+	UE_API virtual FText GetVisualWarningTooltipText() const override;
 	// End of UEdGraphNode interface
 
 	// UAnimGraphNode_Base interface
-	virtual FString GetNodeCategory() const override;
-	virtual void CreateOutputPins() override;
-	virtual void ValidateAnimNodePostCompile(FCompilerResultsLog& MessageLog, UAnimBlueprintGeneratedClass* CompiledClass, int32 CompiledNodeIndex) override;
-	virtual void CustomizePinData(UEdGraphPin* Pin, FName SourcePropertyName, int32 ArrayIndex) const override;
-	virtual void CustomizeDetails(IDetailLayoutBuilder& DetailBuilder) override;
+	UE_API virtual FString GetNodeCategory() const override;
+	UE_API virtual void CreateOutputPins() override;
+	UE_API virtual void ValidateAnimNodePostCompile(FCompilerResultsLog& MessageLog, UAnimBlueprintGeneratedClass* CompiledClass, int32 CompiledNodeIndex) override;
+	UE_API virtual void CustomizePinData(UEdGraphPin* Pin, FName SourcePropertyName, int32 ArrayIndex) const override;
+	UE_API virtual void CustomizeDetails(IDetailLayoutBuilder& DetailBuilder) override;
 	// End of UAnimGraphNode_Base interface
 
 	/**
@@ -57,25 +59,25 @@ public:
 	virtual void CopyNodeDataFrom(const FAnimNode_Base* NewAnimNode){}
 
 	/** Are we currently showing this pin */
-	bool IsPinShown(const FName PinName) const;
+	UE_API bool IsPinShown(const FName PinName) const;
 
 	// Returns the coordinate system that should be used for this bone
 	UE_DEPRECATED(4.14, "This function is deprecated. Use a IAnimNodeEditMode-derived class to implement a new editor mode and override GetEditorMode to return its ID")
-	virtual int32 GetWidgetCoordinateSystem(const USkeletalMeshComponent* SkelComp);
+	UE_API virtual int32 GetWidgetCoordinateSystem(const USkeletalMeshComponent* SkelComp);
 
 	// return current widget mode this anim graph node supports
 	UE_DEPRECATED(4.14, "This function is deprecated. Use a IAnimNodeEditMode-derived class to implement a new editor mode and override GetEditorMode to return its ID")
-	virtual int32 GetWidgetMode(const USkeletalMeshComponent* SkelComp);
+	UE_API virtual int32 GetWidgetMode(const USkeletalMeshComponent* SkelComp);
 	// called when the user changed widget mode by pressing "Space" key
 	UE_DEPRECATED(4.14, "This function is deprecated. Use a IAnimNodeEditMode-derived class to implement a new editor mode and override GetEditorMode to return its ID")
-	virtual int32 ChangeToNextWidgetMode(const USkeletalMeshComponent* SkelComp, int32 CurWidgetMode);
+	UE_API virtual int32 ChangeToNextWidgetMode(const USkeletalMeshComponent* SkelComp, int32 CurWidgetMode);
 	// called when the user set widget mode directly, returns true if InWidgetMode is available
 	UE_DEPRECATED(4.14, "This function is deprecated. Use a IAnimNodeEditMode-derived class to implement a new editor mode and override GetEditorMode to return its ID")
 	virtual bool SetWidgetMode(const USkeletalMeshComponent* SkelComp, int32 InWidgetMode) { return false; }
 
 	// 
 	UE_DEPRECATED(4.14, "This function is deprecated. Use a IAnimNodeEditMode-derived class to implement a new editor mode and override GetEditorMode to return its ID")
-	virtual FName FindSelectedBone();
+	UE_API virtual FName FindSelectedBone();
 
 	// if anim graph node needs other actors to select other bones, move actor's positions when this is called
 	UE_DEPRECATED(4.14, "This function is deprecated. Use a IAnimNodeEditMode-derived class to implement a new editor mode and override GetEditorMode to return its ID")
@@ -102,9 +104,9 @@ public:
 
 public:
 	// set literal value for FVector
-	void SetDefaultValue(const FName InDefaultValueName, const FVector& InValue);
+	UE_API void SetDefaultValue(const FName InDefaultValueName, const FVector& InValue);
 	// get literal value for vector
-	void GetDefaultValue(const FName UpdateDefaultValueName, FVector& OutVec);
+	UE_API void GetDefaultValue(const FName UpdateDefaultValueName, FVector& OutVec);
 
 	void GetDefaultValue(const FName PropName, FRotator& OutValue)
 	{
@@ -145,25 +147,27 @@ public:
 
 protected:
 	// Returns the short descriptive name of the controller
-	virtual FText GetControllerDescription() const;
+	UE_API virtual FText GetControllerDescription() const;
 
 	/**
 	* helper functions for bone control preview
 	*/
 	// local conversion function for drawing
 	UE_DEPRECATED(4.14, "This function is moved to FAnimNodeEditMode. Use a IAnimNodeEditMode-derived class to implement a new editor mode and override GetEditorMode to return its ID")
-	void ConvertToComponentSpaceTransform(const USkeletalMeshComponent* SkelComp, const FTransform & InTransform, FTransform & OutCSTransform, int32 BoneIndex, EBoneControlSpace Space) const;
+	UE_API void ConvertToComponentSpaceTransform(const USkeletalMeshComponent* SkelComp, const FTransform & InTransform, FTransform & OutCSTransform, int32 BoneIndex, EBoneControlSpace Space) const;
 	// convert drag vector in component space to bone space 
 	UE_DEPRECATED(4.14, "This function is moved to FAnimNodeEditMode. Use a IAnimNodeEditMode-derived class to implement a new editor mode and override GetEditorMode to return its ID")
-	FVector ConvertCSVectorToBoneSpace(const USkeletalMeshComponent* SkelComp, FVector& InCSVector, FCSPose<FCompactHeapPose>& MeshBases, const FName& BoneName, const EBoneControlSpace Space);
+	UE_API FVector ConvertCSVectorToBoneSpace(const USkeletalMeshComponent* SkelComp, FVector& InCSVector, FCSPose<FCompactHeapPose>& MeshBases, const FName& BoneName, const EBoneControlSpace Space);
 	// convert rotator in component space to bone space 
 	UE_DEPRECATED(4.14, "This function is moved to FAnimNodeEditMode. Use a IAnimNodeEditMode-derived class to implement a new editor mode and override GetEditorMode to return its ID")
-	FQuat ConvertCSRotationToBoneSpace(const USkeletalMeshComponent* SkelComp, FRotator& InCSRotator, FCSPose<FCompactHeapPose>& MeshBases, const FName& BoneName, const EBoneControlSpace Space);
+	UE_API FQuat ConvertCSRotationToBoneSpace(const USkeletalMeshComponent* SkelComp, FRotator& InCSRotator, FCSPose<FCompactHeapPose>& MeshBases, const FName& BoneName, const EBoneControlSpace Space);
 	// convert widget location according to bone control space
 	UE_DEPRECATED(4.14, "This function is moved to FAnimNodeEditMode. Use a IAnimNodeEditMode-derived class to implement a new editor mode and override GetEditorMode to return its ID")
-	FVector ConvertWidgetLocation(const USkeletalMeshComponent* InSkelComp, FCSPose<FCompactHeapPose>& InMeshBases, const FName& BoneName, const FVector& InLocation, const EBoneControlSpace Space);
+	UE_API FVector ConvertWidgetLocation(const USkeletalMeshComponent* InSkelComp, FCSPose<FCompactHeapPose>& InMeshBases, const FName& BoneName, const FVector& InLocation, const EBoneControlSpace Space);
 
 	virtual const FAnimNode_SkeletalControlBase* GetNode() const PURE_VIRTUAL(UAnimGraphNode_SkeletalControlBase::GetNode, return nullptr;);
 
-	FAnimNode_SkeletalControlBase* GetDebuggedNode() const;
+	UE_API FAnimNode_SkeletalControlBase* GetDebuggedNode() const;
 };
+
+#undef UE_API
diff --git a/Engine/Source/Editor/AnimGraph/Public/AnimGraphNode_StateMachineBase.h b/Engine/Source/Editor/AnimGraph/Public/AnimGraphNode_StateMachineBase.h
--- a/Engine/Source/Editor/AnimGraph/Public/AnimGraphNode_StateMachineBase.h
+++ b/Engine/Source/Editor/AnimGraph/Public/AnimGraphNode_StateMachineBase.h
@@ -9,10 +9,12 @@
 #include "Animation/AnimNode_StateMachine.h"
 #include "AnimGraphNode_StateMachineBase.generated.h"
 
+#define UE_API ANIMGRAPH_API
+
 class INameValidatorInterface;
 
-UCLASS(Abstract)
-class ANIMGRAPH_API UAnimGraphNode_StateMachineBase : public UAnimGraphNode_Base
+UCLASS(MinimalAPI, Abstract)
+class UAnimGraphNode_StateMachineBase : public UAnimGraphNode_Base
 {
 	GENERATED_UCLASS_BODY()
 
@@ -21,29 +23,29 @@ class ANIMGRAPH_API UAnimGraphNode_StateMachineBase : public UAnimGraphNode_Base
 	TObjectPtr<class UAnimationStateMachineGraph> EditorStateMachineGraph;
 
 	// UEdGraphNode interface
-	virtual FLinearColor GetNodeTitleColor() const override;
-	virtual FText GetTooltipText() const override;
-	virtual FText GetNodeTitle(ENodeTitleType::Type TitleType) const override;
-	virtual void PostPlacedNewNode() override;
-	virtual UObject* GetJumpTargetForDoubleClick() const override;
-	virtual void JumpToDefinition() const override;
-	virtual void DestroyNode() override;
-	virtual void PostPasteNode() override;
-	virtual TSharedPtr<class INameValidatorInterface> MakeNameValidator() const override;
-	virtual FString GetDocumentationLink() const override;
-	virtual void OnRenameNode(const FString& NewName) override;
-	virtual TArray<UEdGraph*> GetSubGraphs() const override;
-	virtual FText GetMenuCategory() const override;
+	UE_API virtual FLinearColor GetNodeTitleColor() const override;
+	UE_API virtual FText GetTooltipText() const override;
+	UE_API virtual FText GetNodeTitle(ENodeTitleType::Type TitleType) const override;
+	UE_API virtual void PostPlacedNewNode() override;
+	UE_API virtual UObject* GetJumpTargetForDoubleClick() const override;
+	UE_API virtual void JumpToDefinition() const override;
+	UE_API virtual void DestroyNode() override;
+	UE_API virtual void PostPasteNode() override;
+	UE_API virtual TSharedPtr<class INameValidatorInterface> MakeNameValidator() const override;
+	UE_API virtual FString GetDocumentationLink() const override;
+	UE_API virtual void OnRenameNode(const FString& NewName) override;
+	UE_API virtual TArray<UEdGraph*> GetSubGraphs() const override;
+	UE_API virtual FText GetMenuCategory() const override;
 	// End of UEdGraphNode interface
 
 	// UAnimGraphNode_Base interface
-	virtual void OnProcessDuringCompilation(IAnimBlueprintCompilationContext& InCompilationContext, IAnimBlueprintGeneratedClassCompiledData& OutCompiledData) override;
-	virtual void GetOutputLinkAttributes(FNodeAttributeArray& OutAttributes) const override;
-	virtual void GetRequiredExtensions(TArray<TSubclassOf<UAnimBlueprintExtension>>& OutExtensions) const override;
+	UE_API virtual void OnProcessDuringCompilation(IAnimBlueprintCompilationContext& InCompilationContext, IAnimBlueprintGeneratedClassCompiledData& OutCompiledData) override;
+	UE_API virtual void GetOutputLinkAttributes(FNodeAttributeArray& OutAttributes) const override;
+	UE_API virtual void GetRequiredExtensions(TArray<TSubclassOf<UAnimBlueprintExtension>>& OutExtensions) const override;
 	// End of UAnimGraphNode_Base interface
 
 	//  @return the name of this state machine
-	FString GetStateMachineName();
+	UE_API FString GetStateMachineName();
 
 	// Interface for derived classes to implement
 	virtual FAnimNode_StateMachine& GetNode() PURE_VIRTUAL(UAnimGraphNode_StateMachineBase::GetNode, static FAnimNode_StateMachine Dummy; return Dummy;);
@@ -53,3 +55,5 @@ private:
 	/** Constructing FText strings can be costly, so we cache the node's title */
 	FNodeTextCache CachedFullTitle;
 };
+
+#undef UE_API
diff --git a/Engine/Source/Editor/AnimGraph/Public/AnimGraphNode_TwistCorrectiveNode.h b/Engine/Source/Editor/AnimGraph/Public/AnimGraphNode_TwistCorrectiveNode.h
--- a/Engine/Source/Editor/AnimGraph/Public/AnimGraphNode_TwistCorrectiveNode.h
+++ b/Engine/Source/Editor/AnimGraph/Public/AnimGraphNode_TwistCorrectiveNode.h
@@ -8,13 +8,15 @@
 #include "BoneControllers/AnimNode_TwistCorrectiveNode.h"
 #include "AnimGraphNode_TwistCorrectiveNode.generated.h"
 
+#define UE_API ANIMGRAPH_API
+
 class FCompilerResultsLog;
 
 /**
  * This is the 'source version' of a bone driven controller, which maps part of the state from one bone to another (e.g., 2 * source.x -> target.z)
  */
-UCLASS()
-class ANIMGRAPH_API UAnimGraphNode_TwistCorrectiveNode : public UAnimGraphNode_SkeletalControlBase
+UCLASS(MinimalAPI)
+class UAnimGraphNode_TwistCorrectiveNode : public UAnimGraphNode_SkeletalControlBase
 {
 	GENERATED_UCLASS_BODY()
 
@@ -24,16 +26,18 @@ class ANIMGRAPH_API UAnimGraphNode_TwistCorrectiveNode : public UAnimGraphNode_S
 public:
 
 	// UEdGraphNode interface
-	virtual FText GetNodeTitle(ENodeTitleType::Type TitleType) const override;
-	virtual FText GetTooltipText() const override;
+	UE_API virtual FText GetNodeTitle(ENodeTitleType::Type TitleType) const override;
+	UE_API virtual FText GetTooltipText() const override;
 	// End of UEdGraphNode interface
 
 protected:
 
 	// UAnimGraphNode_SkeletalControlBase protected interface
-	virtual FText GetControllerDescription() const override;
-	virtual void ValidateAnimNodeDuringCompilation(USkeleton* ForSkeleton, FCompilerResultsLog& MessageLog) override;
+	UE_API virtual FText GetControllerDescription() const override;
+	UE_API virtual void ValidateAnimNodeDuringCompilation(USkeleton* ForSkeleton, FCompilerResultsLog& MessageLog) override;
 	virtual const FAnimNode_SkeletalControlBase* GetNode() const override { return &Node; }
 	// End of UAnimGraphNode_SkeletalControlBase protected interface
 
 };
+
+#undef UE_API
diff --git a/Engine/Source/Editor/AnimGraph/Public/AnimNodeEditMode.h b/Engine/Source/Editor/AnimGraph/Public/AnimNodeEditMode.h
--- a/Engine/Source/Editor/AnimGraph/Public/AnimNodeEditMode.h
+++ b/Engine/Source/Editor/AnimGraph/Public/AnimNodeEditMode.h
@@ -19,6 +19,8 @@
 #include "UObject/NameTypes.h"
 #include "UnrealWidgetFwd.h"
 
+#define UE_API ANIMGRAPH_API
+
 class FCanvas;
 class FEditorViewportClient;
 class FPrimitiveDrawInterface;
@@ -35,45 +37,45 @@ struct FViewportClick;
 template <class PoseType> struct FCSPose;
 
 /** Base implementation for anim node edit modes */
-class ANIMGRAPH_API FAnimNodeEditMode : public IAnimNodeEditMode
+class FAnimNodeEditMode : public IAnimNodeEditMode
 {
 public:
-	FAnimNodeEditMode();
+	UE_API FAnimNodeEditMode();
 
 	/** IAnimNodeEditMode interface */
-	virtual ECoordSystem GetWidgetCoordinateSystem() const override;
-	virtual UE::Widget::EWidgetMode GetWidgetMode() const override;
-	virtual UE::Widget::EWidgetMode ChangeToNextWidgetMode(UE::Widget::EWidgetMode CurWidgetMode) override;
-	virtual bool SetWidgetMode(UE::Widget::EWidgetMode InWidgetMode) override;
-	virtual FName GetSelectedBone() const override;
-	virtual void DoTranslation(FVector& InTranslation) override;
-	virtual void DoRotation(FRotator& InRotation) override;
-	virtual void DoScale(FVector& InScale) override;
-	virtual void EnterMode(class UAnimGraphNode_Base* InEditorNode, struct FAnimNode_Base* InRuntimeNode) override;
-	virtual void ExitMode() override;
+	UE_API virtual ECoordSystem GetWidgetCoordinateSystem() const override;
+	UE_API virtual UE::Widget::EWidgetMode GetWidgetMode() const override;
+	UE_API virtual UE::Widget::EWidgetMode ChangeToNextWidgetMode(UE::Widget::EWidgetMode CurWidgetMode) override;
+	UE_API virtual bool SetWidgetMode(UE::Widget::EWidgetMode InWidgetMode) override;
+	UE_API virtual FName GetSelectedBone() const override;
+	UE_API virtual void DoTranslation(FVector& InTranslation) override;
+	UE_API virtual void DoRotation(FRotator& InRotation) override;
+	UE_API virtual void DoScale(FVector& InScale) override;
+	UE_API virtual void EnterMode(class UAnimGraphNode_Base* InEditorNode, struct FAnimNode_Base* InRuntimeNode) override;
+	UE_API virtual void ExitMode() override;
 	virtual bool SupportsPoseWatch() override { return false; };
 
 	/** IPersonaEditMode interface */
-	virtual bool GetCameraTarget(FSphere& OutTarget) const override;
-	virtual class IPersonaPreviewScene& GetAnimPreviewScene() const override;
-	virtual void GetOnScreenDebugInfo(TArray<FText>& OutDebugInfo) const override;
+	UE_API virtual bool GetCameraTarget(FSphere& OutTarget) const override;
+	UE_API virtual class IPersonaPreviewScene& GetAnimPreviewScene() const override;
+	UE_API virtual void GetOnScreenDebugInfo(TArray<FText>& OutDebugInfo) const override;
 
 	/** FEdMode interface */
-	virtual void Render(const FSceneView* View, FViewport* Viewport, FPrimitiveDrawInterface* PDI) override;
-	virtual void DrawHUD(FEditorViewportClient* ViewportClient, FViewport* Viewport, const FSceneView* View, FCanvas* Canvas) override;
-	virtual bool HandleClick(FEditorViewportClient* InViewportClient, HHitProxy* HitProxy, const FViewportClick& Click) override;
-	virtual FVector GetWidgetLocation() const override;
-	virtual bool StartTracking(FEditorViewportClient* InViewportClient, FViewport* InViewport) override;
-	virtual bool EndTracking(FEditorViewportClient* InViewportClient, FViewport* InViewport) override;
-	virtual bool InputKey(FEditorViewportClient* InViewportClient, FViewport* InViewport, FKey InKey, EInputEvent InEvent) override;
-	virtual bool InputDelta(FEditorViewportClient* InViewportClient, FViewport* InViewport, FVector& InDrag, FRotator& InRot, FVector& InScale) override;
-	virtual bool GetCustomDrawingCoordinateSystem(FMatrix& InMatrix, void* InData) override;
-	virtual bool GetCustomInputCoordinateSystem(FMatrix& InMatrix, void* InData) override;
-	virtual bool ShouldDrawWidget() const override;
-	virtual void Tick(FEditorViewportClient* ViewportClient, float DeltaTime) override;
-	virtual void Exit() override;
-
-	virtual void RegisterPoseWatchedNode(UAnimGraphNode_Base* InEditorNode, FAnimNode_Base* InRuntimeNode);
+	UE_API virtual void Render(const FSceneView* View, FViewport* Viewport, FPrimitiveDrawInterface* PDI) override;
+	UE_API virtual void DrawHUD(FEditorViewportClient* ViewportClient, FViewport* Viewport, const FSceneView* View, FCanvas* Canvas) override;
+	UE_API virtual bool HandleClick(FEditorViewportClient* InViewportClient, HHitProxy* HitProxy, const FViewportClick& Click) override;
+	UE_API virtual FVector GetWidgetLocation() const override;
+	UE_API virtual bool StartTracking(FEditorViewportClient* InViewportClient, FViewport* InViewport) override;
+	UE_API virtual bool EndTracking(FEditorViewportClient* InViewportClient, FViewport* InViewport) override;
+	UE_API virtual bool InputKey(FEditorViewportClient* InViewportClient, FViewport* InViewport, FKey InKey, EInputEvent InEvent) override;
+	UE_API virtual bool InputDelta(FEditorViewportClient* InViewportClient, FViewport* InViewport, FVector& InDrag, FRotator& InRot, FVector& InScale) override;
+	UE_API virtual bool GetCustomDrawingCoordinateSystem(FMatrix& InMatrix, void* InData) override;
+	UE_API virtual bool GetCustomInputCoordinateSystem(FMatrix& InMatrix, void* InData) override;
+	UE_API virtual bool ShouldDrawWidget() const override;
+	UE_API virtual void Tick(FEditorViewportClient* ViewportClient, float DeltaTime) override;
+	UE_API virtual void Exit() override;
+
+	UE_API virtual void RegisterPoseWatchedNode(UAnimGraphNode_Base* InEditorNode, FAnimNode_Base* InRuntimeNode);
 
 	struct EditorRuntimeNodePair
 	{
@@ -90,31 +92,31 @@ public:
 	};
 
 	// Start IGizmoEdModeInterface overrides
-	virtual bool BeginTransform(const FGizmoState& InState) override;
-	virtual bool EndTransform(const FGizmoState& InState) override;
+	UE_API virtual bool BeginTransform(const FGizmoState& InState) override;
+	UE_API virtual bool EndTransform(const FGizmoState& InState) override;
 	// End IGizmoEdModeInterface overrides
 	
 protected:
 	// local conversion functions for drawing
-	static void ConvertToComponentSpaceTransform(const USkeletalMeshComponent* SkelComp, const FTransform & InTransform, FTransform & OutCSTransform, int32 BoneIndex, EBoneControlSpace Space);
-	static void ConvertToBoneSpaceTransform(const USkeletalMeshComponent* SkelComp, const FTransform & InCSTransform, FTransform & OutBSTransform, int32 BoneIndex, EBoneControlSpace Space);
+	static UE_API void ConvertToComponentSpaceTransform(const USkeletalMeshComponent* SkelComp, const FTransform & InTransform, FTransform & OutCSTransform, int32 BoneIndex, EBoneControlSpace Space);
+	static UE_API void ConvertToBoneSpaceTransform(const USkeletalMeshComponent* SkelComp, const FTransform & InCSTransform, FTransform & OutBSTransform, int32 BoneIndex, EBoneControlSpace Space);
 	// convert drag vector in component space to bone space 
-	static FVector ConvertCSVectorToBoneSpace(const USkeletalMeshComponent* SkelComp, FVector& InCSVector, FCSPose<FCompactHeapPose>& MeshBases, const FName& BoneName, const EBoneControlSpace Space);
-	static FVector ConvertCSVectorToBoneSpace(const USkeletalMeshComponent* SkelComp, FVector& InCSVector, FCSPose<FCompactHeapPose>& MeshBases, const FBoneSocketTarget& InTarget, const EBoneControlSpace Space);
+	static UE_API FVector ConvertCSVectorToBoneSpace(const USkeletalMeshComponent* SkelComp, FVector& InCSVector, FCSPose<FCompactHeapPose>& MeshBases, const FName& BoneName, const EBoneControlSpace Space);
+	static UE_API FVector ConvertCSVectorToBoneSpace(const USkeletalMeshComponent* SkelComp, FVector& InCSVector, FCSPose<FCompactHeapPose>& MeshBases, const FBoneSocketTarget& InTarget, const EBoneControlSpace Space);
 	// convert rotator in component space to bone space 
-	static FQuat ConvertCSRotationToBoneSpace(const USkeletalMeshComponent* SkelComp, FRotator& InCSRotator, FCSPose<FCompactHeapPose>& MeshBases, const FName& BoneName, const EBoneControlSpace Space);
+	static UE_API FQuat ConvertCSRotationToBoneSpace(const USkeletalMeshComponent* SkelComp, FRotator& InCSRotator, FCSPose<FCompactHeapPose>& MeshBases, const FName& BoneName, const EBoneControlSpace Space);
 	// convert widget location according to bone control space
-	static FVector ConvertWidgetLocation(const USkeletalMeshComponent* InSkelComp, FCSPose<FCompactHeapPose>& InMeshBases, const FName& BoneName, const FVector& InLocation, const EBoneControlSpace Space);
-	static FVector ConvertWidgetLocation(const USkeletalMeshComponent* InSkelComp, FCSPose<FCompactHeapPose>& InMeshBases, const FBoneSocketTarget& Target, const FVector& InLocation, const EBoneControlSpace Space);
+	static UE_API FVector ConvertWidgetLocation(const USkeletalMeshComponent* InSkelComp, FCSPose<FCompactHeapPose>& InMeshBases, const FName& BoneName, const FVector& InLocation, const EBoneControlSpace Space);
+	static UE_API FVector ConvertWidgetLocation(const USkeletalMeshComponent* InSkelComp, FCSPose<FCompactHeapPose>& InMeshBases, const FBoneSocketTarget& Target, const FVector& InLocation, const EBoneControlSpace Space);
 
-	virtual UAnimGraphNode_Base* GetActiveWidgetAnimNode() const; // Return the editor node associated with the selected widget. All widget operations are performed on this node.
-	virtual FAnimNode_Base*	GetActiveWidgetRuntimeAnimNode() const; // Return the runtime node associated with the selected widget. All widget operations are performed on this node.
+	UE_API virtual UAnimGraphNode_Base* GetActiveWidgetAnimNode() const; // Return the editor node associated with the selected widget. All widget operations are performed on this node.
+	UE_API virtual FAnimNode_Base*	GetActiveWidgetRuntimeAnimNode() const; // Return the runtime node associated with the selected widget. All widget operations are performed on this node.
 
 	const bool IsManipulatingWidget() const { return bManipulating; }
 
 	// Manage the start and end of a transform action in the viewport.
-	bool HandleBeginTransform();
-	bool HandleEndTransform();
+	UE_API bool HandleBeginTransform();
+	UE_API bool HandleEndTransform();
 	
 	TArray< EditorRuntimeNodePair > SelectedAnimNodes;	// Selected Anim Graph Nodes
 	TArray< EditorRuntimeNodePair > PoseWatchedAnimNodes; 	// Pose Watched Anim Graph Nodes. 
@@ -124,3 +126,5 @@ private:
 
 	bool bInTransaction;
 };
+
+#undef UE_API
diff --git a/Engine/Source/Editor/AnimGraph/Public/AnimNodeEditModes.h b/Engine/Source/Editor/AnimGraph/Public/AnimNodeEditModes.h
--- a/Engine/Source/Editor/AnimGraph/Public/AnimNodeEditModes.h
+++ b/Engine/Source/Editor/AnimGraph/Public/AnimNodeEditModes.h
@@ -5,16 +5,20 @@
 #include "CoreMinimal.h"
 #include "Editor.h"
 
-struct ANIMGRAPH_API AnimNodeEditModes
+#define UE_API ANIMGRAPH_API
+
+struct AnimNodeEditModes
 {
-	const static FEditorModeID AnimDynamics;
-	const static FEditorModeID AnimNode;
-	const static FEditorModeID TwoBoneIK;
-	const static FEditorModeID ObserveBone;
-	const static FEditorModeID ModifyBone;
-	const static FEditorModeID Fabrik;
-	const static FEditorModeID CCDIK;
-	const static FEditorModeID PoseDriver;
-	const static FEditorModeID SplineIK;
-	const static FEditorModeID LookAt;
+	UE_API const static FEditorModeID AnimDynamics;
+	UE_API const static FEditorModeID AnimNode;
+	UE_API const static FEditorModeID TwoBoneIK;
+	UE_API const static FEditorModeID ObserveBone;
+	UE_API const static FEditorModeID ModifyBone;
+	UE_API const static FEditorModeID Fabrik;
+	UE_API const static FEditorModeID CCDIK;
+	UE_API const static FEditorModeID PoseDriver;
+	UE_API const static FEditorModeID SplineIK;
+	UE_API const static FEditorModeID LookAt;
 };
+
+#undef UE_API
diff --git a/Engine/Source/Editor/AnimGraph/Public/AnimPreviewAttacheInstance.h b/Engine/Source/Editor/AnimGraph/Public/AnimPreviewAttacheInstance.h
--- a/Engine/Source/Editor/AnimGraph/Public/AnimPreviewAttacheInstance.h
+++ b/Engine/Source/Editor/AnimGraph/Public/AnimPreviewAttacheInstance.h
@@ -9,6 +9,8 @@
 #include "AnimNodes/AnimNode_CopyPoseFromMesh.h"
 #include "AnimPreviewAttacheInstance.generated.h"
 
+#define UE_API ANIMGRAPH_API
+
 /** Proxy override for this UAnimInstance-derived class */
 USTRUCT()
 struct FAnimPreviewAttacheInstanceProxy : public FAnimInstanceProxy
@@ -40,15 +42,16 @@ private:
  * Used by Preview in AnimGraph, Playing single animation in Kismet2 and etc
  */
 
-UCLASS(transient, NotBlueprintable, noteditinlinenew)
-class ANIMGRAPH_API UAnimPreviewAttacheInstance : public UAnimInstance
+UCLASS(MinimalAPI, transient, NotBlueprintable, noteditinlinenew)
+class UAnimPreviewAttacheInstance : public UAnimInstance
 {
 	GENERATED_UCLASS_BODY()
 
 	//~ Begin UAnimInstance Interface
-	virtual FAnimInstanceProxy* CreateAnimInstanceProxy() override;
+	UE_API virtual FAnimInstanceProxy* CreateAnimInstanceProxy() override;
 	//~ End UAnimInstance Interface
 };
 
 
 
+#undef UE_API
diff --git a/Engine/Source/Editor/AnimGraph/Public/AnimPreviewInstance.h b/Engine/Source/Editor/AnimGraph/Public/AnimPreviewInstance.h
--- a/Engine/Source/Editor/AnimGraph/Public/AnimPreviewInstance.h
+++ b/Engine/Source/Editor/AnimGraph/Public/AnimPreviewInstance.h
@@ -13,6 +13,8 @@
 #include "AnimNodes/AnimNode_CopyPoseFromMesh.h"
 #include "AnimPreviewInstance.generated.h"
 
+#define UE_API ANIMGRAPH_API
+
 /** Enum to know how montage is being played */
 UENUM()
 enum EMontagePreviewType : int
@@ -27,7 +29,7 @@ enum EMontagePreviewType : int
 
 /** Proxy override for this UAnimInstance-derived class */
 USTRUCT()
-struct ANIMGRAPH_API FAnimPreviewInstanceProxy : public FAnimSingleNodeInstanceProxy
+struct FAnimPreviewInstanceProxy : public FAnimSingleNodeInstanceProxy
 {
 	GENERATED_BODY()
 
@@ -46,20 +48,20 @@ public:
 		bCanProcessAdditiveAnimations = true;
 	}
 
-	virtual void Initialize(UAnimInstance* InAnimInstance) override;
-	virtual void Update(float DeltaSeconds) override;
-	virtual void UpdateAnimationNode(const FAnimationUpdateContext& InContext) override;
-	virtual bool Evaluate(FPoseContext& Output) override;
-	virtual void PreUpdate(UAnimInstance* InAnimInstance, float DeltaSeconds) override;
-	virtual void SetAnimationAsset(UAnimationAsset* NewAsset, USkeletalMeshComponent* MeshComponent, bool bIsLooping, float InPlayRate) override;
+	UE_API virtual void Initialize(UAnimInstance* InAnimInstance) override;
+	UE_API virtual void Update(float DeltaSeconds) override;
+	UE_API virtual void UpdateAnimationNode(const FAnimationUpdateContext& InContext) override;
+	UE_API virtual bool Evaluate(FPoseContext& Output) override;
+	UE_API virtual void PreUpdate(UAnimInstance* InAnimInstance, float DeltaSeconds) override;
+	UE_API virtual void SetAnimationAsset(UAnimationAsset* NewAsset, USkeletalMeshComponent* MeshComponent, bool bIsLooping, float InPlayRate) override;
 
-	void ResetModifiedBone(bool bCurveController = false);
+	UE_API void ResetModifiedBone(bool bCurveController = false);
 
-	FAnimNode_ModifyBone* FindModifiedBone(const FName& InBoneName, bool bCurveController = false);	
+	UE_API FAnimNode_ModifyBone* FindModifiedBone(const FName& InBoneName, bool bCurveController = false);	
 
-	FAnimNode_ModifyBone& ModifyBone(const FName& InBoneName, bool bCurveController = false);
+	UE_API FAnimNode_ModifyBone& ModifyBone(const FName& InBoneName, bool bCurveController = false);
 
-	void RemoveBoneModification(const FName& InBoneName, bool bCurveController = false);
+	UE_API void RemoveBoneModification(const FName& InBoneName, bool bCurveController = false);
 
 	void EnableControllers(bool bEnable)
 	{
@@ -88,7 +90,7 @@ public:
 	}
 #endif
 
-	void RefreshCurveBoneControllers(UAnimationAsset* AssetToRefreshFrom);
+	UE_API void RefreshCurveBoneControllers(UAnimationAsset* AssetToRefreshFrom);
 
 	TArray<FAnimNode_ModifyBone>& GetBoneControllers()
 	{
@@ -103,19 +105,19 @@ public:
 	virtual void AddImpulseAtLocation(FVector Impulse, FVector Location, FName BoneName = NAME_None) {}
 
 	/** Sets an external debug skeletal mesh component to use to debug */
-	void SetDebugSkeletalMeshComponent(USkeletalMeshComponent* InSkeletalMeshComponent);
+	UE_API void SetDebugSkeletalMeshComponent(USkeletalMeshComponent* InSkeletalMeshComponent);
 
 	/** Gets the external debug skeletal mesh component we are debugging */
-	USkeletalMeshComponent* GetDebugSkeletalMeshComponent() const;
+	UE_API USkeletalMeshComponent* GetDebugSkeletalMeshComponent() const;
 
 private:
-	void UpdateCurveController();
+	UE_API void UpdateCurveController();
 
-	void ApplyBoneControllers(TArray<FAnimNode_ModifyBone> &InBoneControllers, FComponentSpacePoseContext& ComponentSpacePoseContext);
+	UE_API void ApplyBoneControllers(TArray<FAnimNode_ModifyBone> &InBoneControllers, FComponentSpacePoseContext& ComponentSpacePoseContext);
 
-	void SetKeyImplementation(const FCompactPose& PreControllerInLocalSpace, const FCompactPose& PostControllerInLocalSpace);
+	UE_API void SetKeyImplementation(const FCompactPose& PreControllerInLocalSpace, const FCompactPose& PostControllerInLocalSpace);
 
-	void AddKeyToSequence(UAnimSequence* Sequence, float Time, const FName& BoneName, const FTransform& AdditiveTransform);
+	UE_API void AddKeyToSequence(UAnimSequence* Sequence, float Time, const FName& BoneName, const FTransform& AdditiveTransform);
 
 private:
 	/** Controllers for individual bones */
@@ -158,8 +160,8 @@ private:
  * Used by Preview in AnimGraph, Playing single animation in Kismet2 and etc
  */
 
-UCLASS(transient, NotBlueprintable, noteditinlinenew)
-class ANIMGRAPH_API UAnimPreviewInstance : public UAnimSingleNodeInstance
+UCLASS(MinimalAPI, transient, NotBlueprintable, noteditinlinenew)
+class UAnimPreviewInstance : public UAnimSingleNodeInstance
 {
 	GENERATED_UCLASS_BODY()
 
@@ -171,49 +173,49 @@ class ANIMGRAPH_API UAnimPreviewInstance : public UAnimSingleNodeInstance
 	int32 MontagePreviewStartSectionIdx;
 
 	//~ Begin UObject Interface
-	virtual void Serialize(FArchive& Ar) override;
+	UE_API virtual void Serialize(FArchive& Ar) override;
 	//~ End UObject Interface
 
 	//~ Begin UAnimInstance Interface
-	virtual void NativeInitializeAnimation() override;
-	virtual FAnimInstanceProxy* CreateAnimInstanceProxy() override;
+	UE_API virtual void NativeInitializeAnimation() override;
+	UE_API virtual FAnimInstanceProxy* CreateAnimInstanceProxy() override;
 	virtual bool CanRunParallelWork() const { return false; }
 protected:
-	virtual void Montage_Advance(float DeltaTime) override;
+	UE_API virtual void Montage_Advance(float DeltaTime) override;
 	//~ End UAnimInstance Interface
 
 public:
 	/** Set SkeletalControl Alpha**/
-	void SetSkeletalControlAlpha(float SkeletalControlAlpha);
+	UE_API void SetSkeletalControlAlpha(float SkeletalControlAlpha);
 
-	UAnimSequence* GetAnimSequence();
+	UE_API UAnimSequence* GetAnimSequence();
 
 	//~ Begin UAnimSingleNodeInstance Interface
-	virtual void RestartMontage(UAnimMontage* Montage, FName FromSection = FName()) override;
-	virtual void SetAnimationAsset(UAnimationAsset* NewAsset, bool bIsLooping = true, float InPlayRate = 1.f) override;
+	UE_API virtual void RestartMontage(UAnimMontage* Montage, FName FromSection = FName()) override;
+	UE_API virtual void SetAnimationAsset(UAnimationAsset* NewAsset, bool bIsLooping = true, float InPlayRate = 1.f) override;
 	//~ End UAnimSingleNodeInstance Interface
 
 	/** Montage preview functions */
-	void MontagePreview_JumpToStart();
-	void MontagePreview_JumpToEnd();
-	void MontagePreview_JumpToPreviewStart();
-	void MontagePreview_Restart();
-	void MontagePreview_PreviewNormal(int32 FromSectionIdx = INDEX_NONE, bool bPlay = true);
-	void MontagePreview_SetLoopNormal(bool bIsLooping, int32 PreferSectionIdx = INDEX_NONE);
-	void MontagePreview_PreviewAllSections(bool bPlay = true);
-	void MontagePreview_SetLoopAllSections(bool bIsLooping);
-	void MontagePreview_SetLoopAllSetupSections(bool bIsLooping);
-	void MontagePreview_ResetSectionsOrder();
-	void MontagePreview_SetLooping(bool bIsLooping);
-	void MontagePreview_SetPlaying(bool bIsPlaying);
-	void MontagePreview_SetReverse(bool bInReverse);
-	void MontagePreview_StepForward();
-	void MontagePreview_StepBackward();
-	void MontagePreview_JumpToPosition(float NewPosition);
-	int32 MontagePreview_FindFirstSectionAsInMontage(int32 AnySectionIdx);
-	int32 MontagePreview_FindLastSection(int32 StartSectionIdx);
-	float MontagePreview_CalculateStepLength();
-	void MontagePreview_RemoveBlendOut();
+	UE_API void MontagePreview_JumpToStart();
+	UE_API void MontagePreview_JumpToEnd();
+	UE_API void MontagePreview_JumpToPreviewStart();
+	UE_API void MontagePreview_Restart();
+	UE_API void MontagePreview_PreviewNormal(int32 FromSectionIdx = INDEX_NONE, bool bPlay = true);
+	UE_API void MontagePreview_SetLoopNormal(bool bIsLooping, int32 PreferSectionIdx = INDEX_NONE);
+	UE_API void MontagePreview_PreviewAllSections(bool bPlay = true);
+	UE_API void MontagePreview_SetLoopAllSections(bool bIsLooping);
+	UE_API void MontagePreview_SetLoopAllSetupSections(bool bIsLooping);
+	UE_API void MontagePreview_ResetSectionsOrder();
+	UE_API void MontagePreview_SetLooping(bool bIsLooping);
+	UE_API void MontagePreview_SetPlaying(bool bIsPlaying);
+	UE_API void MontagePreview_SetReverse(bool bInReverse);
+	UE_API void MontagePreview_StepForward();
+	UE_API void MontagePreview_StepBackward();
+	UE_API void MontagePreview_JumpToPosition(float NewPosition);
+	UE_API int32 MontagePreview_FindFirstSectionAsInMontage(int32 AnySectionIdx);
+	UE_API int32 MontagePreview_FindLastSection(int32 StartSectionIdx);
+	UE_API float MontagePreview_CalculateStepLength();
+	UE_API void MontagePreview_RemoveBlendOut();
 	bool IsPlayingMontage() { return GetActiveMontageInstance() != NULL; }
 
 	/** 
@@ -221,30 +223,30 @@ public:
 	 * @param	InBoneName	The name of the bone modification to find
 	 * @return the bone modification or NULL if no current modification was found
 	 */
-	FAnimNode_ModifyBone* FindModifiedBone(const FName& InBoneName, bool bCurveController=false);
+	UE_API FAnimNode_ModifyBone* FindModifiedBone(const FName& InBoneName, bool bCurveController=false);
 
 	/** 
 	 * Modifies a single bone. Create a new FAnimNode_ModifyBone if one does not exist for the passed-in bone.
 	 * @param	InBoneName	The name of the bone to modify
 	 * @return the new or existing bone modification
 	 */
-	FAnimNode_ModifyBone& ModifyBone(const FName& InBoneName, bool bCurveController=false);
+	UE_API FAnimNode_ModifyBone& ModifyBone(const FName& InBoneName, bool bCurveController=false);
 
 	/**
 	 * Removes an existing bone modification
 	 * @param	InBoneName	The name of the existing modification to remove
 	 */
-	void RemoveBoneModification(const FName& InBoneName, bool bCurveController=false);
+	UE_API void RemoveBoneModification(const FName& InBoneName, bool bCurveController=false);
 
 	/**
 	 * Reset all bone modified
 	 */
-	void ResetModifiedBone(bool bCurveController=false);
+	UE_API void ResetModifiedBone(bool bCurveController=false);
 
 	/**
 	 * Returns all currently active bone controllers on this instance's proxy
 	 */
-	const TArray<FAnimNode_ModifyBone>& GetBoneControllers();
+	UE_API const TArray<FAnimNode_ModifyBone>& GetBoneControllers();
 
 #if WITH_EDITOR	
 	/**
@@ -253,43 +255,44 @@ public:
 	 * It does wait until next update, and it gets the delta of transform before applying curves, and 
 	 * creates curves from it, so you'll need delegate if you'd like to do something after (set with SetKeyCompleteDelegate)
 	 */
-	void SetKey();
+	UE_API void SetKey();
 
 	/**
 	 * Add the delegate to be called when a key is set.
 	 * 
 	 * @param Delegate To be called once set key is completed
 	 */
-	FDelegateHandle AddKeyCompleteDelegate(FSimpleMulticastDelegate::FDelegate InOnSetKeyCompleteDelegate);
+	UE_API FDelegateHandle AddKeyCompleteDelegate(FSimpleMulticastDelegate::FDelegate InOnSetKeyCompleteDelegate);
 
 	/**
 	 * Add the delegate to be called when a key is set.
 	 * 
 	 * @param Delegate To be called once set key is completed
 	 */
-	void RemoveKeyCompleteDelegate(FDelegateHandle InDelegateHandle);
+	UE_API void RemoveKeyCompleteDelegate(FDelegateHandle InDelegateHandle);
 #endif
 
 	/** 
 	 * Refresh Curve Bone Controllers based on TransformCurves from Animation data
 	 */
-	void RefreshCurveBoneControllers();
+	UE_API void RefreshCurveBoneControllers();
 
 	/** 
 	 * Enable Controllers
 	 * This is used by when editing, when controller has to be disabled
 	 */
-	void EnableControllers(bool bEnable);
+	UE_API void EnableControllers(bool bEnable);
 
 	/** Preview physics interaction */
-	void AddImpulseAtLocation(FVector Impulse, FVector Location, FName BoneName = NAME_None);
+	UE_API void AddImpulseAtLocation(FVector Impulse, FVector Location, FName BoneName = NAME_None);
 
 	/** Sets an external debug skeletal mesh component to use to debug */
-	void SetDebugSkeletalMeshComponent(USkeletalMeshComponent* InSkeletalMeshComponent);
+	UE_API void SetDebugSkeletalMeshComponent(USkeletalMeshComponent* InSkeletalMeshComponent);
 
 	/** Gets the external debug skeletal mesh component we are debugging */
-	USkeletalMeshComponent* GetDebugSkeletalMeshComponent() const;
+	UE_API USkeletalMeshComponent* GetDebugSkeletalMeshComponent() const;
 };
 
 
 
+#undef UE_API
diff --git a/Engine/Source/Editor/AnimGraph/Public/AnimStateNodeBase.h b/Engine/Source/Editor/AnimGraph/Public/AnimStateNodeBase.h
--- a/Engine/Source/Editor/AnimGraph/Public/AnimStateNodeBase.h
+++ b/Engine/Source/Editor/AnimGraph/Public/AnimStateNodeBase.h
@@ -32,13 +32,13 @@ class UAnimStateNodeBase : public UEdGraphNode
 	// End of UEdGraphNode interface
 
 	// @return the input pin for this state
-	ANIMGRAPH_API virtual UEdGraphPin* GetInputPin() const { return NULL; }
+	virtual UEdGraphPin* GetInputPin() const { return NULL; }
 
 	// @return the output pin for this state
-	ANIMGRAPH_API virtual UEdGraphPin* GetOutputPin() const { return NULL; }
+	virtual UEdGraphPin* GetOutputPin() const { return NULL; }
 
 	// @return the name of this state
-	ANIMGRAPH_API virtual FString GetStateName() const { return TEXT("BaseState"); }
+	virtual FString GetStateName() const { return TEXT("BaseState"); }
 
 	// Populates the OutTransitions array with a list of transition nodes connected to this state
 	ANIMGRAPH_API virtual void GetTransitionList(TArray<class UAnimStateTransitionNode*>& OutTransitions, bool bWantSortedList = false) const;
diff --git a/Engine/Source/Editor/AnimGraph/Public/AnimationGraph.h b/Engine/Source/Editor/AnimGraph/Public/AnimationGraph.h
--- a/Engine/Source/Editor/AnimGraph/Public/AnimationGraph.h
+++ b/Engine/Source/Editor/AnimGraph/Public/AnimationGraph.h
@@ -9,13 +9,15 @@
 #include "AnimGraphNode_Base.h"
 #include "AnimationGraph.generated.h"
 
+#define UE_API ANIMGRAPH_API
+
 class UEdGraphPin;
 
 /** Delegate fired when a pin's default value is changed */
 DECLARE_MULTICAST_DELEGATE_OneParam(FOnPinDefaultValueChanged, UEdGraphPin* /*InPinThatChanged*/)
 
-UCLASS(BlueprintType)
-class ANIMGRAPH_API UAnimationGraph : public UEdGraph
+UCLASS(MinimalAPI, BlueprintType)
+class UAnimationGraph : public UEdGraph
 {
 	GENERATED_UCLASS_BODY()
 
@@ -28,13 +30,14 @@ class ANIMGRAPH_API UAnimationGraph : public UEdGraph
 
 	/** Returns contained graph nodes of the specified (or child) class */
 	UFUNCTION(BlueprintCallable, Category=AnimationGraph)
-	void GetGraphNodesOfClass(TSubclassOf<UAnimGraphNode_Base> NodeClass, TArray<UAnimGraphNode_Base*>& GraphNodes, bool bIncludeChildClasses = true);
+	UE_API void GetGraphNodesOfClass(TSubclassOf<UAnimGraphNode_Base> NodeClass, TArray<UAnimGraphNode_Base*>& GraphNodes, bool bIncludeChildClasses = true);
 
 private:
 	// UObject interface
-	virtual void PostEditUndo() override;
+	UE_API virtual void PostEditUndo() override;
 
 	// Reconstruct layer nodes post-undo
-	void ReconstructLayerNodes() const;
+	UE_API void ReconstructLayerNodes() const;
 };
 
+#undef UE_API
diff --git a/Engine/Source/Editor/AnimGraph/Public/AnimationGraphSchema.h b/Engine/Source/Editor/AnimGraph/Public/AnimationGraphSchema.h
--- a/Engine/Source/Editor/AnimGraph/Public/AnimationGraphSchema.h
+++ b/Engine/Source/Editor/AnimGraph/Public/AnimationGraphSchema.h
@@ -62,7 +62,7 @@ class UAnimationGraphSchema : public UEdGraphSchema_K2
 	ANIMGRAPH_API virtual void GetContextMenuActions(class UToolMenu* Menu, class UGraphNodeContextMenuContext* Context) const override;
 	ANIMGRAPH_API virtual FText GetPinDisplayName(const UEdGraphPin* Pin) const override;
 	ANIMGRAPH_API virtual bool CanDuplicateGraph(UEdGraph* InSourceGraph) const override;
-	ANIMGRAPH_API virtual bool ShouldAlwaysPurgeOnModification() const override { return true; }
+	virtual bool ShouldAlwaysPurgeOnModification() const override { return true; }
 	ANIMGRAPH_API virtual bool TryCreateConnection(UEdGraphPin* A, UEdGraphPin* B) const override;
 	ANIMGRAPH_API virtual void GetGraphDisplayInformation(const UEdGraph& Graph, /*out*/ FGraphDisplayInfo& DisplayInfo) const override;
 	ANIMGRAPH_API virtual bool CanGraphBeDropped(TSharedPtr<FEdGraphSchemaAction> InAction) const override;
@@ -75,8 +75,8 @@ class UAnimationGraphSchema : public UEdGraphSchema_K2
 	ANIMGRAPH_API virtual bool SearchForAutocastFunction(const FEdGraphPinType& OutputPinType, const FEdGraphPinType& InputPinType, /*out*/ FName& TargetFunction, /*out*/ UClass*& FunctionOwner) const override;
 	ANIMGRAPH_API virtual TOptional<FSearchForAutocastFunctionResults> SearchForAutocastFunction(const FEdGraphPinType& OutputPinType, const FEdGraphPinType& InputPinType) const;
 	ANIMGRAPH_API virtual bool ArePinsCompatible(const UEdGraphPin* PinA, const UEdGraphPin* PinB, const UClass* CallingContext = NULL, bool bIgnoreArray = false) const override;
-	ANIMGRAPH_API virtual bool DoesSupportCollapsedNodes() const override { return false; }
-	ANIMGRAPH_API virtual bool DoesSupportEventDispatcher() const	override { return false; }
+	virtual bool DoesSupportCollapsedNodes() const override { return false; }
+	virtual bool DoesSupportEventDispatcher() const	override { return false; }
 	ANIMGRAPH_API virtual bool DoesSupportAnimNotifyActions() const override;
 	ANIMGRAPH_API virtual void CreateFunctionGraphTerminators(UEdGraph& Graph, UClass* Class) const override;
 	ANIMGRAPH_API virtual bool CanShowDataTooltipForPin(const UEdGraphPin& Pin) const override;
@@ -116,7 +116,7 @@ class UAnimationGraphSchema : public UEdGraphSchema_K2
 	ANIMGRAPH_API static void ConformAnimLayersByGuid(const UAnimBlueprint* InAnimBlueprint, const FBPInterfaceDescription& CurrentInterfaceDesc);
 
 	UE_DEPRECATED(4.24, "Function renamed, please use GetPositionForNewLinkedInputPoseNode")
-	ANIMGRAPH_API static FVector2D GetPositionForNewSubInputNode(UEdGraph& InGraph) { return GetPositionForNewLinkedInputPoseNode(InGraph); }
+	static FVector2D GetPositionForNewSubInputNode(UEdGraph& InGraph) { return GetPositionForNewLinkedInputPoseNode(InGraph); }
 
 	/** Find a position for a newly created linked input pose */
 	ANIMGRAPH_API static FVector2D GetPositionForNewLinkedInputPoseNode(UEdGraph& InGraph);
diff --git a/Engine/Source/Editor/AnimGraph/Public/AnimationStateMachineSchema.h b/Engine/Source/Editor/AnimGraph/Public/AnimationStateMachineSchema.h
--- a/Engine/Source/Editor/AnimGraph/Public/AnimationStateMachineSchema.h
+++ b/Engine/Source/Editor/AnimGraph/Public/AnimationStateMachineSchema.h
@@ -9,12 +9,14 @@
 #include "EdGraph/EdGraphSchema.h"
 #include "AnimationStateMachineSchema.generated.h"
 
+#define UE_API ANIMGRAPH_API
+
 class FMenuBuilder;
 class UEdGraph;
 
 /** Action to add a node to the graph */
 USTRUCT()
-struct ANIMGRAPH_API FEdGraphSchemaAction_NewStateNode : public FEdGraphSchemaAction
+struct FEdGraphSchemaAction_NewStateNode : public FEdGraphSchemaAction
 {
 	GENERATED_USTRUCT_BODY()
 
@@ -30,8 +32,8 @@ struct ANIMGRAPH_API FEdGraphSchemaAction_NewStateNode : public FEdGraphSchemaAc
 		, NodeTemplate(nullptr)
 	{}
 
-	virtual UEdGraphNode* PerformAction(class UEdGraph* ParentGraph, UEdGraphPin* FromPin, const FVector2f& Location, bool bSelectNewNode = true) override;
-	virtual void AddReferencedObjects( FReferenceCollector& Collector ) override;
+	UE_API virtual UEdGraphNode* PerformAction(class UEdGraph* ParentGraph, UEdGraphPin* FromPin, const FVector2f& Location, bool bSelectNewNode = true) override;
+	UE_API virtual void AddReferencedObjects( FReferenceCollector& Collector ) override;
 
 	template <typename NodeType>
 	static NodeType* SpawnNodeFromTemplate(class UEdGraph* ParentGraph, NodeType* InTemplateNode, const UE::Slate::FDeprecateVector2DParameter Location = FVector2f(0.0f, 0.0f), bool bSelectNewNode = true)
@@ -45,7 +47,7 @@ struct ANIMGRAPH_API FEdGraphSchemaAction_NewStateNode : public FEdGraphSchemaAc
 
 /** Action to create new comment */
 USTRUCT()
-struct ANIMGRAPH_API FEdGraphSchemaAction_NewStateComment : public FEdGraphSchemaAction
+struct FEdGraphSchemaAction_NewStateComment : public FEdGraphSchemaAction
 {
 	GENERATED_USTRUCT_BODY();
 
@@ -58,7 +60,7 @@ struct ANIMGRAPH_API FEdGraphSchemaAction_NewStateComment : public FEdGraphSchem
 	{}
 
 	//~ Begin FEdGraphSchemaAction Interface
-	virtual UEdGraphNode* PerformAction(class UEdGraph* ParentGraph, UEdGraphPin* FromPin, const FVector2f& Location, bool bSelectNewNode = true) override;
+	UE_API virtual UEdGraphNode* PerformAction(class UEdGraph* ParentGraph, UEdGraphPin* FromPin, const FVector2f& Location, bool bSelectNewNode = true) override;
 	//~ End FEdGraphSchemaAction Interface
 };
 
@@ -94,3 +96,5 @@ class UAnimationStateMachineSchema : public UEdGraphSchema
 	virtual void BreakSinglePinLink(UEdGraphPin* SourcePin, UEdGraphPin* TargetPin) const override;
 	//~ End UEdGraphSchema Interface
 };
+
+#undef UE_API
diff --git a/Engine/Source/Editor/AnimGraph/Public/IAnimBlueprintCompilationBracketContext.h b/Engine/Source/Editor/AnimGraph/Public/IAnimBlueprintCompilationBracketContext.h
--- a/Engine/Source/Editor/AnimGraph/Public/IAnimBlueprintCompilationBracketContext.h
+++ b/Engine/Source/Editor/AnimGraph/Public/IAnimBlueprintCompilationBracketContext.h
@@ -8,7 +8,7 @@ class FCompilerResultsLog;
 class UAnimGraphNode_Base;
 
 // Interface passed to start/end compilation delegates
-class ANIMGRAPH_API IAnimBlueprintCompilationBracketContext
+class IAnimBlueprintCompilationBracketContext
 {
 public:	
 	virtual ~IAnimBlueprintCompilationBracketContext() {}
diff --git a/Engine/Source/Editor/AnimGraph/Public/IAnimBlueprintCompilationContext.h b/Engine/Source/Editor/AnimGraph/Public/IAnimBlueprintCompilationContext.h
--- a/Engine/Source/Editor/AnimGraph/Public/IAnimBlueprintCompilationContext.h
+++ b/Engine/Source/Editor/AnimGraph/Public/IAnimBlueprintCompilationContext.h
@@ -6,6 +6,8 @@
 #include "KismetCompiler.h"
 #include "Containers/ArrayView.h"
 
+#define UE_API ANIMGRAPH_API
+
 class UEdGraphNode;
 class UEdGraph;
 class UEdGraphPin;
@@ -20,7 +22,7 @@ struct FKismetCompilerOptions;
 #define ANIM_FUNC_DECORATOR	TEXT("__AnimFunc")
 
 /** Interface to the anim BP compiler context for use while compilation is in progress */
-class ANIMGRAPH_API IAnimBlueprintCompilationContext
+class IAnimBlueprintCompilationContext
 {
 public:
 	// Record of a property that can be folded into the class members/constant blocks
@@ -63,7 +65,7 @@ public:
 	virtual ~IAnimBlueprintCompilationContext() {}
 
 	// Get a compilation context from a kismet compiler context assuming that it is an FAnimBlueprintCompilerContext
-	static TUniquePtr<IAnimBlueprintCompilationContext> Get(FKismetCompilerContext& InKismetCompiler);
+	static UE_API TUniquePtr<IAnimBlueprintCompilationContext> Get(FKismetCompilerContext& InKismetCompiler);
 
 	// Spawns an intermediate node associated with the source node (for error purposes)
 	template <typename NodeType>
@@ -217,3 +219,5 @@ protected:
 	// Get the compiler as a base class to avoid circular include issues with templated functions/classes
 	virtual FKismetCompilerContext* GetKismetCompiler() const = 0;
 };
+
+#undef UE_API
diff --git a/Engine/Source/Editor/AnimGraph/Public/IAnimBlueprintCompilerCreationContext.h b/Engine/Source/Editor/AnimGraph/Public/IAnimBlueprintCompilerCreationContext.h
--- a/Engine/Source/Editor/AnimGraph/Public/IAnimBlueprintCompilerCreationContext.h
+++ b/Engine/Source/Editor/AnimGraph/Public/IAnimBlueprintCompilerCreationContext.h
@@ -36,7 +36,7 @@ UE_DEPRECATED(5.0, "Anim BP compiler delegate system is deprecated. Use UAnimBlu
 DECLARE_MULTICAST_DELEGATE_ThreeParams(FOnCopyTermDefaultsToDefaultObject, UObject* /*InDefaultObject*/, IAnimBlueprintCopyTermDefaultsContext& /*InCompilerContext*/, IAnimBlueprintGeneratedClassCompiledData& /*OutCompiledData*/);
 
 /** Interface to the creation of an anim BP compiler context */
-class ANIMGRAPH_API IAnimBlueprintCompilerCreationContext
+class IAnimBlueprintCompilerCreationContext
 {
 public:
 	virtual ~IAnimBlueprintCompilerCreationContext() {}
diff --git a/Engine/Source/Editor/AnimGraph/Public/IAnimBlueprintCompilerHandler.h b/Engine/Source/Editor/AnimGraph/Public/IAnimBlueprintCompilerHandler.h
--- a/Engine/Source/Editor/AnimGraph/Public/IAnimBlueprintCompilerHandler.h
+++ b/Engine/Source/Editor/AnimGraph/Public/IAnimBlueprintCompilerHandler.h
@@ -7,7 +7,7 @@
 class UE_DEPRECATED(5.0, "IAnimBlueprintCompilerHandler is deprecated, use UAnimBlueprintExtension_PropertyAccess instead") IAnimBlueprintCompilerHandler;
 
 /** Empty base class for handlers that are maintained for the lifetime of the compiler context */
-class ANIMGRAPH_API IAnimBlueprintCompilerHandler
+class IAnimBlueprintCompilerHandler
 {
 public:
 	virtual ~IAnimBlueprintCompilerHandler() {}
diff --git a/Engine/Source/Editor/AnimGraph/Public/IAnimBlueprintCompilerHandlerCollection.h b/Engine/Source/Editor/AnimGraph/Public/IAnimBlueprintCompilerHandlerCollection.h
--- a/Engine/Source/Editor/AnimGraph/Public/IAnimBlueprintCompilerHandlerCollection.h
+++ b/Engine/Source/Editor/AnimGraph/Public/IAnimBlueprintCompilerHandlerCollection.h
@@ -10,7 +10,7 @@ class IAnimBlueprintCompilerCreationContext;
 
 class UE_DEPRECATED(5.0, "IAnimBlueprintCompilerHandlerCollection is no longer used") IAnimBlueprintCompilerHandlerCollection;
 
-class ANIMGRAPH_API IAnimBlueprintCompilerHandlerCollection
+class IAnimBlueprintCompilerHandlerCollection
 {
 public:
 	PRAGMA_DISABLE_DEPRECATION_WARNINGS
@@ -33,4 +33,4 @@ protected:
 	// Get a GetHandler with the specified name
 	virtual IAnimBlueprintCompilerHandler* GetHandlerByName(FName InName) const = 0;
 	PRAGMA_ENABLE_DEPRECATION_WARNINGS
-};
\ No newline at end of file
+};
diff --git a/Engine/Source/Editor/AnimGraph/Public/IAnimBlueprintGeneratedClassCompiledData.h b/Engine/Source/Editor/AnimGraph/Public/IAnimBlueprintGeneratedClassCompiledData.h
--- a/Engine/Source/Editor/AnimGraph/Public/IAnimBlueprintGeneratedClassCompiledData.h
+++ b/Engine/Source/Editor/AnimGraph/Public/IAnimBlueprintGeneratedClassCompiledData.h
@@ -13,7 +13,7 @@ struct FBlueprintDebugData;
 struct FAnimBlueprintDebugData;
 
 /** Interface to the writable parts of the generated class that handlers can operate on */
-class ANIMGRAPH_API IAnimBlueprintGeneratedClassCompiledData
+class IAnimBlueprintGeneratedClassCompiledData
 {
 public:
 	// Get the baked state machines data for the currently-compiled class
@@ -47,4 +47,4 @@ public:
 	// @param	InSourceBlendSpace	The blendspace to duplicate
 	// @return a duplicate of the blendspace outered to the class
 	virtual UBlendSpace* AddBlendSpace(UBlendSpace* InSourceBlendSpace) = 0;
-};
\ No newline at end of file
+};
