Fix Alias state nodes not working with getters like GetRelevantAnimTime

#tests editor, PIE

diff --git a/Engine/Source/Editor/AnimGraph/Private/AnimBlueprintExtension_StateMachine.cpp b/Engine/Source/Editor/AnimGraph/Private/AnimBlueprintExtension_StateMachine.cpp
--- a/Engine/Source/Editor/AnimGraph/Private/AnimBlueprintExtension_StateMachine.cpp
+++ b/Engine/Source/Editor/AnimGraph/Private/AnimBlueprintExtension_StateMachine.cpp
@@ -474,7 +474,12 @@ void UAnimBlueprintExtension_StateMachine::AutoWireAnimGetter(class UK2Node_Anim
 		}
 		else if(Pin->PinName == TEXT("StateIndex") || Pin->PinName == TEXT("TransitionIndex"))
 		{
-			Pin->DefaultValue = FString::FromInt(SubNodeIndex);
+			// Don't write an invalid state index which can occur for aliases.
+			// Alias nodes set their default index earlier in UAnimGraphNode_StateMachineBase::OnProcessDuringCompilation
+			if (SubNodeIndex != INDEX_NONE)
+			{
+				Pin->DefaultValue = FString::FromInt(SubNodeIndex);
+			}
 		}
 	}
 }
diff --git a/Engine/Source/Editor/AnimGraph/Private/AnimGraphNode_StateMachineBase.cpp b/Engine/Source/Editor/AnimGraph/Private/AnimGraphNode_StateMachineBase.cpp
--- a/Engine/Source/Editor/AnimGraph/Private/AnimGraphNode_StateMachineBase.cpp
+++ b/Engine/Source/Editor/AnimGraph/Private/AnimGraphNode_StateMachineBase.cpp
@@ -11,6 +11,7 @@
 #include "AnimationStateMachineSchema.h"
 #include "AnimGraphNode_StateMachine.h"
 #include "Kismet2/KismetEditorUtilities.h"
+#include "K2Node_AnimGetter.h"
 #include "AnimBlueprintCompiler.h"
 #include "AnimStateNode.h"
 #include "AnimStateTransitionNode.h"
@@ -31,6 +32,7 @@
 #include "IAnimBlueprintCompilationContext.h"
 #include "Animation/AnimNode_Inertialization.h"
 #include "AnimStateAliasNode.h"
+#include "EdGraphUtilities.h"
 
 /////////////////////////////////////////////////////
 // FAnimStateMachineNodeNameValidator
@@ -606,10 +608,45 @@ void UAnimGraphNode_StateMachineBase::OnProcessDuringCompilation(IAnimBlueprintC
 		// Add aliased state transitions to transition list
 		TArray<int32> StateAliasIndices;
 		Oven.StateIndexToStateAliasNodeIndices.MultiFind(StateIndex, StateAliasIndices);
+
+		// Clone any alias transitions. Alias nodes share a single transition graph for all alias states. But we need a clone per alias state to allow
+		// for each alias transition graph to have correct hardcoded state index pin values. Note: Alias's are simply decompiled into multiple connections.
+		TMap<UAnimStateTransitionNode*, UEdGraph*> AliasTransitionGraphClones;
+
 		for (const int32 AliasIndex : StateAliasIndices)
 		{
+			// Operate on a separate list to avoid re-processing nodes in the final list.
+			TArray<UAnimStateTransitionNode*> AliasTransitionList;
+
 			// Let the final transition list do the sort.
-			Oven.StateAliasNodes[AliasIndex]->GetTransitionList(TransitionList, /*bWantSortedList=*/ false);
+			Oven.StateAliasNodes[AliasIndex]->GetTransitionList(AliasTransitionList, /*bWantSortedList=*/ false);
+
+			for (UAnimStateTransitionNode* AliasTransitionNode : AliasTransitionList)
+			{
+				if (UEdGraph* AliasTransitionGraph = AliasTransitionNode->BoundGraph)
+				{
+					UEdGraph* AliasTransitionGraphClone = FEdGraphUtilities::CloneGraph(AliasTransitionGraph, InCompilationContext.GetConsolidatedEventGraph(), &InCompilationContext.GetMessageLog(), true);
+
+					// Note: We do this here instead of in `UAnimBlueprintExtension_StateMachine::AutoWireAnimGetter` since we have the context of correct state index.
+					for (UEdGraphNode* Node : AliasTransitionGraphClone->Nodes)
+					{
+						if (UK2Node_AnimGetter* GetterNode = Cast<UK2Node_AnimGetter>(Node))
+						{
+							for (UEdGraphPin* Pin : GetterNode->Pins)
+							{
+								if (Pin->PinName == TEXT("StateIndex"))
+								{
+									Pin->DefaultValue = FString::FromInt(StateIndex);
+								}
+							}
+						}
+					}
+
+					AliasTransitionGraphClones.Add(AliasTransitionNode, MoveTemp(AliasTransitionGraphClone));
+				}
+			}
+
+			TransitionList.Append(MoveTemp(AliasTransitionList));
 		}
 		
 		StateNode->GetTransitionList(/*out*/ TransitionList, /*bWantSortedList=*/ true);
@@ -691,7 +728,13 @@ void UAnimGraphNode_StateMachineBase::OnProcessDuringCompilation(IAnimBlueprintC
 
 			Rule.TransitionIndex = TransitionIndex;
 			
-			if (UAnimGraphNode_TransitionResult* TransitionResultNode = CastChecked<UAnimationTransitionGraph>(TransitionNode->BoundGraph)->GetResultNode())
+			UEdGraph* TransitionGraph = TransitionNode->BoundGraph;
+			if (UEdGraph** AliasTransitionGraphClone = AliasTransitionGraphClones.Find(TransitionNode))
+			{
+				TransitionGraph = *AliasTransitionGraphClone;
+			}
+
+			if (UAnimGraphNode_TransitionResult* TransitionResultNode = CastChecked<UAnimationTransitionGraph>(TransitionGraph)->GetResultNode())
 			{
 				if (int32* pIndex = AlreadyMergedTransitionList.Find(TransitionResultNode))
 				{
@@ -699,13 +742,13 @@ void UAnimGraphNode_StateMachineBase::OnProcessDuringCompilation(IAnimBlueprintC
 				}
 				else
 				{
-					Rule.CanTakeDelegateIndex = Extension->ExpandGraphAndProcessNodes(TransitionNode->BoundGraph, TransitionResultNode, InCompilationContext, OutCompiledData, TransitionNode);
+					Rule.CanTakeDelegateIndex = Extension->ExpandGraphAndProcessNodes(TransitionGraph, TransitionResultNode, InCompilationContext, OutCompiledData, TransitionNode);
 					AlreadyMergedTransitionList.Add(TransitionResultNode, Rule.CanTakeDelegateIndex);
 				}
 
 				if (TransitionNode->bAutomaticRuleBasedOnSequencePlayerInState)
 				{
-					if (UAnimationTransitionGraph* TransGraph = Cast<UAnimationTransitionGraph>(TransitionNode->BoundGraph))
+					if (UAnimationTransitionGraph* TransGraph = Cast<UAnimationTransitionGraph>(TransitionGraph))
 					{
 						if (UAnimGraphNode_TransitionResult* ResultNode = TransGraph->GetResultNode())
 						{
diff --git a/Engine/Source/Editor/AnimGraph/Private/K2Node_AnimGetter.cpp b/Engine/Source/Editor/AnimGraph/Private/K2Node_AnimGetter.cpp
--- a/Engine/Source/Editor/AnimGraph/Private/K2Node_AnimGetter.cpp
+++ b/Engine/Source/Editor/AnimGraph/Private/K2Node_AnimGetter.cpp
@@ -5,6 +5,7 @@
 #include "AnimGraphNode_AssetPlayerBase.h"
 #include "AnimGraphNode_Base.h"
 #include "AnimGraphNode_StateMachine.h"
+#include "AnimStateAliasNode.h"
 #include "AnimStateNode.h"
 #include "AnimStateNodeBase.h"
 #include "AnimStateTransitionNode.h"
@@ -186,11 +187,13 @@ void UK2Node_AnimGetter::GetMenuActions(FBlueprintActionDatabaseRegistrar& Actio
 			TArray<UAnimGraphNode_StateMachine*> MachineNodes;
 			TArray<UAnimStateNode*> StateNodes;
 			TArray<UAnimStateTransitionNode*> TransitionNodes;
+			TArray<UAnimStateAliasNode*> AliasNodes;
 			
 			FBlueprintEditorUtils::GetAllNodesOfClass(AnimBlueprint, AssetPlayerNodes);
 			FBlueprintEditorUtils::GetAllNodesOfClass(AnimBlueprint, MachineNodes);
 			FBlueprintEditorUtils::GetAllNodesOfClass(AnimBlueprint, StateNodes);
 			FBlueprintEditorUtils::GetAllNodesOfClass(AnimBlueprint, TransitionNodes);
+			FBlueprintEditorUtils::GetAllNodesOfClass(AnimBlueprint, AliasNodes);
 
 			for(UFunction* Getter : AnimGetters)
 			{
@@ -239,6 +242,26 @@ void UK2Node_AnimGetter::GetMenuActions(FBlueprintActionDatabaseRegistrar& Actio
 							Spawner->DynamicUiSignatureGetter = UBlueprintNodeSpawner::FUiSpecOverrideDelegate::CreateStatic(UiSpecOverride, Params.CachedTitle);
 							ActionRegistrar.AddBlueprintAction(AnimBlueprint, Spawner);
 						}
+
+						for (UAnimStateAliasNode* AliasNode : AliasNodes)
+						{
+							// Get the state machine node from the outer chain
+							UAnimationStateMachineGraph* Graph = Cast<UAnimationStateMachineGraph>(AliasNode->GetOuter());
+							if (Graph)
+							{
+								if (UAnimGraphNode_StateMachine* MachineNode = Cast<UAnimGraphNode_StateMachine>(Graph->GetOuter()))
+								{
+									Params.SourceNode = MachineNode;
+								}
+							}
+
+							Params.SourceStateNode = AliasNode;
+							UpdateCachedTitle(Params);
+
+							UBlueprintNodeSpawner* Spawner = UBlueprintNodeSpawner::Create(UK2Node_AnimGetter::StaticClass(), /*AliasNode->GetGraph()*/nullptr, UBlueprintNodeSpawner::FCustomizeNodeDelegate::CreateUObject(const_cast<UK2Node_AnimGetter*>(this), &UK2Node_AnimGetter::PostSpawnNodeSetup, Params));
+							Spawner->DynamicUiSignatureGetter = UBlueprintNodeSpawner::FUiSpecOverrideDelegate::CreateStatic(UiSpecOverride, Params.CachedTitle);
+							ActionRegistrar.AddBlueprintAction(AnimBlueprint, Spawner);
+						}
 					}
 					else if(GetterRequiresParameter(Getter, TEXT("TransitionIndex")))
 					{
@@ -369,6 +392,14 @@ bool UK2Node_AnimGetter::IsActionFilteredOut(FBlueprintActionFilter const& Filte
 										return false;
 									}
 								}
+								else if (UAnimStateAliasNode* PreviousAliasNode = Cast<UAnimStateAliasNode>(TransitionNode->GetPreviousState()))
+								{
+									// Only allow actions using states that are referencing the previous alias
+									if (SourceStateNode == PreviousAliasNode)
+									{
+										return false;
+									}
+								}
 							}
 							else if(UAnimGraphNode_StateMachine* MachineNode = Cast<UAnimGraphNode_StateMachine>(SourceNode))
 							{
