Fix crash when re-removing a pin on a blend multi node after undoing

diff --git a/Engine/Source/Editor/AnimGraph/Private/AnimGraphNode_MultiWayBlend.cpp b/Engine/Source/Editor/AnimGraph/Private/AnimGraphNode_MultiWayBlend.cpp
--- a/Engine/Source/Editor/AnimGraph/Private/AnimGraphNode_MultiWayBlend.cpp
+++ b/Engine/Source/Editor/AnimGraph/Private/AnimGraphNode_MultiWayBlend.cpp
@@ -61,6 +61,13 @@ void UAnimGraphNode_MultiWayBlend::GetNodeContextMenuActions(UToolMenu* Menu, UG
 	}
 }
 
+void UAnimGraphNode_MultiWayBlend::PostEditUndo()
+{
+	Super::PostEditUndo();
+
+	Node.UpdateCachedAlphas();
+}
+
 void UAnimGraphNode_MultiWayBlend::AddPinToBlendNode()
 {
 	FScopedTransaction Transaction(LOCTEXT("AddBlendMultiPin", "AddBlendMultiPin"));
diff --git a/Engine/Source/Editor/AnimGraph/Public/AnimGraphNode_MultiWayBlend.h b/Engine/Source/Editor/AnimGraph/Public/AnimGraphNode_MultiWayBlend.h
--- a/Engine/Source/Editor/AnimGraph/Public/AnimGraphNode_MultiWayBlend.h
+++ b/Engine/Source/Editor/AnimGraph/Public/AnimGraphNode_MultiWayBlend.h
@@ -34,6 +34,10 @@ class UAnimGraphNode_MultiWayBlend : public UAnimGraphNode_Base
 	virtual void GetNodeContextMenuActions(class UToolMenu* Menu, class UGraphNodeContextMenuContext* Context) const override;
 	// End of UK2Node interface
 
+	// Begin UObject interface
+	virtual void PostEditUndo() override;
+	// End of UObject interface
+
 private:
 	int32 RemovedPinArrayIndex;
 
diff --git a/Engine/Source/Runtime/AnimGraphRuntime/Public/AnimNodes/AnimNode_MultiWayBlend.h b/Engine/Source/Runtime/AnimGraphRuntime/Public/AnimNodes/AnimNode_MultiWayBlend.h
--- a/Engine/Source/Runtime/AnimGraphRuntime/Public/AnimNodes/AnimNode_MultiWayBlend.h
+++ b/Engine/Source/Runtime/AnimGraphRuntime/Public/AnimNodes/AnimNode_MultiWayBlend.h
@@ -86,5 +86,7 @@ public:
 private:
 	// process new weights and then return out
 	ANIMGRAPHRUNTIME_API void UpdateCachedAlphas();
+
+	friend class UAnimGraphNode_MultiWayBlend;
 };
 
