Remove some deprecated 5.4 functions on KismetCompiler

#rnx

diff --git a/Engine/Source/Editor/AnimGraph/Public/IAnimBlueprintCompilationContext.h b/Engine/Source/Editor/AnimGraph/Public/IAnimBlueprintCompilationContext.h
--- a/Engine/Source/Editor/AnimGraph/Public/IAnimBlueprintCompilationContext.h
+++ b/Engine/Source/Editor/AnimGraph/Public/IAnimBlueprintCompilationContext.h
@@ -74,14 +74,6 @@ public:
 		return GetKismetCompiler()->SpawnIntermediateNode<NodeType>(SourceNode, ParentGraph);
 	}
 
-	// Spawns an intermediate event node associated with the source node (for error purposes)
-	template <typename NodeType>
-	UE_DEPRECATED(5.4, "SpawnIntermediateEventNode is equivalent to SpawnIntermediateNode, this redundant function has been deprecated.")
-	NodeType* SpawnIntermediateEventNode(UEdGraphNode* SourceNode, UEdGraphPin* SourcePin = nullptr, UEdGraph* ParentGraph = nullptr)
-	{
-		return GetKismetCompiler()->SpawnIntermediateNode<NodeType>(SourceNode, ParentGraph);
-	}
-
 	// Find a property in the currently-compiled class
 	template <typename FieldType>
 	FieldType* FindClassFProperty(const TCHAR* InFieldPath) const
diff --git a/Engine/Source/Editor/KismetCompiler/Public/KismetCompiler.h b/Engine/Source/Editor/KismetCompiler/Public/KismetCompiler.h
--- a/Engine/Source/Editor/KismetCompiler/Public/KismetCompiler.h
+++ b/Engine/Source/Editor/KismetCompiler/Public/KismetCompiler.h
@@ -245,13 +245,6 @@ public:
 		return Result;
 	}
 
-	template <typename NodeType>
-	UE_DEPRECATED(5.4, "SpawnIntermediateEventNode is equivalent to SpawnIntermediateNode, this redundant function has been deprecated.")
-	NodeType* SpawnIntermediateEventNode(UEdGraphNode* SourceNode, UEdGraphPin* SourcePin = nullptr, UEdGraph* ParentGraph = nullptr)
-	{
-		return SpawnIntermediateNode<NodeType>(SourceNode, ParentGraph);
-	}
-
 	/**
 	 * Moves pin links over from the source-pin to the specified intermediate, 
 	 * and validates the result (additionally logs a redirect from the 
