Previously transformed bones except the last one appear off-set in the newly saved Preview Animation
#fix reset preview bone controllers when switching asset in preview scene (prevents doubling up on adjustments)

diff --git a/Engine/Source/Editor/AnimGraph/Private/AnimPreviewInstance.cpp b/Engine/Source/Editor/AnimGraph/Private/AnimPreviewInstance.cpp
--- a/Engine/Source/Editor/AnimGraph/Private/AnimPreviewInstance.cpp
+++ b/Engine/Source/Editor/AnimGraph/Private/AnimPreviewInstance.cpp
@@ -569,6 +569,12 @@ void UAnimPreviewInstance::SetAnimationAsset(UAnimationAsset* NewAsset, bool bIs
 
 	// make sure to turn that off before setting new asset
 	Proxy.GetRequiredBones().SetUseSourceData(false);
+	
+	// Reset Bone controllers when switching assets
+	if (NewAsset != CurrentAsset)
+	{
+		Proxy.ResetModifiedBone();
+	}
 
 	Super::SetAnimationAsset(NewAsset, bIsLooping, InPlayRate);
 	RootMotionMode = Cast<UAnimMontage>(CurrentAsset) != nullptr ? ERootMotionMode::RootMotionFromMontagesOnly : ERootMotionMode::RootMotionFromEverything;
