Anim State Machine: Fix clamping of values preventing manual entering of valid negatives

diff --git a/Engine/Source/Editor/AnimGraph/Public/AnimStateTransitionNode.h b/Engine/Source/Editor/AnimGraph/Public/AnimStateTransitionNode.h
--- a/Engine/Source/Editor/AnimGraph/Public/AnimStateTransitionNode.h
+++ b/Engine/Source/Editor/AnimGraph/Public/AnimStateTransitionNode.h
@@ -64,13 +64,13 @@ class UAnimStateTransitionNode : public UAnimStateNodeBase
 	// When should the automatic transition rule trigger relative to the time remaining on the relevant asset player:
 	//  < 0 means trigger the transition 'Crossfade Duration' seconds before the end of the asset player, so a standard blend would finish just as the asset player ends
 	// >= 0 means trigger the transition 'Automatic Rule Trigger Time' seconds before the end of the asset player
-	UPROPERTY(EditAnywhere, Category=Transition, meta=(ClampMin=0.0, ForceUnits="s"))
+	UPROPERTY(EditAnywhere, Category=Transition, meta=(ForceUnits="s"))
 	float AutomaticRuleTriggerTime;
 
 	// The minimum time elapsed (in seconds) before the target state can be re-entered via this transition 
 	// Has no effect when set to -1.0.
 	// When set to zero, wait at least one frame before re-entry via this transition is allowed.
-	UPROPERTY(EditAnywhere, Category=Transition, DisplayName="Min Time Before Re-entry", meta=(ClampMin=0.0, ForceUnits="s"))
+	UPROPERTY(EditAnywhere, Category=Transition, DisplayName="Min Time Before Re-entry", meta=(ForceUnits="s"))
 	float MinTimeBeforeReentry;
 
 	// If SyncGroupName is specified, Transition will only be taken if the SyncGroup has valid markers (other transition rules still apply in addition to that).
