Fixed blueprint nodes in an animation blueprint's event graph being annotated twice.

#rnx

diff --git a/Engine/Source/Editor/AnimGraph/Private/AnimBlueprintCompiler.cpp b/Engine/Source/Editor/AnimGraph/Private/AnimBlueprintCompiler.cpp
--- a/Engine/Source/Editor/AnimGraph/Private/AnimBlueprintCompiler.cpp
+++ b/Engine/Source/Editor/AnimGraph/Private/AnimBlueprintCompiler.cpp
@@ -709,8 +709,15 @@ void FAnimBlueprintCompilerContext::ProcessAllAnimationNodes()
 		PruneIsolatedAnimationNodes(RootSet, RootAnimNodeList);
 
 		// Validate the graph
-		ValidateGraphIsWellFormed(ConsolidatedEventGraph);
-
+		{
+			TGuardValue<bool> SilentModeGuard(MessageLog.bSilentMode, true);
+			TGuardValue<bool> AnnotateMentionedNodesGuard(MessageLog.bAnnotateMentionedNodes, false);
+			
+			ValidateGraphIsWellFormed(ConsolidatedEventGraph);
+			
+			MessageLog.Messages.Empty();
+		}
+		
 		FAnimBlueprintGeneratedClassCompiledData CompiledData(GetNewAnimBlueprintClass());
 		FAnimBlueprintCompilationContext CompilerContext(this);
 		UAnimBlueprintExtension::ForEachExtension(AnimBlueprint, [&AllSubGraphsAnimNodeList, &CompiledData, &CompilerContext](UAnimBlueprintExtension* InExtension)
