RemoveCurveNode. Fix for LOCTEXT duplicates.

diff --git a/Engine/Source/Editor/AnimGraph/Private/AnimGraphNode_RemoveCurve.cpp b/Engine/Source/Editor/AnimGraph/Private/AnimGraphNode_RemoveCurve.cpp
--- a/Engine/Source/Editor/AnimGraph/Private/AnimGraphNode_RemoveCurve.cpp
+++ b/Engine/Source/Editor/AnimGraph/Private/AnimGraphNode_RemoveCurve.cpp
@@ -17,7 +17,7 @@ FText UAnimGraphNode_RemoveCurve::GetMenuCategory() const
 
 FText UAnimGraphNode_RemoveCurve::GetTooltipText() const
 {	
-	return LOCTEXT("AnimGraphNode_RemoveCurve_Category", "Removes animation curves from the pose input");
+	return LOCTEXT("AnimGraphNode_RemoveCurve_Tooltip", "Removes animation curves from the pose input");
 }
 
 FText UAnimGraphNode_RemoveCurve::GetNodeTitle(ENodeTitleType::Type TitleType) const
