[UE] RBAN Add an 'Allow Scale Propagation' flag to each RBAN instance to support disabling the pose scale propagation on a per node basis.
- the flag is true by default but will only have an effect if the pose scale propagation feature is enabled globaly by the p.RigidBodyNode.EnablePoseScalePropagation CVar
- The associated checkbox in the RBAN editor UI details panel will only be visible if the pose scale propagation feature is enabled globaly.

#Synced-CL 51534790

diff --git a/Engine/Source/Editor/AnimGraph/Private/AnimGraphNode_RigidBody.cpp b/Engine/Source/Editor/AnimGraph/Private/AnimGraphNode_RigidBody.cpp
--- a/Engine/Source/Editor/AnimGraph/Private/AnimGraphNode_RigidBody.cpp
+++ b/Engine/Source/Editor/AnimGraph/Private/AnimGraphNode_RigidBody.cpp
@@ -148,6 +148,27 @@ void UAnimGraphNode_RigidBody::CustomizeDetails(IDetailLayoutBuilder& DetailBuil
 				]
 			]
 		];
+
+	// Add per node check box for allowing bone transform scale propagation. This check box is hidden unless the the feature is enabled globally via CVar.
+	{
+		IDetailCategoryBuilder& SettingsCategory = DetailBuilder.EditCategory(TEXT("Settings"));
+		TSharedRef<IPropertyHandle> AllowPoseScalePropagationHandle = DetailBuilder.GetProperty(TEXT("Node.bAllowPoseScalePropagation"), GetClass());
+
+		TAttribute<EVisibility> VisiblityAttr = TAttribute<EVisibility>::Create(TAttribute<EVisibility>::FGetter::CreateStatic(&UAnimGraphNode_RigidBody::ShouldShowPoseScalePropagationSettings));
+
+		SettingsCategory.AddCustomRow(AllowPoseScalePropagationHandle->GetPropertyDisplayName(), true)
+			.NameContent()
+			[
+				AllowPoseScalePropagationHandle->CreatePropertyNameWidget()
+			]
+			.ValueContent()
+			[
+				AllowPoseScalePropagationHandle->CreatePropertyValueWidget()
+			]
+			.Visibility(VisiblityAttr);
+
+		AllowPoseScalePropagationHandle->MarkHiddenByCustomization();
+	}
 }
 
 void UAnimGraphNode_RigidBody::PostEditChangeProperty(struct FPropertyChangedEvent& PropertyChangedEvent)
@@ -207,4 +228,9 @@ bool UAnimGraphNode_RigidBody::AreAnyConstraintsHidden() const
 	return false;
 }
 
+EVisibility UAnimGraphNode_RigidBody::ShouldShowPoseScalePropagationSettings()
+{
+	return RigidBodyNodeIsPoseScalePropagationEnabledGlobally() ? EVisibility::Visible : EVisibility::Collapsed;
+}
+
 #undef LOCTEXT_NAMESPACE
diff --git a/Engine/Source/Editor/AnimGraph/Public/AnimGraphNode_RigidBody.h b/Engine/Source/Editor/AnimGraph/Public/AnimGraphNode_RigidBody.h
--- a/Engine/Source/Editor/AnimGraph/Public/AnimGraphNode_RigidBody.h
+++ b/Engine/Source/Editor/AnimGraph/Public/AnimGraphNode_RigidBody.h
@@ -41,6 +41,8 @@ protected:
 
 	bool AreAnyBodiesHidden() const;
 	bool AreAnyConstraintsHidden() const;
+	
+	static EVisibility ShouldShowPoseScalePropagationSettings();
 
 private:
 	TWeakObjectPtr<UPoseWatchElement> PoseWatchElementBodies;
diff --git a/Engine/Source/Runtime/AnimGraphRuntime/Private/BoneControllers/AnimNode_RigidBody.cpp b/Engine/Source/Runtime/AnimGraphRuntime/Private/BoneControllers/AnimNode_RigidBody.cpp
--- a/Engine/Source/Runtime/AnimGraphRuntime/Private/BoneControllers/AnimNode_RigidBody.cpp
+++ b/Engine/Source/Runtime/AnimGraphRuntime/Private/BoneControllers/AnimNode_RigidBody.cpp
@@ -203,6 +203,7 @@ FAnimNode_RigidBody::FAnimNode_RigidBody()
 	, bTransferBoneVelocities(false)
 	, bFreezeIncomingPoseOnStart(false)
 	, bClampLinearTranslationLimitToRefPose(false)
+	, bAllowPoseScalePropagation(true)
 	, WorldSpaceMinimumScale(0.01f)
 	, EvaluationResetTime(0.01f)
 	, bEnabled(false)
@@ -696,7 +697,7 @@ void FAnimNode_RigidBody::EvaluateSkeletalControl_AnyThread(FComponentSpacePoseC
 		const bool bUseDeferredSimulationTask =
 			((SimulationTiming == ESimulationTiming::Default) || bRBAN_DeferredSimulationForceDefault) ? bRBAN_DeferredSimulationDefault : (SimulationTiming == ESimulationTiming::Deferred);
 
-		bool bApplyBoneTransformScaleToOutputPose = CVarEnableRigidBodyNodePoseScalePropagation.GetValueOnAnyThread();
+		bool bApplyBoneTransformScaleToOutputPose = bAllowPoseScalePropagation && CVarEnableRigidBodyNodePoseScalePropagation.GetValueOnAnyThread();
 
 		FVector SimSpaceGravity(0.f);
 
@@ -2035,5 +2036,10 @@ bool FAnimNode_RigidBody::IsValidToEvaluate(const USkeleton* Skeleton, const FBo
 	return BaseBoneRef.IsValidToEvaluate(RequiredBones);
 }
 
+bool RigidBodyNodeIsPoseScalePropagationEnabledGlobally()
+{
+	return CVarEnableRigidBodyNodePoseScalePropagation.GetValueOnAnyThread();
+}
+
 #undef LOCTEXT_NAMESPACE
 
diff --git a/Engine/Source/Runtime/AnimGraphRuntime/Public/BoneControllers/AnimNode_RigidBody.h b/Engine/Source/Runtime/AnimGraphRuntime/Public/BoneControllers/AnimNode_RigidBody.h
--- a/Engine/Source/Runtime/AnimGraphRuntime/Public/BoneControllers/AnimNode_RigidBody.h
+++ b/Engine/Source/Runtime/AnimGraphRuntime/Public/BoneControllers/AnimNode_RigidBody.h
@@ -319,6 +319,14 @@ public:
 	UPROPERTY(EditAnywhere, Category = Settings)
 	uint8 bClampLinearTranslationLimitToRefPose : 1;
 
+	/**
+		Allow / Prevent bone transform scales denfined in the input pose from being passed to the output pose. This setting is true by default 
+		but will only have an effect if bone scale propagation is enabled project wide with the p.RigidBodyNode.EnablePoseScalePropagation 
+		console variable. The details pannel checkbox will be hidden if the feature is not enabled.
+	*/
+	UPROPERTY(EditAnywhere, Category = Settings)
+	uint8 bAllowPoseScalePropagation : 1;
+
 	/**
 		For world-space simulations, if the magnitude of the component's 3D scale is less than WorldSpaceMinimumScale, do not update the node.
 	*/
@@ -645,3 +653,5 @@ inline FVector ConvertCSPositionToSimSpace(
 	default: ensureMsgf(false, TEXT("Unsupported Simulation Space")); return InCSPosition;
 	}
 }
+
+ANIMGRAPHRUNTIME_API bool RigidBodyNodeIsPoseScalePropagationEnabledGlobally();
