Ensure that the saved cached pose name is unique
#rnx

#Synced-CL 52426642
#Synced-CL 52426642

diff --git a/Engine/Source/Editor/AnimGraph/Private/AnimGraphNode_SaveCachedPose.cpp b/Engine/Source/Editor/AnimGraph/Private/AnimGraphNode_SaveCachedPose.cpp
--- a/Engine/Source/Editor/AnimGraph/Private/AnimGraphNode_SaveCachedPose.cpp
+++ b/Engine/Source/Editor/AnimGraph/Private/AnimGraphNode_SaveCachedPose.cpp
@@ -111,6 +111,7 @@ bool UAnimGraphNode_SaveCachedPose::IsCompatibleWithGraph(const UEdGraph* Target
 void UAnimGraphNode_SaveCachedPose::OnRenameNode(const FString& NewName)
 {
 	CacheName = NewName;
+	MakeUniqueCacheName();
 	FBlueprintEditorUtils::MarkBlueprintAsStructurallyModified(GetBlueprint());
 }
 
@@ -119,9 +120,55 @@ TSharedPtr<class INameValidatorInterface> UAnimGraphNode_SaveCachedPose::MakeNam
 	return MakeShareable(new FCachedPoseNameValidator(GetBlueprint(), CacheName));
 }
 
+void UAnimGraphNode_SaveCachedPose::PostEditChangeProperty(FPropertyChangedEvent& PropertyChangedEvent)
+{
+	const FName PropertyName = PropertyChangedEvent.GetPropertyName();
+
+	if (PropertyName == GET_MEMBER_NAME_CHECKED(UAnimGraphNode_SaveCachedPose, CacheName))
+	{
+	 	if (MakeUniqueCacheName())
+	 	{
+			FBlueprintEditorUtils::MarkBlueprintAsStructurallyModified(GetBlueprint());
+		}
+	}
+
+	Super::PostEditChangeProperty(PropertyChangedEvent);
+}
+
+void UAnimGraphNode_SaveCachedPose::PostPasteNode()
+{
+	Super::PostPasteNode();
+	if (MakeUniqueCacheName())
+	{
+		FBlueprintEditorUtils::MarkBlueprintAsStructurallyModified(GetBlueprint());
+	}
+}
+
 void UAnimGraphNode_SaveCachedPose::GetRequiredExtensions(TArray<TSubclassOf<UAnimBlueprintExtension>>& OutExtensions) const
 {
 	OutExtensions.Add(UAnimBlueprintExtension_CachedPose::StaticClass());
 }
 
+bool UAnimGraphNode_SaveCachedPose::MakeUniqueCacheName()
+{
+	bool bHasChanged = false;
+	// Disallow empty names
+	if (CacheName.IsEmpty())
+	{
+		CacheName = TEXT("CachedPose");
+		bHasChanged = true;
+	}
+	// Check for duplicate names among other SaveCachedPose nodes
+	if (TSharedPtr<INameValidatorInterface> Validator = MakeNameValidator())
+	{
+		FString CurrentName = CacheName;
+		Validator->FindValidString(CacheName);
+		if (CurrentName != CacheName)
+		{
+			bHasChanged = true;
+		}
+	}
+	return bHasChanged;
+}
+
 #undef LOCTEXT_NAMESPACE
diff --git a/Engine/Source/Editor/AnimGraph/Public/AnimGraphNode_SaveCachedPose.h b/Engine/Source/Editor/AnimGraph/Public/AnimGraphNode_SaveCachedPose.h
--- a/Engine/Source/Editor/AnimGraph/Public/AnimGraphNode_SaveCachedPose.h
+++ b/Engine/Source/Editor/AnimGraph/Public/AnimGraphNode_SaveCachedPose.h
@@ -31,6 +31,8 @@ class UAnimGraphNode_SaveCachedPose : public UAnimGraphNode_Base
 	virtual TSharedPtr<class INameValidatorInterface> MakeNameValidator() const override;
 	virtual bool IsCompatibleWithGraph(const UEdGraph* TargetGraph) const override;
 	virtual bool IsNodeRootSet() const override { return true; }
+	virtual void PostEditChangeProperty(FPropertyChangedEvent& PropertyChangedEvent) override;
+	virtual void PostPasteNode() override;
 	// End of UEdGraphNode interface
 
 	// UK2Node interface.
@@ -44,6 +46,9 @@ class UAnimGraphNode_SaveCachedPose : public UAnimGraphNode_Base
 	// End of UAnimGraphNode_Base interface
 
 private:
+	
+	bool MakeUniqueCacheName();
+	
 	/** Constructing FText strings can be costly, so we cache the node's title */
 	FNodeTextCache CachedNodeTitle;
 };
