UAnimGraphNode_SaveCachedPose: fix name always modified post rename because it conflicts with the existing name

#rnx

#Synced-CL 52578518

diff --git a/Engine/Source/Editor/AnimGraph/Private/AnimGraphNode_SaveCachedPose.cpp b/Engine/Source/Editor/AnimGraph/Private/AnimGraphNode_SaveCachedPose.cpp
--- a/Engine/Source/Editor/AnimGraph/Private/AnimGraphNode_SaveCachedPose.cpp
+++ b/Engine/Source/Editor/AnimGraph/Private/AnimGraphNode_SaveCachedPose.cpp
@@ -15,15 +15,18 @@
 class FCachedPoseNameValidator : public FStringSetNameValidator
 {
 public:
-	FCachedPoseNameValidator(class UBlueprint* InBlueprint, const FString& InExistingName)
+	FCachedPoseNameValidator(const UBlueprint* InBlueprint, const UAnimGraphNode_SaveCachedPose* CurrentNode, const FString& InExistingName)
 		: FStringSetNameValidator(InExistingName)
 	{
-		TArray<UAnimGraphNode_SaveCachedPose*> Nodes;
-		FBlueprintEditorUtils::GetAllNodesOfClass<UAnimGraphNode_SaveCachedPose>(InBlueprint, Nodes);
+		TArray<const UAnimGraphNode_SaveCachedPose*> Nodes;
+		FBlueprintEditorUtils::GetAllNodesOfClass<const UAnimGraphNode_SaveCachedPose>(InBlueprint, Nodes);
 
-		for (auto Node : Nodes)
+		for (const UAnimGraphNode_SaveCachedPose* Node : Nodes)
 		{
-			Names.Add(Node->CacheName);
+			if (Node != CurrentNode)
+			{
+				Names.Add(Node->CacheName);
+			}
 		}
 	}
 };
@@ -111,13 +114,12 @@ bool UAnimGraphNode_SaveCachedPose::IsCompatibleWithGraph(const UEdGraph* Target
 void UAnimGraphNode_SaveCachedPose::OnRenameNode(const FString& NewName)
 {
 	CacheName = NewName;
-	MakeUniqueCacheName();
 	FBlueprintEditorUtils::MarkBlueprintAsStructurallyModified(GetBlueprint());
 }
 
 TSharedPtr<class INameValidatorInterface> UAnimGraphNode_SaveCachedPose::MakeNameValidator() const
 {
-	return MakeShareable(new FCachedPoseNameValidator(GetBlueprint(), CacheName));
+	return MakeShared<FCachedPoseNameValidator>(GetBlueprint(), this, CacheName);
 }
 
 void UAnimGraphNode_SaveCachedPose::PostEditChangeProperty(FPropertyChangedEvent& PropertyChangedEvent)
@@ -159,14 +161,11 @@ bool UAnimGraphNode_SaveCachedPose::MakeUniqueCacheName()
 		bHasChanged = true;
 	}
 	// Check for duplicate names among other SaveCachedPose nodes
-	if (TSharedPtr<INameValidatorInterface> Validator = MakeNameValidator())
+	FCachedPoseNameValidator Validator(GetBlueprint(), this, "");
+	if (Validator.IsValid(CacheName,true) != EValidatorResult::Ok)
 	{
-		FString CurrentName = CacheName;
-		Validator->FindValidString(CacheName);
-		if (CurrentName != CacheName)
-		{
-			bHasChanged = true;
-		}
+		Validator.FindValidString(CacheName);
+		bHasChanged = true;
 	}
 	return bHasChanged;
 }
