#Synced-CL 52740731

diff --git a/Engine/Source/Editor/AnimGraph/Private/AnimPreviewInstance.cpp b/Engine/Source/Editor/AnimGraph/Private/AnimPreviewInstance.cpp
--- a/Engine/Source/Editor/AnimGraph/Private/AnimPreviewInstance.cpp
+++ b/Engine/Source/Editor/AnimGraph/Private/AnimPreviewInstance.cpp
@@ -417,6 +417,7 @@ void FAnimPreviewInstanceProxy::AddKeyToSequence(UAnimSequence* Sequence, float
 void FAnimPreviewInstanceProxy::SetDebugSkeletalMeshComponent(USkeletalMeshComponent* InSkeletalMeshComponent)
 {
 	CopyPoseNode.SourceMeshComponent = InSkeletalMeshComponent;
+	CopyPoseNode.bCopyCurves = true;
 	CopyPoseNode.Initialize_AnyThread(FAnimationInitializeContext(this));
 }
 
